Let FRHIViewDesc::EDimension decide the RHI's SRV/UAV dimension to be able to create 2D array SRVs/UAVs for cubemap textures
Keep the behavior from 25147071 when we want to create a 2D SRV of something which is not a 2D texture. This ensures we do not ignore FRHIViewDesc::FTexture::ArrayRange and select the right slice of a texture

diff --git a/Engine/Source/Runtime/D3D12RHI/Private/D3D12SRV.cpp b/Engine/Source/Runtime/D3D12RHI/Private/D3D12SRV.cpp
--- a/Engine/Source/Runtime/D3D12RHI/Private/D3D12SRV.cpp
+++ b/Engine/Source/Runtime/D3D12RHI/Private/D3D12SRV.cpp
@@ -174,12 +174,10 @@ static FD3D12ShaderResourceView::EFlags TranslateDesc(D3D12_SHADER_RESOURCE_VIEW
 
 	uint32 const PlaneSlice = UE::DXGIUtilities::GetPlaneSliceFromViewFormat(BaseFormat, ViewFormat);
 	FRHIRange8 const PlaneRange(PlaneSlice, 1);
-
-	// No need to use Info.Dimension, since D3D supports mixing Texture2D view types.
-	// Create a view which matches the underlying resource dimension.
-	switch (TextureDesc.Dimension)
+	FRHIViewDesc::EDimension ViewDimension = UE::RHICore::AdjustViewInfoDimensionForNarrowing(Info, TextureDesc);
+	switch (ViewDimension)
 	{
-	case ETextureDimension::Texture2D:
+	case FRHIViewDesc::EDimension::Texture2D:
 		if (TextureDesc.NumSamples > 1)
 		{
 			SRVDesc.ViewDimension = D3D12_SRV_DIMENSION_TEXTURE2DMS;
@@ -193,7 +191,7 @@ static FD3D12ShaderResourceView::EFlags TranslateDesc(D3D12_SHADER_RESOURCE_VIEW
 		}
 		break;
 
-	case ETextureDimension::Texture2DArray:
+	case FRHIViewDesc::EDimension::Texture2DArray:
 		if (TextureDesc.NumSamples > 1)
 		{
 			SRVDesc.ViewDimension = D3D12_SRV_DIMENSION_TEXTURE2DMSARRAY;
@@ -211,19 +209,19 @@ static FD3D12ShaderResourceView::EFlags TranslateDesc(D3D12_SHADER_RESOURCE_VIEW
 		}
 		break;
 
-	case ETextureDimension::Texture3D:
+	case FRHIViewDesc::EDimension::Texture3D:
 		SRVDesc.ViewDimension = D3D12_SRV_DIMENSION_TEXTURE3D;
 		SRVDesc.Texture3D.MostDetailedMip = Info.MipRange.First;
 		SRVDesc.Texture3D.MipLevels = Info.MipRange.Num;
 		break;
 
-	case ETextureDimension::TextureCube:
+	case FRHIViewDesc::EDimension::TextureCube:
 		SRVDesc.ViewDimension = D3D12_SRV_DIMENSION_TEXTURECUBE;
 		SRVDesc.TextureCube.MostDetailedMip = Info.MipRange.First;
 		SRVDesc.TextureCube.MipLevels = Info.MipRange.Num;
 		break;
 
-	case ETextureDimension::TextureCubeArray:
+	case FRHIViewDesc::EDimension::TextureCubeArray:
 		SRVDesc.ViewDimension = D3D12_SRV_DIMENSION_TEXTURECUBEARRAY;
 		SRVDesc.TextureCubeArray.MostDetailedMip = Info.MipRange.First;
 		SRVDesc.TextureCubeArray.MipLevels = Info.MipRange.Num;
diff --git a/Engine/Source/Runtime/D3D12RHI/Private/D3D12UAV.cpp b/Engine/Source/Runtime/D3D12RHI/Private/D3D12UAV.cpp
--- a/Engine/Source/Runtime/D3D12RHI/Private/D3D12UAV.cpp
+++ b/Engine/Source/Runtime/D3D12RHI/Private/D3D12UAV.cpp
@@ -153,18 +153,17 @@ static FD3D12UnorderedAccessView::EFlags TranslateDesc(D3D12_UNORDERED_ACCESS_VI
 	uint32 const PlaneSlice = UE::DXGIUtilities::GetPlaneSliceFromViewFormat(BaseFormat, ViewFormat);
 	FRHIRange8 const PlaneRange(PlaneSlice, 1);
 
-	// No need to use Info.Dimension, since D3D supports mixing Texture2D view types.
-	// Create a view which matches the underlying resource dimension.
-	switch (TextureDesc.Dimension)
+	FRHIViewDesc::EDimension ViewDimension = UE::RHICore::AdjustViewInfoDimensionForNarrowing(Info, TextureDesc);
+	switch (ViewDimension)
 	{
-	case ETextureDimension::Texture2D:
+	case FRHIViewDesc::EDimension::Texture2D:
 		UAVDesc.ViewDimension = D3D12_UAV_DIMENSION_TEXTURE2D;
 		UAVDesc.Texture2D.MipSlice = Info.MipLevel;
 		UAVDesc.Texture2D.PlaneSlice = PlaneSlice;
 		break;
 
-	case ETextureDimension::TextureCube:
-	case ETextureDimension::TextureCubeArray:
+	case FRHIViewDesc::EDimension::TextureCube:
+	case FRHIViewDesc::EDimension::TextureCubeArray:
 		UAVDesc.ViewDimension = D3D12_UAV_DIMENSION_TEXTURE2DARRAY;
 		UAVDesc.Texture2DArray.FirstArraySlice = Info.ArrayRange.First * 6;
 		UAVDesc.Texture2DArray.ArraySize = Info.ArrayRange.Num * 6;
@@ -172,7 +171,7 @@ static FD3D12UnorderedAccessView::EFlags TranslateDesc(D3D12_UNORDERED_ACCESS_VI
 		UAVDesc.Texture2DArray.PlaneSlice = PlaneSlice;
 		break;
 
-	case ETextureDimension::Texture2DArray:
+	case FRHIViewDesc::EDimension::Texture2DArray:
 		UAVDesc.ViewDimension = D3D12_UAV_DIMENSION_TEXTURE2DARRAY;
 		UAVDesc.Texture2DArray.FirstArraySlice = Info.ArrayRange.First;
 		UAVDesc.Texture2DArray.ArraySize = Info.ArrayRange.Num;
@@ -180,7 +179,7 @@ static FD3D12UnorderedAccessView::EFlags TranslateDesc(D3D12_UNORDERED_ACCESS_VI
 		UAVDesc.Texture2DArray.PlaneSlice = PlaneSlice;
 		break;
 
-	case ETextureDimension::Texture3D:
+	case FRHIViewDesc::EDimension::Texture3D:
 		UAVDesc.ViewDimension = D3D12_UAV_DIMENSION_TEXTURE3D;
 		UAVDesc.Texture3D.FirstWSlice = 0;
 		UAVDesc.Texture3D.WSize = FMath::Max(TextureDesc.Depth >> Info.MipLevel, 1);
diff --git a/Engine/Source/Runtime/RHI/Private/RHIResources.cpp b/Engine/Source/Runtime/RHI/Private/RHIResources.cpp
--- a/Engine/Source/Runtime/RHI/Private/RHIResources.cpp
+++ b/Engine/Source/Runtime/RHI/Private/RHIResources.cpp
@@ -240,11 +240,36 @@ FRHIViewDesc::FTexture::FViewInfo FRHIViewDesc::FTexture::GetViewInfo(FRHITextur
 	);
 
 	checkf(ArrayRange.Num > 0 || ArrayRange.First == 0, TEXT("ArrayRange.Num cannot be zero, unless creating a view of the entire range."));
-	checkf((ArrayRange.First + ArrayRange.Num) <= Desc.ArraySize, TEXT("Array range (first: %d, num: %d) is out of bounds for texture description (array size: %d)."),
-		ArrayRange.First,
-		ArrayRange.Num,
-		Desc.ArraySize
-	);
+
+	// make sure the view fits in the texture
+	{
+		uint16 TextureArraySize = Desc.IsTextureCube() ? Desc.ArraySize * 6 : Desc.ArraySize;
+		uint16 ViewArrayFirst = ArrayRange.First;
+		uint16 ViewArrayNum = ArrayRange.Num == 0 ? Desc.ArraySize : ArrayRange.Num;
+
+		bool bIsCubeView = false;
+		if (Dimension == EDimension::TextureCube || Dimension == EDimension::TextureCubeArray)
+		{
+			bIsCubeView = true;
+			ViewArrayFirst *= 6;
+			ViewArrayNum *= 6;
+		}
+
+		uint16 SliceDividerForCheckMessage = bIsCubeView ? 6 : 1; // We want the message to report the number of elements in the units of the view
+		checkf(ViewArrayFirst + ViewArrayNum <= TextureArraySize, TEXT("Array range (first: %d, num: %d) is out of bounds for texture description (array size: %d)."),
+			ViewArrayFirst / SliceDividerForCheckMessage,
+			ViewArrayNum / SliceDividerForCheckMessage,
+			TextureArraySize / SliceDividerForCheckMessage
+		);
+	}
+
+	// When ArrayRange.Num == 0, we use the number of elements from the texture. If the view is a 2D array and the texture a cube (array), we need to do x6 on the number of slices
+	// We already checked that we can only create cube views on cube textures, so we only need to take into account the 2D view on cube texture case
+	uint16 AdjustedTextureArraySize = Desc.ArraySize;
+	if (Dimension == EDimension::Texture2DArray && Desc.IsTextureCube())
+	{
+		AdjustedTextureArraySize *= 6;
+	}
 
 	FViewInfo Info = {};
 	Info.Format = Format == PF_Unknown ? Desc.Format : Format;
@@ -293,8 +318,8 @@ FRHIViewDesc::FTexture::FViewInfo FRHIViewDesc::FTexture::GetViewInfo(FRHITextur
 	);
 
 	Info.ArrayRange.First = ArrayRange.First;
-	Info.ArrayRange.Num = ArrayRange.Num == 0 ? Desc.ArraySize : ArrayRange.Num;
-	Info.bAllSlices = Info.ArrayRange.First == 0 && Info.ArrayRange.Num == Desc.ArraySize;
+	Info.ArrayRange.Num = ArrayRange.Num == 0 ? AdjustedTextureArraySize : ArrayRange.Num;
+	Info.bAllSlices = Info.ArrayRange.First == 0 && Info.ArrayRange.Num == AdjustedTextureArraySize;
 
 	return Info;
 }
diff --git a/Engine/Source/Runtime/RHICore/Private/RHICore.cpp b/Engine/Source/Runtime/RHICore/Private/RHICore.cpp
--- a/Engine/Source/Runtime/RHICore/Private/RHICore.cpp
+++ b/Engine/Source/Runtime/RHICore/Private/RHICore.cpp
@@ -42,5 +42,20 @@ void ResolveRenderPassTargets(const FRHIRenderPassInfo& RenderPassInfo, TFunctio
 	ResolveTexture({ DSV.DepthStencilTarget, DSV.ResolveTarget, 0, 0, RenderPassInfo.ResolveRect });
 }
 
+FRHIViewDesc::EDimension AdjustViewInfoDimensionForNarrowing(const FRHIViewDesc::FTexture::FViewInfo& ViewInfo, const FRHITextureDesc& TextureDesc)
+{
+	FRHIViewDesc::EDimension ViewInfoDimension = ViewInfo.Dimension;
+	// some RHIs do not support creating a 2D view to index a specific slice in a 2D array/texture. 
+	// But they do support binding a single slice 2D array as a Texture2D in the shader
+	if (ViewInfoDimension == FRHIViewDesc::EDimension::Texture2D && (TextureDesc.IsTextureCube() || TextureDesc.IsTextureArray()))
+	{
+		ViewInfoDimension = FRHIViewDesc::EDimension::Texture2DArray;
+		ensureAlwaysMsgf(ViewInfo.ArrayRange.Num == 1 || TextureDesc.ArraySize == 1, TEXT("Trying to create a 2D SRV with more than 1 element"));
+	}
+
+	return ViewInfoDimension;
+}
+
+
 } //! RHICore
 } //! UE
\ No newline at end of file
diff --git a/Engine/Source/Runtime/RHICore/Public/RHICore.h b/Engine/Source/Runtime/RHICore/Public/RHICore.h
--- a/Engine/Source/Runtime/RHICore/Public/RHICore.h
+++ b/Engine/Source/Runtime/RHICore/Public/RHICore.h
@@ -34,6 +34,7 @@ struct FResolveTextureInfo
 };
 
 RHICORE_API void ResolveRenderPassTargets(const FRHIRenderPassInfo& Info, TFunction<void(FResolveTextureInfo)> Function);
+RHICORE_API FRHIViewDesc::EDimension AdjustViewInfoDimensionForNarrowing(const FRHIViewDesc::FTexture::FViewInfo& ViewInfo, const FRHITextureDesc& TextureDesc);
 
 } //! UE::RHICore
 
diff --git a/Engine/Source/Runtime/VulkanRHI/Private/VulkanUAV.cpp b/Engine/Source/Runtime/VulkanRHI/Private/VulkanUAV.cpp
--- a/Engine/Source/Runtime/VulkanRHI/Private/VulkanUAV.cpp
+++ b/Engine/Source/Runtime/VulkanRHI/Private/VulkanUAV.cpp
@@ -436,7 +436,7 @@ void FVulkanShaderResourceView::UpdateView()
 
 		uint32 ArrayFirst = Info.ArrayRange.First;
 		uint32 ArrayNum = Info.ArrayRange.Num;
-		if (Texture->GetDesc().IsTextureCube())
+		if (Info.Dimension == FRHIViewDesc::EDimension::TextureCube || Info.Dimension == FRHIViewDesc::EDimension::TextureCubeArray)
 		{
 			ArrayFirst *= 6;
 			ArrayNum *= 6;
@@ -526,7 +526,7 @@ void FVulkanUnorderedAccessView::UpdateView()
 
 		uint32 ArrayFirst = Info.ArrayRange.First;
 		uint32 ArrayNum = Info.ArrayRange.Num;
-		if (Texture->GetDesc().IsTextureCube())
+		if (Info.Dimension == FRHIViewDesc::EDimension::TextureCube || Info.Dimension == FRHIViewDesc::EDimension::TextureCubeArray)
 		{
 			ArrayFirst *= 6;
 			ArrayNum *= 6;
diff --git a/Engine/Source/Runtime/Windows/D3D11RHI/Private/D3D11SRV.cpp b/Engine/Source/Runtime/Windows/D3D11RHI/Private/D3D11SRV.cpp
--- a/Engine/Source/Runtime/Windows/D3D11RHI/Private/D3D11SRV.cpp
+++ b/Engine/Source/Runtime/Windows/D3D11RHI/Private/D3D11SRV.cpp
@@ -74,11 +74,10 @@ void FD3D11ShaderResourceView::UpdateView()
 
 		SRVDesc.Format = UE::DXGIUtilities::FindShaderResourceFormat(DXGI_FORMAT(GPixelFormats[Info.Format].PlatformFormat), Info.bSRGB);
 
-		// No need to use Info.Dimension, since D3D supports mixing Texture2D view types.
-		// Create a view which matches the underlying resource dimension.
-		switch (TextureDesc.Dimension)
+		FRHIViewDesc::EDimension ViewDimension = UE::RHICore::AdjustViewInfoDimensionForNarrowing(Info, TextureDesc);
+		switch (ViewDimension)
 		{
-		case ETextureDimension::Texture2D:
+		case FRHIViewDesc::EDimension::Texture2D:
 			if (TextureDesc.NumSamples > 1)
 			{
 				SRVDesc.ViewDimension = D3D11_SRV_DIMENSION_TEXTURE2DMS;
@@ -91,7 +90,7 @@ void FD3D11ShaderResourceView::UpdateView()
 			}
 			break;
 
-		case ETextureDimension::Texture2DArray:
+		case FRHIViewDesc::EDimension::Texture2DArray:
 			if (TextureDesc.NumSamples > 1)
 			{
 				SRVDesc.ViewDimension = D3D11_SRV_DIMENSION_TEXTURE2DMSARRAY;
@@ -108,19 +107,19 @@ void FD3D11ShaderResourceView::UpdateView()
 			}
 			break;
 
-		case ETextureDimension::Texture3D:
+		case FRHIViewDesc::EDimension::Texture3D:
 			SRVDesc.ViewDimension = D3D11_SRV_DIMENSION_TEXTURE3D;
 			SRVDesc.Texture3D.MostDetailedMip = Info.MipRange.First;
 			SRVDesc.Texture3D.MipLevels       = Info.MipRange.Num;
 			break;
 
-		case ETextureDimension::TextureCube:
+		case FRHIViewDesc::EDimension::TextureCube:
 			SRVDesc.ViewDimension = D3D11_SRV_DIMENSION_TEXTURECUBE;
 			SRVDesc.TextureCube.MostDetailedMip = Info.MipRange.First;
 			SRVDesc.TextureCube.MipLevels       = Info.MipRange.Num;
 			break;
 
-		case ETextureDimension::TextureCubeArray:
+		case FRHIViewDesc::EDimension::TextureCubeArray:
 			SRVDesc.ViewDimension = D3D11_SRV_DIMENSION_TEXTURECUBEARRAY;
 			SRVDesc.TextureCubeArray.MostDetailedMip  = Info.MipRange.First;
 			SRVDesc.TextureCubeArray.MipLevels        = Info.MipRange.Num;
diff --git a/Engine/Source/Runtime/Windows/D3D11RHI/Private/D3D11UAV.cpp b/Engine/Source/Runtime/Windows/D3D11RHI/Private/D3D11UAV.cpp
--- a/Engine/Source/Runtime/Windows/D3D11RHI/Private/D3D11UAV.cpp
+++ b/Engine/Source/Runtime/Windows/D3D11RHI/Private/D3D11UAV.cpp
@@ -69,31 +69,30 @@ void FD3D11UnorderedAccessView::UpdateView()
 
 		UAVDesc.Format = UE::DXGIUtilities::FindUnorderedAccessFormat(DXGI_FORMAT(GPixelFormats[Info.Format].PlatformFormat));
 
-		// No need to use Info.Dimension, since D3D supports mixing Texture2D view types.
-		// Create a view which matches the underlying resource dimension.
-		switch (TextureDesc.Dimension)
+		switch (Info.Dimension)
 		{
-		case ETextureDimension::Texture2D:
+		case FRHIViewDesc::EDimension::Texture2D:
 			UAVDesc.ViewDimension = D3D11_UAV_DIMENSION_TEXTURE2D;
 			UAVDesc.Texture2D.MipSlice = Info.MipLevel;
+			ensureAlwaysMsgf(Info.ArrayRange.First == 0, TEXT("Trying to create an UAV beyond the first slice. This is not supported on d3d11, and binding a single slice 2D array as an HLSL RWTexture2D is not supported either"));
 			break;
 
-		case ETextureDimension::TextureCube:
-		case ETextureDimension::TextureCubeArray:
+		case FRHIViewDesc::EDimension::TextureCube:
+		case FRHIViewDesc::EDimension::TextureCubeArray:
 			UAVDesc.ViewDimension = D3D11_UAV_DIMENSION_TEXTURE2DARRAY;
 			UAVDesc.Texture2DArray.FirstArraySlice = Info.ArrayRange.First * 6;
 			UAVDesc.Texture2DArray.ArraySize       = Info.ArrayRange.Num * 6;
 			UAVDesc.Texture2DArray.MipSlice        = Info.MipLevel;
 			break;
 
-		case ETextureDimension::Texture2DArray:
+		case FRHIViewDesc::EDimension::Texture2DArray:
 			UAVDesc.ViewDimension = D3D11_UAV_DIMENSION_TEXTURE2DARRAY;
 			UAVDesc.Texture2DArray.FirstArraySlice = Info.ArrayRange.First;
 			UAVDesc.Texture2DArray.ArraySize       = Info.ArrayRange.Num;
 			UAVDesc.Texture2DArray.MipSlice        = Info.MipLevel;
 			break;
 
-		case ETextureDimension::Texture3D:
+		case FRHIViewDesc::EDimension::Texture3D:
 			UAVDesc.ViewDimension = D3D11_UAV_DIMENSION_TEXTURE3D;
 			UAVDesc.Texture3D.FirstWSlice = 0;
 			UAVDesc.Texture3D.WSize       = FMath::Max(TextureDesc.Depth >> Info.MipLevel, 1);
