Fix Debug Client compilation on Android

diff --git a/Engine/Source/Runtime/Engine/Classes/Engine/StreamableManager.h b/Engine/Source/Runtime/Engine/Classes/Engine/StreamableManager.h
--- a/Engine/Source/Runtime/Engine/Classes/Engine/StreamableManager.h
+++ b/Engine/Source/Runtime/Engine/Classes/Engine/StreamableManager.h
@@ -473,9 +473,9 @@ ENUM_CLASS_FLAGS(EStreamableManagerCombinedHandleOptions);
 struct FStreamableManager : public FGCObject
 {
 	// Default priority for all async loads
-	static const TAsyncLoadPriority DefaultAsyncLoadPriority = 0;
+	static constexpr TAsyncLoadPriority DefaultAsyncLoadPriority = 0;
 	// Priority to try and load immediately
-	static const TAsyncLoadPriority AsyncLoadHighPriority = 100;
+	static constexpr TAsyncLoadPriority AsyncLoadHighPriority = 100;
 
 	/** 
 	 * This is the primary streamable operation. Requests streaming of one or more target objects. When complete, a delegate function is called. Returns a Streamable Handle.
diff --git a/Engine/Source/Runtime/RHI/Public/Android/AndroidDataDrivenShaderPlatformInfo.inl b/Engine/Source/Runtime/RHI/Public/Android/AndroidDataDrivenShaderPlatformInfo.inl
--- a/Engine/Source/Runtime/RHI/Public/Android/AndroidDataDrivenShaderPlatformInfo.inl
+++ b/Engine/Source/Runtime/RHI/Public/Android/AndroidDataDrivenShaderPlatformInfo.inl
@@ -5,15 +5,15 @@
 class FDataDrivenShaderPlatformInfo final : public FGenericDataDrivenShaderPlatformInfo
 {
 public:
-	IMPLEMENT_DDPSPI_SETTING(GetTargetsTiledGPU,			true);
-	IMPLEMENT_DDPSPI_SETTING(GetIsHlslcc,					false);
-	IMPLEMENT_DDPSPI_SETTING(GetSupportsDxc,				true);
-	IMPLEMENT_DDPSPI_SETTING(SupportsUniformBufferObjects,  true);
+	IMPLEMENT_DDPSPI_SETTING(GetTargetsTiledGPU,				true);
+	IMPLEMENT_DDPSPI_SETTING(GetIsHlslcc,						false);
+	IMPLEMENT_DDPSPI_SETTING(GetSupportsDxc,					true);
+	IMPLEMENT_DDPSPI_SETTING(GetSupportsUniformBufferObjects,	true);
 
 #if USE_STATIC_FEATURE_LEVEL_ENUMS
 	IMPLEMENT_DDPSPI_SETTING_WITH_RETURN_TYPE(ERHIFeatureLevel::Type, GetMaxFeatureLevel, FStaticFeatureLevel(UE_ANDROID_STATIC_FEATURE_LEVEL));
 	IMPLEMENT_DDPSPI_SETTING(GetIsMobile, (UE_ANDROID_STATIC_FEATURE_LEVEL == ERHIFeatureLevel::ES3_1));
-	IMPLEMENT_DDPSPI_SETTING(SupportsManualVertexFetch, (UE_ANDROID_STATIC_FEATURE_LEVEL > ERHIFeatureLevel::ES3_1));
-	IMPLEMENT_DDPSPI_SETTING(SupportsMobileMultiView, (UE_ANDROID_STATIC_FEATURE_LEVEL == ERHIFeatureLevel::ES3_1));
+	IMPLEMENT_DDPSPI_SETTING(GetSupportsManualVertexFetch, (UE_ANDROID_STATIC_FEATURE_LEVEL > ERHIFeatureLevel::ES3_1));
+	IMPLEMENT_DDPSPI_SETTING(GetSupportsMobileMultiView, (UE_ANDROID_STATIC_FEATURE_LEVEL == ERHIFeatureLevel::ES3_1));
 #endif
 };
