Minimal workaround for Android / Vulkan use-after-free RHI thread crash in FRHITextureReference / UpdateTextureReference.
 - Hold a reference on the FRHITextureReference object to prevent deletion while its still in an RHI command list.

diff --git a/Engine/Source/Runtime/RHI/Private/RHICommandList.cpp b/Engine/Source/Runtime/RHI/Private/RHICommandList.cpp
--- a/Engine/Source/Runtime/RHI/Private/RHICommandList.cpp
+++ b/Engine/Source/Runtime/RHI/Private/RHICommandList.cpp
@@ -2097,7 +2097,12 @@ void FRHICommandListBase::UpdateTextureReference(FRHITextureReference* TextureRe
 		return;
 	}
 
-	EnqueueLambda(TEXT("UpdateTextureReference"), [TextureRef, NewTexture](FRHICommandListBase& RHICmdList)
+	// Workaround for a crash bug where FRHITextureReferences are deleted before this command is executed on the RHI thread.
+	// Take a reference on the FRHITextureReference object to keep it alive.
+	// @todo dev-pr - This should be refactored out when we eventually remove FRHITextureReference.
+	TRefCountPtr<FRHITextureReference> Ref = TextureRef;
+
+	EnqueueLambda(TEXT("UpdateTextureReference"), [TextureRef, NewTexture, LocalRef = MoveTemp(Ref)](FRHICommandListBase& RHICmdList)
 	{
 		GDynamicRHI->RHIUpdateTextureReference(RHICmdList, TextureRef, NewTexture);
 	});
