Enable full RDG events when using "r.rhisetgpucaptureoptions 1".

diff --git a/Engine/Source/Runtime/RHI/Private/DynamicRHI.cpp b/Engine/Source/Runtime/RHI/Private/DynamicRHI.cpp
--- a/Engine/Source/Runtime/RHI/Private/DynamicRHI.cpp
+++ b/Engine/Source/Runtime/RHI/Private/DynamicRHI.cpp
@@ -557,6 +557,22 @@ void FDynamicRHI::EnableIdealGPUCaptureOptions(bool bEnabled)
 			ShowMaterialDrawEventVar->Set(bEnabled ? -1 : 0);
 		}
 	}
+
+	// Full pass names in RDG
+	static IConsoleVariable* RDGEvents = IConsoleManager::Get().FindConsoleVariable(TEXT("r.RDG.Events"));
+	if (RDGEvents)
+	{
+		if (bEnabled)
+		{
+			UE_LOG(LogRHI, Display, TEXT("Enabling full RDG events (r.RDG.Events 3)"));
+			RDGEvents->Set(3, ECVF_SetByCode);
+		}
+		else
+		{
+			UE_LOG(LogRHI, Display, TEXT("Resetting RDG events to default (r.RDG.Events)"));
+			RDGEvents->Unset(ECVF_SetByCode);
+		}
+	}
 }
 
 FTextureReferenceRHIRef FDynamicRHI::RHICreateTextureReference(FRHICommandListBase& RHICmdList, FRHITexture* InReferencedTexture)
