RHI: Addded a helper macro for conditionally adding code if shader data is enabed (RHI_IF_SHADER_DEBUG_DATA).

#rnx

#changelist validated
#virtualized

diff --git a/Engine/Source/Runtime/RHI/Public/RHIResources.h b/Engine/Source/Runtime/RHI/Public/RHIResources.h
--- a/Engine/Source/Runtime/RHI/Public/RHIResources.h
+++ b/Engine/Source/Runtime/RHI/Public/RHIResources.h
@@ -755,6 +755,12 @@ public:
 	#define RHI_INCLUDE_SHADER_DEBUG_DATA 0
 #endif
 
+#if RHI_INCLUDE_SHADER_DEBUG_DATA
+	#define RHI_IF_SHADER_DEBUG_DATA(...)	__VA_ARGS__
+#else
+	#define RHI_IF_SHADER_DEBUG_DATA(...)
+#endif
+
 class FRHIShader : public FRHIResource
 {
 public:
