Fix crash when RHI validation is used and a precache PSO fails to compile.

Failing to compile a PSO is not a fatal error if the request is coming from the precache system. The RHI validation layer must account for that, instead of blindly trying to use the FRHIGraphicsPipelineState pointer returned by the underlying RHI.

#rnx

diff --git a/Engine/Source/Runtime/RHI/Public/RHIValidation.h b/Engine/Source/Runtime/RHI/Public/RHIValidation.h
--- a/Engine/Source/Runtime/RHI/Public/RHIValidation.h
+++ b/Engine/Source/Runtime/RHI/Public/RHIValidation.h
@@ -357,7 +357,10 @@ public:
 	{
 		ValidatePipeline(Initializer);
 		FGraphicsPipelineStateRHIRef PSO = RHI->RHICreateGraphicsPipelineState(Initializer);
-		PSO->DSMode = Initializer.DepthStencilState->ActualDSMode;
+		if (PSO.IsValid())
+		{
+			PSO->DSMode = Initializer.DepthStencilState->ActualDSMode;
+		}
 		return PSO;
 	}
 
@@ -370,7 +373,10 @@ public:
 	{
 		ValidatePipeline(Initializer);
 		FGraphicsPipelineStateRHIRef PSO = RHI->RHICreateGraphicsPipelineState(Initializer);
-		PSO->DSMode = Initializer.DepthStencilState->ActualDSMode;
+		if (PSO.IsValid())
+		{
+			PSO->DSMode = Initializer.DepthStencilState->ActualDSMode;
+		}
 		return PSO;
 	}
 
