PR #11129: Add bDepthBounds to pipeline file cache

diff --git a/Engine/Source/Runtime/RHI/Private/PipelineFileCache.cpp b/Engine/Source/Runtime/RHI/Private/PipelineFileCache.cpp
--- a/Engine/Source/Runtime/RHI/Private/PipelineFileCache.cpp
+++ b/Engine/Source/Runtime/RHI/Private/PipelineFileCache.cpp
@@ -82,13 +82,14 @@ enum class EPipelineCacheFileFormatVersions : uint32
 	AddingDepthClipMode = 25,
 	BeforeStableCacheVersioning = 26,
 	RemovingLineAA = 27,
+	AddingDepthBounds = 28,
 };
 
 const uint64 FPipelineCacheFileFormatMagic = 0x5049504543414348; // PIPECACH
 const uint64 FPipelineCacheTOCFileFormatMagic = 0x544F435354415232; // TOCSTAR2
 const uint64 FPipelineCacheEOFFileFormatMagic = 0x454F462D4D41524B; // EOF-MARK
-const RHI_API uint32 FPipelineCacheFileFormatCurrentVersion = (uint32)EPipelineCacheFileFormatVersions::RemovingLineAA;
-const int32  FPipelineCacheGraphicsDescPartsNum = 66; // parser will expect this number of parts in a description string
+const RHI_API uint32 FPipelineCacheFileFormatCurrentVersion = (uint32)EPipelineCacheFileFormatVersions::AddingDepthBounds;
+const int32  FPipelineCacheGraphicsDescPartsNum = 67; // parser will expect this number of parts in a description string
 
 /**
   * PipelineFileCache API access
@@ -564,6 +565,10 @@ FString FPipelineCacheFileFormatPSO::GraphicsDescriptor::StateToString() const
 		, uint32(MultiViewCount)
 		, uint32(bHasFragmentDensityAttachment)
 	);
+
+	Result += FString::Printf(TEXT("%d,")
+		, uint32(bDepthBounds)
+	);
 	
 	FVertexElement NullVE;
 	FMemory::Memzero(NullVE);
@@ -649,6 +654,10 @@ void FPipelineCacheFileFormatPSO::GraphicsDescriptor::AddStateToReadableString(T
 	OutBuilder << bHasFragmentDensityAttachment;
 	OutBuilder << TEXT("\n");
 
+	OutBuilder << TEXT(" DB:");
+	OutBuilder << bDepthBounds;
+	OutBuilder << TEXT("\n");
+
 	OutBuilder << TEXT(" NumVE ");
 	OutBuilder << VertexDescriptor.Num();
 	OutBuilder << TEXT("\n");
@@ -745,6 +754,14 @@ bool FPipelineCacheFileFormatPSO::GraphicsDescriptor::StateFromString(const FStr
 		bHasFragmentDensityAttachment = (bool)LocalHasFDM;
 	}
 
+	// parse depth bounds
+	{
+		uint32 DepthBounds = 0;
+		check(PartEnd - PartIt >= 1);
+		LexFromString(DepthBounds, *PartIt++);
+		bDepthBounds = (bool)DepthBounds;
+	}
+
 	check(PartEnd - PartIt >= 1); //not a very robust parser
 	int32 VertDescNum = 0;
 	LexFromString(VertDescNum, *PartIt++);
@@ -1140,6 +1157,8 @@ bool FPipelineCacheFileFormatPSO::Verify() const
 			KeyHash = FCrc::MemCrc32(&Key.GraphicsDesc.MultiViewCount, sizeof(Key.GraphicsDesc.MultiViewCount), KeyHash);
 			KeyHash = FCrc::MemCrc32(&Key.GraphicsDesc.bHasFragmentDensityAttachment, sizeof(Key.GraphicsDesc.bHasFragmentDensityAttachment), KeyHash);
 
+			KeyHash = FCrc::MemCrc32(&Key.GraphicsDesc.bDepthBounds, sizeof(Key.GraphicsDesc.bDepthBounds), KeyHash);
+
 			for(auto const& Element : Key.GraphicsDesc.VertexDescriptor)
 			{
 				KeyHash = FCrc::MemCrc32(&Element, sizeof(FVertexElement), KeyHash);
@@ -1349,6 +1368,11 @@ bool FPipelineCacheFileFormatPSO::Verify() const
 				Ar << Info.GraphicsDesc.bHasFragmentDensityAttachment;
 			}
 
+			if (Ar.GameNetVer() >= (uint32)EPipelineCacheFileFormatVersions::AddingDepthBounds)
+			{
+				Ar << Info.GraphicsDesc.bDepthBounds;
+			}
+
 			break;
 		}
 		case FPipelineCacheFileFormatPSO::DescriptorType::RayTracing:
@@ -1529,6 +1553,8 @@ FPipelineCacheFileFormatPSO::FPipelineCacheFileFormatPSO()
 	PSO.GraphicsDesc.MultiViewCount = (uint8)Init.MultiViewCount;
 	PSO.GraphicsDesc.bHasFragmentDensityAttachment = Init.bHasFragmentDensityAttachment;
 
+	PSO.GraphicsDesc.bDepthBounds = Init.bDepthBounds;
+
 #if !UE_BUILD_SHIPPING
 	bOK = bOK && PSO.Verify();
 #endif
@@ -1584,6 +1610,7 @@ bool FPipelineCacheFileFormatPSO::operator==(const FPipelineCacheFileFormatPSO&
 						GraphicsDesc.DepthStore == Other.GraphicsDesc.DepthStore && GraphicsDesc.StencilLoad == Other.GraphicsDesc.StencilLoad && GraphicsDesc.StencilStore == Other.GraphicsDesc.StencilStore &&
 						GraphicsDesc.SubpassHint == Other.GraphicsDesc.SubpassHint && GraphicsDesc.SubpassIndex == Other.GraphicsDesc.SubpassIndex &&
 						GraphicsDesc.MultiViewCount == Other.GraphicsDesc.MultiViewCount && GraphicsDesc.bHasFragmentDensityAttachment == Other.GraphicsDesc.bHasFragmentDensityAttachment &&
+						GraphicsDesc.bDepthBounds == Other.GraphicsDesc.bDepthBounds &&
 					FMemory::Memcmp(&GraphicsDesc.BlendState, &Other.GraphicsDesc.BlendState, sizeof(FBlendStateInitializerRHI)) == 0 &&
 					FMemory::Memcmp(&GraphicsDesc.RasterizerState, &Other.GraphicsDesc.RasterizerState, sizeof(FPipelineFileCacheRasterizerState)) == 0 &&
 					FMemory::Memcmp(&GraphicsDesc.DepthStencilState, &Other.GraphicsDesc.DepthStencilState, sizeof(FDepthStencilStateInitializerRHI)) == 0 &&
diff --git a/Engine/Source/Runtime/RHI/Public/PipelineFileCache.h b/Engine/Source/Runtime/RHI/Public/PipelineFileCache.h
--- a/Engine/Source/Runtime/RHI/Public/PipelineFileCache.h
+++ b/Engine/Source/Runtime/RHI/Public/PipelineFileCache.h
@@ -145,6 +145,8 @@ struct FPipelineCacheFileFormatPSO
 
 		uint8	MultiViewCount;
 		bool	bHasFragmentDensityAttachment;
+
+		bool	bDepthBounds;
 		
 		RHI_API FString ToString() const;
 		RHI_API void AddToReadableString(TReadableStringBuilder& OutBuilder) const;
diff --git a/Engine/Source/Runtime/RenderCore/Private/PipelineCacheUtilities.cpp b/Engine/Source/Runtime/RenderCore/Private/PipelineCacheUtilities.cpp
--- a/Engine/Source/Runtime/RenderCore/Private/PipelineCacheUtilities.cpp
+++ b/Engine/Source/Runtime/RenderCore/Private/PipelineCacheUtilities.cpp
@@ -68,8 +68,9 @@ namespace Private
 		enum class EVersion : int32
 		{
 			AddingPipelineCacheVersion = 5,
+			AddingDepthBounds = 6,
 
-			Current = AddingPipelineCacheVersion
+			Current = AddingDepthBounds
 		};
 
 		/** Magic to reject other files */
diff --git a/Engine/Source/Runtime/RenderCore/Private/ShaderPipelineCache.cpp b/Engine/Source/Runtime/RenderCore/Private/ShaderPipelineCache.cpp
--- a/Engine/Source/Runtime/RenderCore/Private/ShaderPipelineCache.cpp
+++ b/Engine/Source/Runtime/RenderCore/Private/ShaderPipelineCache.cpp
@@ -1074,6 +1074,8 @@ bool FShaderPipelineCacheTask::Precompile(FRHICommandListImmediate& RHICmdList,
 				GraphicsInitializer.MultiViewCount = PSO.GraphicsDesc.MultiViewCount;
 				GraphicsInitializer.bHasFragmentDensityAttachment = PSO.GraphicsDesc.bHasFragmentDensityAttachment;
 
+				GraphicsInitializer.bDepthBounds = PSO.GraphicsDesc.bDepthBounds;
+
 				GraphicsInitializer.DepthStencilTargetFormat = PSO.GraphicsDesc.DepthStencilFormat;
 				GraphicsInitializer.DepthStencilTargetFlag = PSO.GraphicsDesc.DepthStencilFlags;
 				GraphicsInitializer.DepthTargetLoadAction = PSO.GraphicsDesc.DepthLoad;
