Fix Mac build errors after driver denylist refactor.

diff --git a/Engine/Source/Runtime/RHI/Private/DynamicRHI.cpp b/Engine/Source/Runtime/RHI/Private/DynamicRHI.cpp
--- a/Engine/Source/Runtime/RHI/Private/DynamicRHI.cpp
+++ b/Engine/Source/Runtime/RHI/Private/DynamicRHI.cpp
@@ -241,9 +241,7 @@ void RHIDetectAndWarnOfBadDrivers(bool bHasEditorToken)
 		return;
 	}
 
-	int32 CVarValue = CVarWarnOfBadDrivers.GetValueOnGameThread();
-
-	if (!CVarValue || GRHIVendorId == 0 || bHasEditorToken || FApp::IsUnattended())
+	if (!GWarnOfBadDrivers || GRHIVendorId == 0 || bHasEditorToken || FApp::IsUnattended())
 	{
 		return;
 	}
@@ -256,7 +254,7 @@ void RHIDetectAndWarnOfBadDrivers(bool bHasEditorToken)
 								 *NSLOCTEXT("MessageDialog", "UpdateMacOSX_Title", "Update macOS").ToString());
 		
 #if !UE_BUILD_SHIPPING
-		if (CVarBadDriverWarningIsFatal.GetValueOnGameThread())
+		if (GBadDriverWarningIsFatal)
 		{
 			// Force a fatal error depending on CVar
 		UE_LOG(LogRHI, Fatal, TEXT("Fatal crash requested when graphics drivers are out of date.\n")
