Remove unused RHI raytracing API functions on IRHICommandContext
 - The virtual function definitions were unnecessary since they are defined by the base IRHIComputeContext class.

diff --git a/Engine/Source/Runtime/RHI/Public/RHIContext.h b/Engine/Source/Runtime/RHI/Public/RHIContext.h
--- a/Engine/Source/Runtime/RHI/Public/RHIContext.h
+++ b/Engine/Source/Runtime/RHI/Public/RHIContext.h
@@ -782,26 +782,6 @@ public:
 		checkNoEntry();
 	}
 
-	virtual void RHIBuildAccelerationStructures(TConstArrayView<FRayTracingGeometryBuildParams> Params, const FRHIBufferRange& ScratchBufferRange)
-	{
-		checkNoEntry();
-	}
-
-	void RHIBuildAccelerationStructures(TConstArrayView<FRayTracingGeometryBuildParams> Params)
-	{
-		checkNoEntry();
-	}
-
-	void RHIBuildAccelerationStructure(FRHIRayTracingGeometry* Geometry)
-	{
-		checkNoEntry();
-	}
-
-	virtual void RHIBuildAccelerationStructure(const FRayTracingSceneBuildParams& SceneBuildParams)
-	{
-		checkNoEntry();
-	}
-
 	virtual void RHIRayTraceDispatch(FRHIRayTracingPipelineState* RayTracingPipelineState, FRHIRayTracingShader* RayGenShader,
 		FRHIRayTracingScene* Scene,
 		const FRayTracingShaderBindings& GlobalResourceBindings,
