Removing per-parameter binding code that was deprecated in 5.3

diff --git a/Engine/Source/Runtime/RHI/Public/RHICommandList.h b/Engine/Source/Runtime/RHI/Public/RHICommandList.h
--- a/Engine/Source/Runtime/RHI/Public/RHICommandList.h
+++ b/Engine/Source/Runtime/RHI/Public/RHICommandList.h
@@ -2329,12 +2329,6 @@ protected:
 		: FRHICommandListBase(GPUMask, bImmediate)
 	{}
 
-public:
-	UE_DEPRECATED(5.3, "FlushAllPendingComputeParameters isn't needed with automatic batching removed.")
-	void FlushAllPendingComputeParameters()
-	{
-	}
-
 public:
 	static inline FRHIComputeCommandList& Get(FRHICommandListBase& RHICmdList)
 	{
@@ -2400,22 +2394,6 @@ public:
 		ALLOC_COMMAND(FRHICommandSetUniformBufferDynamicOffset)(Slot, Offset);
 	}
 
-	UE_DEPRECATED(5.3, "FRHIBatchedShaderParameters and SetBatchedShaderParameters should be used instead of setting individual parameters.")
-	FORCEINLINE_DEBUGGABLE void SetShaderUniformBuffer(FRHIComputeShader* Shader, uint32 BaseIndex, FRHIUniformBuffer* UniformBuffer)
-	{
-		FRHIBatchedShaderParameters& BatchedParameters = GetScratchShaderParameters();
-		BatchedParameters.SetShaderUniformBuffer(BaseIndex, UniformBuffer);
-		SetBatchedShaderParameters(Shader, BatchedParameters);
-	}
-
-	UE_DEPRECATED(5.3, "FRHIBatchedShaderParameters and SetBatchedShaderParameters should be used instead of setting individual parameters.")
-	FORCEINLINE void SetShaderUniformBuffer(const FComputeShaderRHIRef& Shader, uint32 BaseIndex, FRHIUniformBuffer* UniformBuffer)
-	{
-PRAGMA_DISABLE_DEPRECATION_WARNINGS
-		SetShaderUniformBuffer(Shader.GetReference(), BaseIndex, UniformBuffer);
-PRAGMA_ENABLE_DEPRECATION_WARNINGS
-	}
-
 	FORCEINLINE_DEBUGGABLE void SetShaderParameters(
 		FRHIComputeShader* InShader
 		, TConstArrayView<uint8> InParametersData
@@ -2483,110 +2461,6 @@ PRAGMA_ENABLE_DEPRECATION_WARNINGS
 		}
 	}
 
-	UE_DEPRECATED(5.3, "FRHIBatchedShaderParameters and SetBatchedShaderParameters should be used instead of setting individual parameters.")
-	FORCEINLINE_DEBUGGABLE void SetShaderParameter(FRHIComputeShader* Shader, uint32 BufferIndex, uint32 BaseIndex, uint32 NumBytes, const void* NewValue)
-	{
-		FRHIBatchedShaderParameters& BatchedParameters = GetScratchShaderParameters();
-		BatchedParameters.SetShaderParameter(BufferIndex, BaseIndex, NumBytes, NewValue);
-		SetBatchedShaderParameters(Shader, BatchedParameters);
-	}
-
-	UE_DEPRECATED(5.3, "FRHIBatchedShaderParameters and SetBatchedShaderParameters should be used instead of setting individual parameters.")
-	FORCEINLINE void SetShaderParameter(FComputeShaderRHIRef& Shader, uint32 BufferIndex, uint32 BaseIndex, uint32 NumBytes, const void* NewValue)
-	{
-PRAGMA_DISABLE_DEPRECATION_WARNINGS
-		SetShaderParameter(Shader.GetReference(), BufferIndex, BaseIndex, NumBytes, NewValue);
-PRAGMA_ENABLE_DEPRECATION_WARNINGS
-	}
-
-	UE_DEPRECATED(5.3, "FRHIBatchedShaderParameters and SetBatchedShaderParameters should be used instead of setting individual parameters.")
-	FORCEINLINE_DEBUGGABLE void SetShaderTexture(FRHIComputeShader* Shader, uint32 TextureIndex, FRHITexture* Texture)
-	{
-		FRHIBatchedShaderParameters& BatchedParameters = GetScratchShaderParameters();
-		BatchedParameters.SetShaderTexture(TextureIndex, Texture);
-		SetBatchedShaderParameters(Shader, BatchedParameters);
-	}
-
-	UE_DEPRECATED(5.3, "FRHIBatchedShaderParameters and SetBatchedShaderParameters should be used instead of setting individual parameters.")
-	FORCEINLINE_DEBUGGABLE void SetShaderResourceViewParameter(FRHIComputeShader* Shader, uint32 SamplerIndex, FRHIShaderResourceView* SRV)
-	{
-		FRHIBatchedShaderParameters& BatchedParameters = GetScratchShaderParameters();
-		BatchedParameters.SetShaderResourceViewParameter(SamplerIndex, SRV);
-		SetBatchedShaderParameters(Shader, BatchedParameters);
-	}
-
-	UE_DEPRECATED(5.3, "FRHIBatchedShaderParameters and SetBatchedShaderParameters should be used instead of setting individual parameters.")
-	FORCEINLINE_DEBUGGABLE void SetShaderSampler(FRHIComputeShader* Shader, uint32 SamplerIndex, FRHISamplerState* State)
-	{
-		// Immutable samplers can't be set dynamically
-		check(!State->IsImmutable());
-		if (State->IsImmutable())
-		{
-			return;
-		}
-
-		FRHIBatchedShaderParameters& BatchedParameters = GetScratchShaderParameters();
-		BatchedParameters.SetShaderSampler(SamplerIndex, State);
-		SetBatchedShaderParameters(Shader, BatchedParameters);
-	}
-
-	UE_DEPRECATED(5.3, "FRHIBatchedShaderParameters and SetBatchedShaderParameters should be used instead of setting individual parameters.")
-	FORCEINLINE_DEBUGGABLE void SetUAVParameter(FRHIComputeShader* Shader, uint32 UAVIndex, FRHIUnorderedAccessView* UAV)
-	{
-		FRHIBatchedShaderParameters& BatchedParameters = GetScratchShaderParameters();
-		BatchedParameters.SetUAVParameter(UAVIndex, UAV);
-		SetBatchedShaderParameters(Shader, BatchedParameters);
-	}
-
-	UE_DEPRECATED(5.3, "FRHIBatchedShaderParameters and SetBatchedShaderParameters should be used instead of setting individual parameters.")
-	FORCEINLINE_DEBUGGABLE void SetUAVParameter(FRHIComputeShader* Shader, uint32 UAVIndex, FRHIUnorderedAccessView* UAV, uint32 InitialCount)
-	{
-		checkNoEntry(); // @todo: support append/consume buffers
-
-		FRHIBatchedShaderParameters& BatchedParameters = GetScratchShaderParameters();
-		BatchedParameters.SetUAVParameter(UAVIndex, UAV);
-		SetBatchedShaderParameters(Shader, BatchedParameters);
-	}
-
-	UE_DEPRECATED(5.3, "FRHIBatchedShaderParameters and SetBatchedShaderParameters should be used instead of setting individual parameters.")
-	FORCEINLINE_DEBUGGABLE void SetBindlessTexture(FRHIComputeShader* Shader, uint32 Index, FRHITexture* Texture)
-	{
-		FRHIBatchedShaderParameters& BatchedParameters = GetScratchShaderParameters();
-		BatchedParameters.SetBindlessTexture(Index, Texture);
-		SetBatchedShaderParameters(Shader, BatchedParameters);
-	}
-
-	UE_DEPRECATED(5.3, "FRHIBatchedShaderParameters and SetBatchedShaderParameters should be used instead of setting individual parameters.")
-	FORCEINLINE_DEBUGGABLE void SetBindlessResourceView(FRHIComputeShader* Shader, uint32 Index, FRHIShaderResourceView* SRV)
-	{
-		FRHIBatchedShaderParameters& BatchedParameters = GetScratchShaderParameters();
-		BatchedParameters.SetBindlessResourceView(Index, SRV);
-		SetBatchedShaderParameters(Shader, BatchedParameters);
-	}
-
-	UE_DEPRECATED(5.3, "FRHIBatchedShaderParameters and SetBatchedShaderParameters should be used instead of setting individual parameters.")
-	FORCEINLINE_DEBUGGABLE void SetBindlessSampler(FRHIComputeShader* Shader, uint32 Index, FRHISamplerState* State)
-	{
-		// Immutable samplers can't be set dynamically
-		check(!State->IsImmutable());
-		if (State->IsImmutable())
-		{
-			return;
-		}
-
-		FRHIBatchedShaderParameters& BatchedParameters = GetScratchShaderParameters();
-		BatchedParameters.SetBindlessSampler(Index, State);
-		SetBatchedShaderParameters(Shader, BatchedParameters);
-	}
-
-	UE_DEPRECATED(5.3, "FRHIBatchedShaderParameters and SetBatchedShaderParameters should be used instead of setting individual parameters.")
-	FORCEINLINE_DEBUGGABLE void SetBindlessUAV(FRHIComputeShader* Shader, uint32 Index, FRHIUnorderedAccessView* UAV)
-	{
-		FRHIBatchedShaderParameters& BatchedParameters = GetScratchShaderParameters();
-		BatchedParameters.SetBindlessUAV(Index, UAV);
-		SetBatchedShaderParameters(Shader, BatchedParameters);
-	}
-
 	FORCEINLINE_DEBUGGABLE void SetComputePipelineState(FComputePipelineState* ComputePipelineState, FRHIComputeShader* ComputeShader)
 	{
 		OnBoundShaderChanged(ComputeShader);
@@ -3245,25 +3119,6 @@ public:
 		FRHICommandList::EnqueueLambda(TEXT("TRHILambdaCommand"), Forward<LAMBDA>(Lambda));
 	}
 
-	using FRHIComputeCommandList::SetShaderUniformBuffer;
-
-	UE_DEPRECATED(5.3, "FRHIBatchedShaderParameters and SetBatchedShaderParameters should be used instead of setting individual parameters.")
-	FORCEINLINE_DEBUGGABLE void SetShaderUniformBuffer(FRHIGraphicsShader* Shader, uint32 BaseIndex, FRHIUniformBuffer* UniformBuffer)
-	{
-		FRHIBatchedShaderParameters& BatchedParameters = GetScratchShaderParameters();
-		BatchedParameters.SetShaderUniformBuffer(BaseIndex, UniformBuffer);
-		SetBatchedShaderParameters(Shader, BatchedParameters);
-	}
-
-	template <typename TShaderRHI>
-	UE_DEPRECATED(5.3, "FRHIBatchedShaderParameters and SetBatchedShaderParameters should be used instead of setting individual parameters.")
-	FORCEINLINE void SetShaderUniformBuffer(const TRefCountPtr<TShaderRHI>& Shader, uint32 BaseIndex, FRHIUniformBuffer* UniformBuffer)
-	{
-PRAGMA_DISABLE_DEPRECATION_WARNINGS
-		SetShaderUniformBuffer(Shader.GetReference(), BaseIndex, UniformBuffer);
-PRAGMA_ENABLE_DEPRECATION_WARNINGS
-	}
-
 	using FRHIComputeCommandList::SetShaderParameters;
 
 	FORCEINLINE_DEBUGGABLE void SetShaderParameters(
@@ -3339,154 +3194,6 @@ PRAGMA_ENABLE_DEPRECATION_WARNINGS
 		}
 	}
 
-	using FRHIComputeCommandList::SetShaderParameter;
-
-	UE_DEPRECATED(5.3, "FRHIBatchedShaderParameters and SetBatchedShaderParameters should be used instead of setting individual parameters.")
-	FORCEINLINE_DEBUGGABLE void SetShaderParameter(FRHIGraphicsShader* Shader, uint32 BufferIndex, uint32 BaseIndex, uint32 NumBytes, const void* NewValue)
-	{
-		FRHIBatchedShaderParameters& BatchedParameters = GetScratchShaderParameters();
-		BatchedParameters.SetShaderParameter(BufferIndex, BaseIndex, NumBytes, NewValue);
-		SetBatchedShaderParameters(Shader, BatchedParameters);
-	}
-
-	template <typename TShaderRHI>
-	UE_DEPRECATED(5.3, "FRHIBatchedShaderParameters and SetBatchedShaderParameters should be used instead of setting individual parameters.")
-	FORCEINLINE void SetShaderParameter(const TRefCountPtr<TShaderRHI>& Shader, uint32 BufferIndex, uint32 BaseIndex, uint32 NumBytes, const void* NewValue)
-	{
-PRAGMA_DISABLE_DEPRECATION_WARNINGS
-		SetShaderParameter(Shader.GetReference(), BufferIndex, BaseIndex, NumBytes, NewValue);
-PRAGMA_ENABLE_DEPRECATION_WARNINGS
-	}
-
-	using FRHIComputeCommandList::SetShaderTexture;
-
-	UE_DEPRECATED(5.3, "FRHIBatchedShaderParameters and SetBatchedShaderParameters should be used instead of setting individual parameters.")
-	FORCEINLINE_DEBUGGABLE void SetShaderTexture(FRHIGraphicsShader* Shader, uint32 TextureIndex, FRHITexture* Texture)
-	{
-		FRHIBatchedShaderParameters& BatchedParameters = GetScratchShaderParameters();
-		BatchedParameters.SetShaderTexture(TextureIndex, Texture);
-		SetBatchedShaderParameters(Shader, BatchedParameters);
-	}
-
-	template <typename TShaderRHI>
-	UE_DEPRECATED(5.3, "FRHIBatchedShaderParameters and SetBatchedShaderParameters should be used instead of setting individual parameters.")
-	FORCEINLINE_DEBUGGABLE void SetShaderTexture(const TRefCountPtr<TShaderRHI>& Shader, uint32 TextureIndex, FRHITexture* Texture)
-	{
-PRAGMA_DISABLE_DEPRECATION_WARNINGS
-		SetShaderTexture(Shader.GetReference(), TextureIndex, Texture);
-PRAGMA_ENABLE_DEPRECATION_WARNINGS
-	}
-
-	using FRHIComputeCommandList::SetShaderResourceViewParameter;
-
-	UE_DEPRECATED(5.3, "FRHIBatchedShaderParameters and SetBatchedShaderParameters should be used instead of setting individual parameters.")
-	FORCEINLINE_DEBUGGABLE void SetShaderResourceViewParameter(FRHIGraphicsShader* Shader, uint32 SamplerIndex, FRHIShaderResourceView* SRV)
-	{
-		FRHIBatchedShaderParameters& BatchedParameters = GetScratchShaderParameters();
-		BatchedParameters.SetShaderResourceViewParameter(SamplerIndex, SRV);
-		SetBatchedShaderParameters(Shader, BatchedParameters);
-	}
-
-	template <typename TShaderRHI>
-	UE_DEPRECATED(5.3, "FRHIBatchedShaderParameters and SetBatchedShaderParameters should be used instead of setting individual parameters.")
-	FORCEINLINE_DEBUGGABLE void SetShaderResourceViewParameter(const TRefCountPtr<TShaderRHI>& Shader, uint32 SamplerIndex, FRHIShaderResourceView* SRV)
-	{
-PRAGMA_DISABLE_DEPRECATION_WARNINGS
-		SetShaderResourceViewParameter(Shader.GetReference(), SamplerIndex, SRV);
-PRAGMA_ENABLE_DEPRECATION_WARNINGS
-	}
-
-	using FRHIComputeCommandList::SetShaderSampler;
-
-	UE_DEPRECATED(5.3, "FRHIBatchedShaderParameters and SetBatchedShaderParameters should be used instead of setting individual parameters.")
-	FORCEINLINE_DEBUGGABLE void SetShaderSampler(FRHIGraphicsShader* Shader, uint32 SamplerIndex, FRHISamplerState* State)
-	{
-		// Immutable samplers can't be set dynamically
-		check(!State->IsImmutable());
-		if (State->IsImmutable())
-		{
-			return;
-		}
-
-		FRHIBatchedShaderParameters& BatchedParameters = GetScratchShaderParameters();
-		BatchedParameters.SetShaderSampler(SamplerIndex, State);
-		SetBatchedShaderParameters(Shader, BatchedParameters);
-	}
-
-	template <typename TShaderRHI>
-	UE_DEPRECATED(5.3, "FRHIBatchedShaderParameters and SetBatchedShaderParameters should be used instead of setting individual parameters.")
-	FORCEINLINE_DEBUGGABLE void SetShaderSampler(const TRefCountPtr<TShaderRHI>& Shader, uint32 SamplerIndex, FRHISamplerState* State)
-	{
-PRAGMA_DISABLE_DEPRECATION_WARNINGS
-		SetShaderSampler(Shader.GetReference(), SamplerIndex, State);
-PRAGMA_ENABLE_DEPRECATION_WARNINGS
-	}
-
-	using FRHIComputeCommandList::SetUAVParameter;
-
-	UE_DEPRECATED(5.3, "FRHIBatchedShaderParameters and SetBatchedShaderParameters should be used instead of setting individual parameters.")
-	FORCEINLINE_DEBUGGABLE void SetUAVParameter(FRHIPixelShader* Shader, uint32 UAVIndex, FRHIUnorderedAccessView* UAV)
-	{
-		FRHIBatchedShaderParameters& BatchedParameters = GetScratchShaderParameters();
-		BatchedParameters.SetUAVParameter(UAVIndex, UAV);
-		SetBatchedShaderParameters(Shader, BatchedParameters);
-	}
-
-	UE_DEPRECATED(5.3, "FRHIBatchedShaderParameters and SetBatchedShaderParameters should be used instead of setting individual parameters.")
-	FORCEINLINE_DEBUGGABLE void SetUAVParameter(const TRefCountPtr<FRHIPixelShader>& Shader, uint32 UAVIndex, FRHIUnorderedAccessView* UAV)
-	{
-PRAGMA_DISABLE_DEPRECATION_WARNINGS
-		SetUAVParameter(Shader.GetReference(), UAVIndex, UAV);
-PRAGMA_ENABLE_DEPRECATION_WARNINGS
-	}
-
-	using FRHIComputeCommandList::SetBindlessTexture;
-
-	UE_DEPRECATED(5.3, "FRHIBatchedShaderParameters and SetBatchedShaderParameters should be used instead of setting individual parameters.")
-	FORCEINLINE_DEBUGGABLE void SetBindlessTexture(FRHIGraphicsShader* Shader, uint32 Index, FRHITexture* Texture)
-	{
-		FRHIBatchedShaderParameters& BatchedParameters = GetScratchShaderParameters();
-		BatchedParameters.SetBindlessTexture(Index, Texture);
-		SetBatchedShaderParameters(Shader, BatchedParameters);
-	}
-
-	using FRHIComputeCommandList::SetBindlessResourceView;
-
-	UE_DEPRECATED(5.3, "FRHIBatchedShaderParameters and SetBatchedShaderParameters should be used instead of setting individual parameters.")
-	FORCEINLINE_DEBUGGABLE void SetBindlessResourceView(FRHIGraphicsShader* Shader, uint32 Index, FRHIShaderResourceView* SRV)
-	{
-		FRHIBatchedShaderParameters& BatchedParameters = GetScratchShaderParameters();
-		BatchedParameters.SetBindlessResourceView(Index, SRV);
-		SetBatchedShaderParameters(Shader, BatchedParameters);
-	}
-
-	using FRHIComputeCommandList::SetBindlessSampler;
-
-	UE_DEPRECATED(5.3, "FRHIBatchedShaderParameters and SetBatchedShaderParameters should be used instead of setting individual parameters.")
-	FORCEINLINE_DEBUGGABLE void SetBindlessSampler(FRHIGraphicsShader* Shader, uint32 Index, FRHISamplerState* State)
-	{
-		// Immutable samplers can't be set dynamically
-		check(!State->IsImmutable());
-		if (State->IsImmutable())
-		{
-			return;
-		}
-
-		FRHIBatchedShaderParameters& BatchedParameters = GetScratchShaderParameters();
-		BatchedParameters.SetBindlessSampler(Index, State);
-		SetBatchedShaderParameters(Shader, BatchedParameters);
-	}
-
-	using FRHIComputeCommandList::SetBindlessUAV;
-
-	UE_DEPRECATED(5.3, "FRHIBatchedShaderParameters and SetBatchedShaderParameters should be used instead of setting individual parameters.")
-	FORCEINLINE_DEBUGGABLE void SetBindlessUAV(FRHIPixelShader* Shader, uint32 Index, FRHIUnorderedAccessView* UAV)
-	{
-		FRHIBatchedShaderParameters& BatchedParameters = GetScratchShaderParameters();
-		BatchedParameters.SetBindlessUAV(Index, UAV);
-		SetBatchedShaderParameters(Shader, BatchedParameters);
-	}
-
 	FORCEINLINE_DEBUGGABLE void SetBlendFactor(const FLinearColor& BlendFactor = FLinearColor::White)
 	{
 		//check(IsOutsideRenderPass());
diff --git a/Engine/Source/Runtime/RenderCore/Private/ShaderParameters.cpp b/Engine/Source/Runtime/RenderCore/Private/ShaderParameters.cpp
--- a/Engine/Source/Runtime/RenderCore/Private/ShaderParameters.cpp
+++ b/Engine/Source/Runtime/RenderCore/Private/ShaderParameters.cpp
@@ -16,9 +16,6 @@
 
 IMPLEMENT_TYPE_LAYOUT(FShaderParameter);
 IMPLEMENT_TYPE_LAYOUT(FShaderResourceParameter);
-PRAGMA_DISABLE_DEPRECATION_WARNINGS
-IMPLEMENT_TYPE_LAYOUT(FRWShaderParameter);
-PRAGMA_ENABLE_DEPRECATION_WARNINGS
 IMPLEMENT_TYPE_LAYOUT(FShaderUniformBufferParameter);
 
 void FShaderParameter::Bind(const FShaderParameterMap& ParameterMap,const TCHAR* ParameterName,EShaderParameterFlags Flags)
diff --git a/Engine/Source/Runtime/RenderCore/Public/ShaderParameterUtils.h b/Engine/Source/Runtime/RenderCore/Public/ShaderParameterUtils.h
--- a/Engine/Source/Runtime/RenderCore/Public/ShaderParameterUtils.h
+++ b/Engine/Source/Runtime/RenderCore/Public/ShaderParameterUtils.h
@@ -409,81 +409,6 @@ inline void UnsetShaderParametersLegacyCS(TRHICmdList& RHICmdList, const TShader
 	}
 }
 
-/**
- * Sets the value of a  shader parameter.  Template'd on shader type
- * A template parameter specified the type of the parameter value.
- * NOTE: Shader should be the param ref type, NOT the param type, since Shader is passed by value. 
- * Otherwise AddRef/ReleaseRef will be called many times.
- */
-template<typename ShaderRHIParamRef, class ParameterType, typename TRHICmdList>
-UE_DEPRECATED(5.3, "SetShaderValue with FRHIBatchedShaderParameters should be used.")
-void SetShaderValue(
-	TRHICmdList& RHICmdList,
-	const ShaderRHIParamRef& Shader,
-	const FShaderParameter& Parameter,
-	const ParameterType& Value,
-	uint32 ElementIndex = 0
-	)
-{
-	FRHIBatchedShaderParameters& BatchedParameters = RHICmdList.GetScratchShaderParameters();
-	SetShaderValue(BatchedParameters, Parameter, Value, ElementIndex);
-	RHICmdList.SetBatchedShaderParameters(Shader, BatchedParameters);
-}
-
-
-template<typename ShaderRHIParamRef, class ParameterType>
-UE_DEPRECATED(5.3, "SetShaderValue with FRHIBatchedShaderParameters should be used.")
-void SetShaderValueOnContext(
-	IRHICommandContext& RHICmdListContext,
-	const ShaderRHIParamRef& Shader,
-	const FShaderParameter& Parameter,
-	const ParameterType& Value,
-	uint32 ElementIndex = 0
-	)
-{
-	FRHIBatchedShaderParameters BatchedParameters;
-	SetShaderValue(BatchedParameters, Parameter, Value, ElementIndex);
-	RHICmdListContext.RHISetBatchedShaderParameters(Shader, BatchedParameters);
-}
-
-/**
- * Sets the value of a shader parameter array.  Template'd on shader type
- * A template parameter specified the type of the parameter value.
- * NOTE: Shader should be the param ref type, NOT the param type, since Shader is passed by value. 
- * Otherwise AddRef/ReleaseRef will be called many times.
- */
-template<typename ShaderRHIParamRef,class ParameterType, typename TRHICmdList>
-UE_DEPRECATED(5.3, "SetShaderValueArray with FRHIBatchedShaderParameters should be used.")
-void SetShaderValueArray(
-	TRHICmdList& RHICmdList,
-	const ShaderRHIParamRef& Shader,
-	const FShaderParameter& Parameter,
-	const ParameterType* Values,
-	uint32 NumElements,
-	uint32 BaseElementIndex = 0
-	)
-{
-	FRHIBatchedShaderParameters& BatchedParameters = RHICmdList.GetScratchShaderParameters();
-	SetShaderValueArray(BatchedParameters, Parameter, Values, NumElements, BaseElementIndex);
-	RHICmdList.SetBatchedShaderParameters(Shader, BatchedParameters);
-}
-
-/** Specialization of the above for C++ bool type. */
-template<typename ShaderRHIParamRef, typename TRHICmdList>
-UE_DEPRECATED(5.3, "SetShaderValueArray with FRHIBatchedShaderParameters should be used.")
-void SetShaderValueArray(
-	TRHICmdList& RHICmdList,
-	const ShaderRHIParamRef& Shader,
-	const FShaderParameter& Parameter,
-	const bool* Values,
-	uint32 NumElements,
-	uint32 BaseElementIndex = 0
-	)
-{
-	UE_LOG(LogShaders, Fatal, TEXT("SetShaderValueArray does not support bool arrays."));
-}
-
-
 // LWC_TODO: Setting guards to catch attempts to pass a type with double components. Could just convert these to the correct type internally, but would prefer to catch potential issues + optimize where possible.
 #define GUARD_SETSHADERVALUE(_TYPE)																																						\
 template<typename ShaderRHIParamRef, typename TRHICmdList>																																\
@@ -507,31 +432,6 @@ GUARD_SETSHADERVALUE(FQuat4)
 GUARD_SETSHADERVALUE(::FSphere3)
 GUARD_SETSHADERVALUE(FBox3)
 
-/**
- * Sets the value of a shader surface parameter (e.g. to access MSAA samples).
- * Template'd on shader type (e.g. pixel shader or compute shader).
- */
-template<typename TRHIShader, typename TRHICmdList>
-UE_DEPRECATED(5.3, "SetTextureParameter with FRHIBatchedShaderParameters should be used.")
-FORCEINLINE void SetTextureParameter(TRHICmdList& RHICmdList, TRHIShader* Shader, const FShaderResourceParameter& Parameter, FRHITexture* TextureRHI)
-{
-	FRHIBatchedShaderParameters& BatchedParameters = RHICmdList.GetScratchShaderParameters();
-	SetTextureParameter(BatchedParameters, Parameter, TextureRHI);
-	RHICmdList.SetBatchedShaderParameters(Shader, BatchedParameters);
-}
-
-/**
- * Sets the value of a shader sampler parameter. Template'd on shader type.
- */
-template<typename TRHIShader, typename TRHICmdList>
-UE_DEPRECATED(5.3, "SetSamplerParameter with FRHIBatchedShaderParameters should be used.")
-FORCEINLINE void SetSamplerParameter(TRHICmdList& RHICmdList, TRHIShader* Shader, const FShaderResourceParameter& Parameter, FRHISamplerState* SamplerStateRHI)
-{
-	FRHIBatchedShaderParameters& BatchedParameters = RHICmdList.GetScratchShaderParameters();
-	SetSamplerParameter(BatchedParameters, Parameter, SamplerStateRHI);
-	RHICmdList.SetBatchedShaderParameters(Shader, BatchedParameters);
-}
-
 /**
  * Sets the value of a shader texture parameter. Template'd on shader type.
  */
@@ -542,213 +442,3 @@ FORCEINLINE void SetTextureParameter(TRHICmdList& RHICmdList, TRHIShader* Shader
 	SetTextureParameter(BatchedParameters, TextureParameter, SamplerParameter, SamplerStateRHI, TextureRHI);
 	RHICmdList.SetBatchedShaderParameters(Shader, BatchedParameters);
 }
-
-/**
- * Sets the value of a shader texture parameter.  Template'd on shader type
- */
-template<typename TRHIShader, typename TRHICmdList>
-UE_DEPRECATED(5.3, "SetTextureParameter with FRHIBatchedShaderParameters should be used.")
-FORCEINLINE void SetTextureParameter(TRHICmdList& RHICmdList, TRHIShader* Shader, const FShaderResourceParameter& TextureParameter, const FShaderResourceParameter& SamplerParameter, const FTexture* Texture)
-{
-	FRHIBatchedShaderParameters& BatchedParameters = RHICmdList.GetScratchShaderParameters();
-	SetTextureParameter(BatchedParameters, TextureParameter, SamplerParameter, Texture);
-	RHICmdList.SetBatchedShaderParameters(Shader, BatchedParameters);
-}
-
-/**
- * Sets the value of a shader resource view parameter
- * Template'd on shader type (e.g. pixel shader or compute shader).
- */
-template<typename TRHIShader, typename TRHICmdList>
-UE_DEPRECATED(5.3, "SetSRVParameter with FRHIBatchedShaderParameters should be used.")
-FORCEINLINE void SetSRVParameter(TRHICmdList& RHICmdList, TRHIShader* Shader, const FShaderResourceParameter& Parameter, FRHIShaderResourceView* SRV)
-{
-	FRHIBatchedShaderParameters& BatchedParameters = RHICmdList.GetScratchShaderParameters();
-	SetSRVParameter(BatchedParameters, Parameter, SRV);
-	RHICmdList.SetBatchedShaderParameters(Shader, BatchedParameters);
-}
-
-template<typename TRHIShader, typename TRHICmdList>
-UE_DEPRECATED(5.3, "SetSRVParameter with FRHIBatchedShaderParameters should be used.")
-FORCEINLINE void SetSRVParameter(TRHICmdList& RHICmdList, const TRefCountPtr<TRHIShader>& Shader, const FShaderResourceParameter& Parameter, FRHIShaderResourceView* SRV)
-{
-	FRHIBatchedShaderParameters& BatchedParameters = RHICmdList.GetScratchShaderParameters();
-	SetSRVParameter(BatchedParameters, Parameter, SRV);
-	RHICmdList.SetBatchedShaderParameters(Shader.GetReference(), BatchedParameters);
-}
-
-template<typename TRHIShader, typename TRHICmdList>
-UE_DEPRECATED(5.3, "SetUAVParameter with FRHIBatchedShaderParameters should be used.")
-FORCEINLINE void SetUAVParameterSafeShader(TRHICmdList& RHICmdList, TRHIShader* Shader, const FShaderResourceParameter& Parameter, FRHIUnorderedAccessView* UAV)
-{
-	FRHIBatchedShaderParameters& BatchedParameters = RHICmdList.GetScratchShaderParameters();
-	SetUAVParameter(BatchedParameters, Parameter, UAV);
-	RHICmdList.SetBatchedShaderParameters(Shader, BatchedParameters);
-}
-
-UE_DEPRECATED(5.3, "SetUAVParameter with FRHIBatchedShaderParameters should be used.")
-FORCEINLINE void SetUAVParameter(FRHIComputeCommandList& RHICmdList, FRHIComputeShader* Shader, const FShaderResourceParameter& Parameter, FRHIUnorderedAccessView* UAV)
-{
-	FRHIBatchedShaderParameters& BatchedParameters = RHICmdList.GetScratchShaderParameters();
-	SetUAVParameter(BatchedParameters, Parameter, UAV);
-	RHICmdList.SetBatchedShaderParameters(Shader, BatchedParameters);
-}
-
-UE_DEPRECATED(5.3, "SetUAVParameter with FRHIBatchedShaderParameters should be used.")
-FORCEINLINE void SetUAVParameter(FRHICommandList& RHICmdList, FRHIPixelShader* Shader, const FShaderResourceParameter& Parameter, FRHIUnorderedAccessView* UAV)
-{
-	FRHIBatchedShaderParameters& BatchedParameters = RHICmdList.GetScratchShaderParameters();
-	SetUAVParameter(BatchedParameters, Parameter, UAV);
-	RHICmdList.SetBatchedShaderParameters(Shader, BatchedParameters);
-}
-
-
-template<typename TRHICmdList>
-UE_DEPRECATED(5.3, "SetUAVParameter with FRHIBatchedShaderParameters should be used.")
-inline bool SetUAVParameterIfCS(TRHICmdList& RHICmdList, FRHIVertexShader* Shader, const FShaderResourceParameter& UAVParameter, FRHIUnorderedAccessView* UAV)
-{
-	return false;
-}
-
-template<typename TRHICmdList>
-UE_DEPRECATED(5.3, "SetUAVParameter with FRHIBatchedShaderParameters should be used.")
-inline bool SetUAVParameterIfCS(TRHICmdList& RHICmdList, FRHIPixelShader* Shader, const FShaderResourceParameter& UAVParameter, FRHIUnorderedAccessView* UAV)
-{
-PRAGMA_DISABLE_DEPRECATION_WARNINGS
-	SetUAVParameter(RHICmdList, Shader, UAVParameter, UAV);
-PRAGMA_ENABLE_DEPRECATION_WARNINGS
-	return UAVParameter.IsBound();
-}
-
-template<typename TRHICmdList>
-UE_DEPRECATED(5.3, "SetUAVParameter with FRHIBatchedShaderParameters should be used.")
-inline bool SetUAVParameterIfCS(TRHICmdList& RHICmdList, FRHIGeometryShader* Shader, const FShaderResourceParameter& UAVParameter, FRHIUnorderedAccessView* UAV)
-{
-	return false;
-}
-
-template<typename TRHICmdList>
-UE_DEPRECATED(5.3, "SetUAVParameter with FRHIBatchedShaderParameters should be used.")
-inline bool SetUAVParameterIfCS(TRHICmdList& RHICmdList, FRHIComputeShader* Shader, const FShaderResourceParameter& UAVParameter, FRHIUnorderedAccessView* UAV)
-{
-PRAGMA_DISABLE_DEPRECATION_WARNINGS
-	SetUAVParameter(RHICmdList, Shader, UAVParameter, UAV);
-PRAGMA_ENABLE_DEPRECATION_WARNINGS
-	return UAVParameter.IsBound();
-}
-
-PRAGMA_DISABLE_DEPRECATION_WARNINGS
-template<typename TShaderRHIRef, typename TRHICmdList>
-inline void FRWShaderParameter::SetBuffer(TRHICmdList& RHICmdList, const TShaderRHIRef& Shader, const FRWBuffer& RWBuffer) const
-{
-	if (!SetUAVParameterIfCS(RHICmdList, Shader, UAVParameter, RWBuffer.UAV))
-	{
-		SetSRVParameter(RHICmdList, Shader, SRVParameter, RWBuffer.SRV);
-	}
-}
-
-template<typename TShaderRHIRef, typename TRHICmdList>
-inline void FRWShaderParameter::SetBuffer(TRHICmdList& RHICmdList, const TShaderRHIRef& Shader, const FRWBufferStructured& RWBuffer) const
-{
-	if (!SetUAVParameterIfCS(RHICmdList, Shader, UAVParameter, RWBuffer.UAV))
-	{
-		SetSRVParameter(RHICmdList, Shader, SRVParameter, RWBuffer.SRV);
-	}
-}
-
-template<typename TShaderRHIRef, typename TRHICmdList>
-inline void FRWShaderParameter::SetTexture(TRHICmdList& RHICmdList, const TShaderRHIRef& Shader, FRHITexture* Texture, FRHIUnorderedAccessView* UAV) const
-{
-	if (!SetUAVParameterIfCS(RHICmdList, Shader, UAVParameter, UAV))
-	{
-		SetTextureParameter(RHICmdList, Shader, SRVParameter, Texture);
-	}
-}
-
-template<typename TRHICmdList>
-inline void FRWShaderParameter::SetUAV(TRHICmdList& RHICmdList, FRHIComputeShader* ComputeShader, FRHIUnorderedAccessView* UAV) const
-{
-	SetUAVParameter(RHICmdList, ComputeShader, UAVParameter, UAV);
-}
-
-template<typename TRHICmdList>
-inline void FRWShaderParameter::UnsetUAV(TRHICmdList& RHICmdList, FRHIComputeShader* ComputeShader) const
-{
-	SetUAVParameter(RHICmdList, ComputeShader,UAVParameter,FUnorderedAccessViewRHIRef());
-}
-PRAGMA_ENABLE_DEPRECATION_WARNINGS
-
-/** Sets the value of a shader uniform buffer parameter to a uniform buffer containing the struct. */
-template<typename TShaderRHIRef, typename TRHICmdList>
-UE_DEPRECATED(5.3, "SetUniformBufferParameter with FRHIBatchedShaderParameters should be used.")
-inline void SetUniformBufferParameter(
-	TRHICmdList& RHICmdList,
-	const TShaderRHIRef& Shader,
-	const FShaderUniformBufferParameter& Parameter,
-	FRHIUniformBuffer* UniformBufferRHI
-	)
-{
-	FRHIBatchedShaderParameters& BatchedParameters = RHICmdList.GetScratchShaderParameters();
-	SetUniformBufferParameter(BatchedParameters, Parameter, UniformBufferRHI);
-	RHICmdList.SetBatchedShaderParameters(Shader, BatchedParameters);
-}
-
-/** Sets the value of a shader uniform buffer parameter to a uniform buffer containing the struct. */
-template<typename TShaderRHIRef, typename TBufferStruct, typename TRHICmdList>
-UE_DEPRECATED(5.3, "SetUniformBufferParameter with FRHIBatchedShaderParameters should be used.")
-inline void SetUniformBufferParameter(
-	TRHICmdList& RHICmdList,
-	const TShaderRHIRef& Shader,
-	const TShaderUniformBufferParameter<TBufferStruct>& Parameter,
-	const TUniformBufferRef<TBufferStruct>& UniformBufferRef
-	)
-{
-	FRHIBatchedShaderParameters& BatchedParameters = RHICmdList.GetScratchShaderParameters();
-	SetUniformBufferParameter(BatchedParameters, Parameter, UniformBufferRef);
-	RHICmdList.SetBatchedShaderParameters(Shader, BatchedParameters);
-}
-
-/** Sets the value of a shader uniform buffer parameter to a uniform buffer containing the struct. */
-template<typename TShaderRHIRef, typename TBufferStruct, typename TRHICmdList>
-UE_DEPRECATED(5.3, "SetUniformBufferParameter with FRHIBatchedShaderParameters should be used.")
-inline void SetUniformBufferParameter(
-	TRHICmdList& RHICmdList,
-	const TShaderRHIRef& Shader,
-	const TShaderUniformBufferParameter<TBufferStruct>& Parameter,
-	const TUniformBuffer<TBufferStruct>& UniformBuffer
-	)
-{
-	FRHIBatchedShaderParameters& BatchedParameters = RHICmdList.GetScratchShaderParameters();
-	SetUniformBufferParameter(BatchedParameters, Parameter, UniformBuffer);
-	RHICmdList.SetBatchedShaderParameters(Shader, BatchedParameters);
-}
-
-/** Sets the value of a shader uniform buffer parameter to a value of the struct. */
-template<typename TShaderRHIRef,typename TBufferStruct>
-UE_DEPRECATED(5.3, "SetUniformBufferParameterImmediate with FRHIBatchedShaderParameters should be used.")
-inline void SetUniformBufferParameterImmediate(
-	FRHICommandList& RHICmdList,
-	const TShaderRHIRef& Shader,
-	const TShaderUniformBufferParameter<TBufferStruct>& Parameter,
-	const TBufferStruct& UniformBufferValue
-	)
-{
-	FRHIBatchedShaderParameters& BatchedParameters = RHICmdList.GetScratchShaderParameters();
-	SetUniformBufferParameterImmediate(BatchedParameters, Parameter, UniformBufferValue);
-	RHICmdList.SetBatchedShaderParameters(Shader, BatchedParameters);
-}
-
-/** Sets the value of a shader uniform buffer parameter to a value of the struct. */
-template<typename TShaderRHIRef,typename TBufferStruct, typename TRHICmdList>
-UE_DEPRECATED(5.3, "SetUniformBufferParameterImmediate with FRHIBatchedShaderParameters should be used.")
-inline void SetUniformBufferParameterImmediate(
-	TRHICmdList& RHICmdList,
-	const TShaderRHIRef& Shader,
-	const TShaderUniformBufferParameter<TBufferStruct>& Parameter,
-	const TBufferStruct& UniformBufferValue
-	)
-{
-	FRHIBatchedShaderParameters& BatchedParameters = RHICmdList.GetScratchShaderParameters();
-	SetUniformBufferParameterImmediate(BatchedParameters, Parameter, UniformBufferValue);
-	RHICmdList.SetBatchedShaderParameters(Shader, BatchedParameters);
-}
\ No newline at end of file
diff --git a/Engine/Source/Runtime/RenderCore/Public/ShaderParameters.h b/Engine/Source/Runtime/RenderCore/Public/ShaderParameters.h
--- a/Engine/Source/Runtime/RenderCore/Public/ShaderParameters.h
+++ b/Engine/Source/Runtime/RenderCore/Public/ShaderParameters.h
@@ -105,64 +105,6 @@ private:
 	LAYOUT_FIELD_INITIALIZED(EShaderParameterType, Type, {});
 };
 
-/** A class that binds either a UAV or SRV of a resource. */
-class UE_DEPRECATED(5.3, "Individual shader parameters should be used instead of FRWShaderParameter.") FRWShaderParameter
-{
-	DECLARE_EXPORTED_TYPE_LAYOUT(FRWShaderParameter, RENDERCORE_API, NonVirtual);
-public:
-
-	void Bind(const FShaderParameterMap& ParameterMap,const TCHAR* BaseName)
-	{
-		SRVParameter.Bind(ParameterMap,BaseName);
-
-		// If the shader wants to bind the parameter as a UAV, the parameter name must start with "RW"
-		FString UAVName = FString(TEXT("RW")) + BaseName;
-		UAVParameter.Bind(ParameterMap,*UAVName);
-
-		// Verify that only one of the UAV or SRV parameters is accessed by the shader.
-		checkf(!(SRVParameter.GetNumResources() && UAVParameter.GetNumResources()),TEXT("Shader binds SRV and UAV of the same resource: %s"),BaseName);
-	}
-
-	FORCEINLINE bool IsBound() const
-	{
-		return SRVParameter.IsBound() || UAVParameter.IsBound();
-	}
-
-	FORCEINLINE bool IsUAVBound() const
-	{
-		return UAVParameter.IsBound();
-	}
-
-	uint32 GetUAVIndex() const
-	{
-		return UAVParameter.GetBaseIndex();
-	}
-
-	friend FArchive& operator<<(FArchive& Ar,FRWShaderParameter& Parameter)
-	{
-		return Ar << Parameter.SRVParameter << Parameter.UAVParameter;
-	}
-
-	template<typename TShaderRHIRef, typename TRHICmdList>
-	inline void SetBuffer(TRHICmdList& RHICmdList, const TShaderRHIRef& Shader, const FRWBuffer& RWBuffer) const;
-
-	template<typename TShaderRHIRef, typename TRHICmdList>
-	inline void SetBuffer(TRHICmdList& RHICmdList, const TShaderRHIRef& Shader, const FRWBufferStructured& RWBuffer) const;
-
-	template<typename TShaderRHIRef, typename TRHICmdList>
-	inline void SetTexture(TRHICmdList& RHICmdList, const TShaderRHIRef& Shader, FRHITexture* Texture, FRHIUnorderedAccessView* UAV) const;
-
-	template<typename TRHICmdList>
-	inline void SetUAV(TRHICmdList& RHICmdList, FRHIComputeShader* ComputeShader, FRHIUnorderedAccessView* UAV) const;
-
-	template<typename TRHICmdList>
-	inline void UnsetUAV(TRHICmdList& RHICmdList, FRHIComputeShader* ComputeShader) const;
-
-private:
-	LAYOUT_FIELD(FShaderResourceParameter, SRVParameter);
-	LAYOUT_FIELD(FShaderResourceParameter, UAVParameter);
-};
-
 class FShaderUniformBufferParameter
 {
 	DECLARE_EXPORTED_TYPE_LAYOUT(FShaderUniformBufferParameter, RENDERCORE_API, NonVirtual);
