Replace use of FRHICommandListExecutor::GetImmediateCommandList() with FRHICommandListImmediate::Get()
 - Only in places where it is trivially proven the call is only made on the render thread, due to an existing check(IsInRenderingThread()) assert somewhere in the function.
 - FRHICommandListImmediate::Get() itself contains a check(IsInRenderingThread()), so this enforces correct threading, and removes the need for extra checks at the call sites.
 - Remaining uses of FRHICommandListExecutor::GetImmediateCommandList() need investigation. Some may be bugs.
 - Also some changes to make use of the passed-in RHICmdList where possible (e.g. render commands that are given the immediate command list, but call the global getter rather than using the argument they were given).

diff --git a/Engine/Plugins/Enterprise/LidarPointCloud/Source/LidarPointCloudRuntime/Private/Rendering/LidarPointCloudRenderBuffers.cpp b/Engine/Plugins/Enterprise/LidarPointCloud/Source/LidarPointCloudRuntime/Private/Rendering/LidarPointCloudRenderBuffers.cpp
--- a/Engine/Plugins/Enterprise/LidarPointCloud/Source/LidarPointCloudRuntime/Private/Rendering/LidarPointCloudRenderBuffers.cpp
+++ b/Engine/Plugins/Enterprise/LidarPointCloud/Source/LidarPointCloudRuntime/Private/Rendering/LidarPointCloudRenderBuffers.cpp
@@ -32,13 +32,13 @@ TGlobalResource<FLidarPointCloudRenderBuffer> GDummyLidarPointCloudRenderBuffer(
 void FLidarPointCloudIndexBuffer::Resize(const uint32 & RequestedCapacity)
 {
 	// This must be called from Rendering thread
-	check(IsInRenderingThread());
+	FRHICommandListImmediate& RHICmdList = FRHICommandListImmediate::Get();
 
 	if (Capacity != RequestedCapacity)
 	{
 		ReleaseResource();
 		Capacity = RequestedCapacity;
-		InitResource(FRHICommandListExecutor::GetImmediateCommandList());
+		InitResource(RHICmdList);
 	}
 }
 
diff --git a/Engine/Plugins/Experimental/TextureGraph/Source/TextureGraphEngine/2D/TextureHelper.cpp b/Engine/Plugins/Experimental/TextureGraph/Source/TextureGraphEngine/2D/TextureHelper.cpp
--- a/Engine/Plugins/Experimental/TextureGraph/Source/TextureGraphEngine/2D/TextureHelper.cpp
+++ b/Engine/Plugins/Experimental/TextureGraph/Source/TextureGraphEngine/2D/TextureHelper.cpp
@@ -821,7 +821,7 @@ size_t TextureHelper::RoundUpTo(size_t Size, size_t DesiredRounding)
 
 RawBufferPtr TextureHelper::RawFromResource(const FTexture2DRHIRef& ResourceRHI, const BufferDescriptor& Desc)
 {
-	check(IsInRenderingThread()); 
+	FRHICommandListImmediate& RHICmdList = FRHICommandListImmediate::Get();
 
 	try
 	{
@@ -838,12 +838,11 @@ RawBufferPtr TextureHelper::RawFromResource(const FTexture2DRHIRef& ResourceRHI,
 
 		const TUniquePtr<FRHIGPUTextureReadback> TextureReadback = MakeUnique<FRHIGPUTextureReadback>(TEXT("RawFromResourceTextureReadback"));
 
-		FRHICommandListImmediate& RHI = FRHICommandListExecutor::GetImmediateCommandList();
-		RHI.FlushResources();
-		RHI.ImmediateFlush(EImmediateFlushType::WaitForOutstandingTasksOnly);
+		RHICmdList.FlushResources();
+		RHICmdList.ImmediateFlush(EImmediateFlushType::WaitForOutstandingTasksOnly);
 
-		TextureReadback->EnqueueCopy(RHI, ResourceRHI);
-		RHI.BlockUntilGPUIdle();
+		TextureReadback->EnqueueCopy(RHICmdList, ResourceRHI);
+		RHICmdList.BlockUntilGPUIdle();
 
 		//check(TextureReadback->IsReady());
 		{
diff --git a/Engine/Plugins/Media/AvfMedia/Source/AvfMedia/Private/Player/AvfMediaVideoSampler.cpp b/Engine/Plugins/Media/AvfMedia/Source/AvfMedia/Private/Player/AvfMediaVideoSampler.cpp
--- a/Engine/Plugins/Media/AvfMedia/Source/AvfMedia/Private/Player/AvfMediaVideoSampler.cpp
+++ b/Engine/Plugins/Media/AvfMedia/Source/AvfMedia/Private/Player/AvfMediaVideoSampler.cpp
@@ -173,7 +173,7 @@ void FAvfMediaVideoSampler::Tick()
 
 void FAvfMediaVideoSampler::ProcessFrame(CVPixelBufferRef Frame, FTimespan SampleTime, FTimespan SampleDuration)
 {
-	check(IsInRenderingThread());
+	FRHICommandListImmediate& RHICmdList = FRHICommandListImmediate::Get();
 
 	const int32 FrameHeight = CVPixelBufferGetHeight(Frame);
 	const int32 FrameWidth = CVPixelBufferGetWidth(Frame);
@@ -234,17 +234,16 @@ void FAvfMediaVideoSampler::ProcessFrame(CVPixelBufferRef Frame, FTimespan Sampl
 				.SetFlags(ETextureCreateFlags::Dynamic | ETextureCreateFlags::NoTiling | ETextureCreateFlags::ShaderResource)
 				.SetBulkData(new FAvfTexture2DResourceWrapper(UVTextureRef));
 			
-			TRefCountPtr<FRHITexture> YTex = RHICreateTexture(YDesc);
-			TRefCountPtr<FRHITexture> UVTex = RHICreateTexture(UVDesc);
+			TRefCountPtr<FRHITexture> YTex = RHICmdList.CreateTexture(YDesc);
+			TRefCountPtr<FRHITexture> UVTex = RHICmdList.CreateTexture(UVDesc);
 
 			const FRHITextureCreateDesc Desc =
 				FRHITextureCreateDesc::Create2D(TEXT("Info"), YWidth, YHeight, PF_B8G8R8A8)
 				.SetFlags(ETextureCreateFlags::Dynamic | ETextureCreateFlags::NoTiling | ETextureCreateFlags::ShaderResource | ETextureCreateFlags::RenderTargetable | ETextureCreateFlags::SRGB);
 			
-			ShaderResource = RHICreateTexture(Desc);
+			ShaderResource = RHICmdList.CreateTexture(Desc);
 			
 			// render video frame into sink texture
-			FRHICommandListImmediate& RHICmdList = FRHICommandListExecutor::GetImmediateCommandList();
 			{
 				// configure media shaders
 				auto ShaderMap = GetGlobalShaderMap(GMaxRHIFeatureLevel);
@@ -298,7 +297,7 @@ void FAvfMediaVideoSampler::ProcessFrame(CVPixelBufferRef Frame, FTimespan Sampl
 				.SetFlags(ETextureCreateFlags::SRGB | ETextureCreateFlags::Dynamic | ETextureCreateFlags::NoTiling | ETextureCreateFlags::ShaderResource)
 				.SetBulkData(new FAvfTexture2DResourceWrapper(TextureRef));
 
-			ShaderResource = RHICreateTexture(Desc);
+			ShaderResource = RHICmdList.CreateTexture(Desc);
 			
 			CFRelease(TextureRef);
 		}
diff --git a/Engine/Plugins/Media/ElectraUtil/Source/ElectraSamples/Private/ElectraTextureSampleAndroid.cpp b/Engine/Plugins/Media/ElectraUtil/Source/ElectraSamples/Private/ElectraTextureSampleAndroid.cpp
--- a/Engine/Plugins/Media/ElectraUtil/Source/ElectraSamples/Private/ElectraTextureSampleAndroid.cpp
+++ b/Engine/Plugins/Media/ElectraUtil/Source/ElectraSamples/Private/ElectraTextureSampleAndroid.cpp
@@ -485,7 +485,7 @@ uint32 FElectraTextureSample::GetStride() const
 
 bool FElectraTextureSample::Convert(FTexture2DRHIRef& InDstTexture, const FConversionHints& Hints)
 {
-	check(IsInRenderingThread());
+	FRHICommandListImmediate& RHICmdList = FRHICommandListImmediate::Get();
 
 	if (GDynamicRHI->RHIIsRenderingSuspended())
 	{
@@ -519,8 +519,6 @@ bool FElectraTextureSample::Convert(FTexture2DRHIRef& InDstTexture, const FConve
 		InputTexture = Texture;
 	}
 
-	FRHICommandListImmediate& RHICmdList = FRHICommandListExecutor::GetImmediateCommandList();
-
 	SCOPED_DRAW_EVENT(RHICmdList, AndroidMediaOutputConvertTexture);
 	SCOPED_GPU_STAT(RHICmdList, MediaAndroidDecoder_Convert);
 
diff --git a/Engine/Plugins/Media/ElectraUtil/Source/ElectraSamples/Private/ElectraTextureSamplePC.cpp b/Engine/Plugins/Media/ElectraUtil/Source/ElectraSamples/Private/ElectraTextureSamplePC.cpp
--- a/Engine/Plugins/Media/ElectraUtil/Source/ElectraSamples/Private/ElectraTextureSamplePC.cpp
+++ b/Engine/Plugins/Media/ElectraUtil/Source/ElectraSamples/Private/ElectraTextureSamplePC.cpp
@@ -338,9 +338,7 @@ bool FElectraTextureSample::Convert(FTexture2DRHIRef& InDstTexture, const FConve
 {
 	LLM_SCOPE(ELLMTag::MediaStreaming);
 
-	check(IsInRenderingThread());
-
-	FRHICommandListImmediate& RHICmdList = FRHICommandListExecutor::GetImmediateCommandList();
+	FRHICommandListImmediate& RHICmdList = FRHICommandListImmediate::Get();
 
 	SCOPED_DRAW_EVENT(RHICmdList, WinMediaOutputConvertTexture);
 	SCOPED_GPU_STAT(RHICmdList, MediaWinDecoder_Convert);
diff --git a/Engine/Plugins/Media/MediaIOFramework/Source/MediaIOCore/Private/Player/MediaIOCorePlayerBase.cpp b/Engine/Plugins/Media/MediaIOFramework/Source/MediaIOCore/Private/Player/MediaIOCorePlayerBase.cpp
--- a/Engine/Plugins/Media/MediaIOFramework/Source/MediaIOCore/Private/Player/MediaIOCorePlayerBase.cpp
+++ b/Engine/Plugins/Media/MediaIOFramework/Source/MediaIOCore/Private/Player/MediaIOCorePlayerBase.cpp
@@ -1312,9 +1312,7 @@ TSharedPtr<FMediaIOCoreTextureSampleBase> FMediaIOCorePlayerBase::PickSampleToRe
 
 void FMediaIOCorePlayerBase::TransferTexture_RenderThread(const TSharedPtr<FMediaIOCoreTextureSampleBase>& Sample, const TSharedPtr<FMediaIOCoreTextureSampleBase>& JITRProxySample)
 {
-	checkSlow(IsInRenderingThread());
-
-	FRHICommandListImmediate& RHICmdList = FRHICommandListExecutor::GetImmediateCommandList();
+	FRHICommandListImmediate& RHICmdList = FRHICommandListImmediate::Get();
 
 	// DMA P2P transfer
 	const bool bIsSampleAwaitingForGpuTransfer = Sample->IsAwaitingForGPUTransfer();
diff --git a/Engine/Plugins/Media/WmfMedia/Source/WmfMedia/Private/Wmf/WmfMediaHardwareVideoDecodingRendering.cpp b/Engine/Plugins/Media/WmfMedia/Source/WmfMedia/Private/Wmf/WmfMediaHardwareVideoDecodingRendering.cpp
--- a/Engine/Plugins/Media/WmfMedia/Source/WmfMedia/Private/Wmf/WmfMediaHardwareVideoDecodingRendering.cpp
+++ b/Engine/Plugins/Media/WmfMedia/Source/WmfMedia/Private/Wmf/WmfMediaHardwareVideoDecodingRendering.cpp
@@ -113,14 +113,12 @@ bool FWmfMediaHardwareVideoDecodingParameters::ConvertTextureFormat_RenderThread
 		return false;
 	}
 
-	check(IsInRenderingThread());
+	FRHICommandListImmediate& RHICmdList = FRHICommandListImmediate::Get();
 	check(InSample);
-
+	
 	TComPtr<ID3D11Texture2D> SampleTexture = InSample->GetSourceTexture();
 
 	{
-		FRHICommandListImmediate& RHICmdList = FRHICommandListExecutor::GetImmediateCommandList();
-
 		SCOPED_DRAW_EVENT(RHICmdList, FWmfMediaHardwareVideoDecodingParameters_Convert);
 		SCOPED_GPU_STAT(RHICmdList, MediaTextureConversion);
 
diff --git a/Engine/Plugins/Runtime/nDisplay/Source/DisplayClusterShaders/Private/Shaders/DisplayClusterShadersGenerateMips.cpp b/Engine/Plugins/Runtime/nDisplay/Source/DisplayClusterShaders/Private/Shaders/DisplayClusterShadersGenerateMips.cpp
--- a/Engine/Plugins/Runtime/nDisplay/Source/DisplayClusterShaders/Private/Shaders/DisplayClusterShadersGenerateMips.cpp
+++ b/Engine/Plugins/Runtime/nDisplay/Source/DisplayClusterShaders/Private/Shaders/DisplayClusterShadersGenerateMips.cpp
@@ -9,8 +9,6 @@
 
 bool FDisplayClusterShadersGenerateMips::GenerateMips(FRHICommandListImmediate& RHICmdList, FRHITexture2D* InOutMipsTexture, const FDisplayClusterShaderParameters_GenerateMips& InSettings)
 {
-	check(IsInRenderingThread());
-
 	if (InSettings.IsEnabled())
 	{
 		check(InOutMipsTexture);
@@ -19,8 +17,6 @@ bool FDisplayClusterShadersGenerateMips::GenerateMips(FRHICommandListImmediate&
 		FGenerateMipsParams GenerateMipsParams{ InSettings.MipsSamplerFilter == TF_Nearest ? SF_Point : (InSettings.MipsSamplerFilter == TF_Trilinear ? SF_Trilinear : SF_Bilinear),
 			InSettings.MipsAddressU == TA_Wrap ? AM_Wrap : (InSettings.MipsAddressU == TA_Mirror ? AM_Mirror : AM_Clamp),
 			InSettings.MipsAddressV == TA_Wrap ? AM_Wrap : (InSettings.MipsAddressV == TA_Mirror ? AM_Mirror : AM_Clamp) };
-
-		check(&RHICmdList == &FRHICommandListExecutor::GetImmediateCommandList());
 		
 		FRDGBuilder GraphBuilder(RHICmdList);
 		TRefCountPtr<IPooledRenderTarget> PoolRenderTarget = CreateRenderTarget(InOutMipsTexture, TEXT("nDisplayViewportMips"));
diff --git a/Engine/Source/Runtime/Apple/MetalRHI/Private/MetalTexture.cpp b/Engine/Source/Runtime/Apple/MetalRHI/Private/MetalTexture.cpp
--- a/Engine/Source/Runtime/Apple/MetalRHI/Private/MetalTexture.cpp
+++ b/Engine/Source/Runtime/Apple/MetalRHI/Private/MetalTexture.cpp
@@ -1513,8 +1513,7 @@ FTexture2DRHIRef FMetalDynamicRHI::RHIAsyncReallocateTexture2D(FRHITexture2D* Ol
 {
     MTL_SCOPED_AUTORELEASE_POOL;
 
-    check(IsInRenderingThread());
-    FRHICommandListImmediate& RHICmdList = FRHICommandListExecutor::GetImmediateCommandList();
+	FRHICommandListImmediate& RHICmdList = FRHICommandListImmediate::Get();
 
     FMetalSurface* OldTexture = ResourceCast(OldTextureRHI);
 
diff --git a/Engine/Source/Runtime/D3D12RHI/Private/D3D12DirectCommandListManager.cpp b/Engine/Source/Runtime/D3D12RHI/Private/D3D12DirectCommandListManager.cpp
--- a/Engine/Source/Runtime/D3D12RHI/Private/D3D12DirectCommandListManager.cpp
+++ b/Engine/Source/Runtime/D3D12RHI/Private/D3D12DirectCommandListManager.cpp
@@ -168,10 +168,10 @@ uint64 FD3D12ManualFence::GetCompletedFenceValue(bool bUpdateCachedFenceValue)
 
 void FD3D12ManualFence::AdvanceFrame()
 {
-	check(IsInRenderingThread());
+	FRHICommandListImmediate& RHICmdList = FRHICommandListImmediate::Get();
 
 	const uint64 NextValue = NextFenceValue.Increment();
-	FRHICommandListExecutor::GetImmediateCommandList().EnqueueLambda([this, NextValue](FRHICommandListImmediate&)
+	RHICmdList.EnqueueLambda([this, NextValue](FRHICommandListImmediate&)
 	{
 		for (FFencePair& Pair : FencePairs)
 		{
diff --git a/Engine/Source/Runtime/D3D12RHI/Private/D3D12RHI.cpp b/Engine/Source/Runtime/D3D12RHI/Private/D3D12RHI.cpp
--- a/Engine/Source/Runtime/D3D12RHI/Private/D3D12RHI.cpp
+++ b/Engine/Source/Runtime/D3D12RHI/Private/D3D12RHI.cpp
@@ -334,7 +334,7 @@ void FD3D12DynamicRHI::Shutdown()
 	}
 
 	// Flush all pending deletes before destroying the device or any command contexts.
-	FRHICommandListExecutor::GetImmediateCommandList().ImmediateFlush(EImmediateFlushType::FlushRHIThreadFlushResources);
+	FRHICommandListImmediate::Get().ImmediateFlush(EImmediateFlushType::FlushRHIThreadFlushResources);
 
 	RHIShutdownFlipTracking();
 	ShutdownSubmissionPipe();
diff --git a/Engine/Source/Runtime/Engine/Private/Streaming/StreamingManagerTexture.cpp b/Engine/Source/Runtime/Engine/Private/Streaming/StreamingManagerTexture.cpp
--- a/Engine/Source/Runtime/Engine/Private/Streaming/StreamingManagerTexture.cpp
+++ b/Engine/Source/Runtime/Engine/Private/Streaming/StreamingManagerTexture.cpp
@@ -388,10 +388,10 @@ bool FRenderAssetStreamingManager::StreamOutRenderAssetData( int64 RequiredMemor
 			{
 				// Queue up the process on the render thread and wait for everything to complete.
 				ENQUEUE_RENDER_COMMAND(FlushResourceCommand)(
-					[](FRHICommandList& RHICmdList)
+					[](FRHICommandListImmediate& RHICmdList)
 					{				
-						FRHICommandListExecutor::GetImmediateCommandList().ImmediateFlush(EImmediateFlushType::FlushRHIThreadFlushResources);
-						RHIFlushResources();
+						RHICmdList.ImmediateFlush(EImmediateFlushType::FlushRHIThreadFlushResources);
+						RHICmdList.FlushResources();
 					});
 				FlushRenderingCommands();
 				TempMemoryUsed = 0;
diff --git a/Engine/Source/Runtime/Engine/Private/VT/VirtualTextureUploadCache.cpp b/Engine/Source/Runtime/Engine/Private/VT/VirtualTextureUploadCache.cpp
--- a/Engine/Source/Runtime/Engine/Private/VT/VirtualTextureUploadCache.cpp
+++ b/Engine/Source/Runtime/Engine/Private/VT/VirtualTextureUploadCache.cpp
@@ -390,8 +390,7 @@ void FVirtualTextureUploadCache::Finalize(FRDGBuilder& GraphBuilder)
 
 void FVirtualTextureUploadCache::ReleaseRHI()
 {
-	check(IsInRenderingThread());
-	FRHICommandListImmediate& RHICmdList = FRHICommandListExecutor::GetImmediateCommandList();
+	FRHICommandListImmediate& RHICmdList = FRHICommandListImmediate::Get();
 
 	// Complete/Cancel all work will release allocated staging buffers.
 	UpdateFreeList(RHICmdList, true);
diff --git a/Engine/Source/Runtime/Landscape/Private/LandscapeGrassMapsBuilder.cpp b/Engine/Source/Runtime/Landscape/Private/LandscapeGrassMapsBuilder.cpp
--- a/Engine/Source/Runtime/Landscape/Private/LandscapeGrassMapsBuilder.cpp
+++ b/Engine/Source/Runtime/Landscape/Private/LandscapeGrassMapsBuilder.cpp
@@ -257,17 +257,18 @@ namespace UE::Landscape
 		}
 
 		ENQUEUE_RENDER_COMMAND(FFlushResourcesCommand)(
-			[ResultsReadyEvent, bBlockRTUntilGPUComplete](FRHICommandList& RHICmdList)
+			[ResultsReadyEvent, bBlockRTUntilGPUComplete](FRHICommandListImmediate& RHICmdList)
 			{
-				FRHICommandListExecutor::GetImmediateCommandList().ImmediateFlush(EImmediateFlushType::FlushRHIThreadFlushResources);
-				RHIFlushResources();
-				FRHICommandListExecutor::GetImmediateCommandList().SubmitCommandsAndFlushGPU();
+				RHICmdList.ImmediateFlush(EImmediateFlushType::FlushRHIThreadFlushResources);
+				RHICmdList.FlushResources();
+				RHICmdList.SubmitCommandsAndFlushGPU();
+
 				if (ResultsReadyEvent)
 				{
 					if (bBlockRTUntilGPUComplete)
 					{
 						// Block render thread waiting for GPU to complete.  Note this can be very expensive on some platforms.
-						FRHICommandListExecutor::GetImmediateCommandList().BlockUntilGPUIdle();
+						RHICmdList.BlockUntilGPUIdle();
 					}
 					
 					ResultsReadyEvent->Trigger();
diff --git a/Engine/Source/Runtime/Launch/Private/LaunchEngineLoop.cpp b/Engine/Source/Runtime/Launch/Private/LaunchEngineLoop.cpp
--- a/Engine/Source/Runtime/Launch/Private/LaunchEngineLoop.cpp
+++ b/Engine/Source/Runtime/Launch/Private/LaunchEngineLoop.cpp
@@ -5939,10 +5939,10 @@ void FEngineLoop::Tick()
 		}
 		{
 			ENQUEUE_RENDER_COMMAND(WaitForOutstandingTasksOnly_for_DelaySceneRenderCompletion)(
-				[](FRHICommandList& RHICmdList)
+				[](FRHICommandListImmediate& RHICmdList)
 				{
 					QUICK_SCOPE_CYCLE_COUNTER(STAT_DelaySceneRenderCompletion_TaskWait);
-					FRHICommandListExecutor::GetImmediateCommandList().ImmediateFlush(EImmediateFlushType::WaitForOutstandingTasksOnly);
+					RHICmdList.ImmediateFlush(EImmediateFlushType::WaitForOutstandingTasksOnly);
 				});
 		}
 #endif
diff --git a/Engine/Source/Runtime/MediaAssets/Private/Misc/MediaTextureResource.cpp b/Engine/Source/Runtime/MediaAssets/Private/Misc/MediaTextureResource.cpp
--- a/Engine/Source/Runtime/MediaAssets/Private/Misc/MediaTextureResource.cpp
+++ b/Engine/Source/Runtime/MediaAssets/Private/Misc/MediaTextureResource.cpp
@@ -367,8 +367,8 @@ void FMediaTextureResource::FlushPendingData()
 
 void FMediaTextureResource::Render(const FRenderParams& Params)
 {
-	check(IsInRenderingThread());
-	SCOPED_GPU_STAT(FRHICommandListExecutor::GetImmediateCommandList(), MediaTextureResource);
+	FRHICommandListImmediate& RHICmdList = FRHICommandListImmediate::Get();
+	SCOPED_GPU_STAT(RHICmdList, MediaTextureResource);
 
 	LLM_SCOPE(ELLMTag::MediaStreaming);
 	SCOPE_CYCLE_COUNTER(STAT_MediaAssets_MediaTextureResourceRender);
@@ -376,8 +376,8 @@ void FMediaTextureResource::Render(const FRenderParams& Params)
 
 	TSharedPtr<FPriorSamples, ESPMode::ThreadSafe> LocalPriorSamples;
 	{
-	FScopeLock Lock(&PriorSamplesCS);
-	LocalPriorSamples = PriorSamples;
+		FScopeLock Lock(&PriorSamplesCS);
+		LocalPriorSamples = PriorSamples;
 	}
 
 	LocalPriorSamples->Update();
@@ -622,7 +622,7 @@ void FMediaTextureResource::Render(const FRenderParams& Params)
 
 					CacheRenderTarget(OutputTarget, TEXT("MipGeneration"), MipGenerationCache);
 
-					FRDGBuilder GraphBuilder(FRHICommandListExecutor::GetImmediateCommandList());
+					FRDGBuilder GraphBuilder(RHICmdList);
 					FRDGTextureRef MipOutputTexture = GraphBuilder.RegisterExternalTexture(MipGenerationCache);
 					FGenerateMips::Execute(GraphBuilder, GetFeatureLevel(), MipOutputTexture, FGenerateMipsParams{ SF_Bilinear, AM_Clamp, AM_Clamp, AM_Clamp }, GenerateMipsPass);
 					GraphBuilder.Execute();
diff --git a/Engine/Source/Runtime/OpenGLDrv/Private/OpenGLQuery.cpp b/Engine/Source/Runtime/OpenGLDrv/Private/OpenGLQuery.cpp
--- a/Engine/Source/Runtime/OpenGLDrv/Private/OpenGLQuery.cpp
+++ b/Engine/Source/Runtime/OpenGLDrv/Private/OpenGLQuery.cpp
@@ -637,8 +637,7 @@ FOpenGLRenderQuery::FOpenGLRenderQuery(ERenderQueryType InQueryType)
 	, bInvalidResource(true)
 	, QueryType(InQueryType)
 {
-	check(IsInRenderingThread());
-	FRHICommandListExecutor::GetImmediateCommandList().EnqueueLambda(
+	FRHICommandListImmediate::Get().EnqueueLambda(
 		[this](FRHICommandListImmediate&) { AcquireResource(); }
 	);
 }
diff --git a/Engine/Source/Runtime/OpenGLDrv/Private/OpenGLState.cpp b/Engine/Source/Runtime/OpenGLDrv/Private/OpenGLState.cpp
--- a/Engine/Source/Runtime/OpenGLDrv/Private/OpenGLState.cpp
+++ b/Engine/Source/Runtime/OpenGLDrv/Private/OpenGLState.cpp
@@ -299,8 +299,7 @@ FSamplerStateRHIRef FOpenGLDynamicRHI::RHICreateSamplerState(const FSamplerState
 
 	SamplerState->Resource = 0;
 
-	check(IsInRenderingThread());
-	FRHICommandListExecutor::GetImmediateCommandList().EnqueueLambda([SamplerState](FRHICommandListImmediate&)
+	FRHICommandListImmediate::Get().EnqueueLambda([SamplerState](FRHICommandListImmediate&)
 	{
 		VERIFY_GL_SCOPE();
 		FOpenGL::GenSamplers( 1, &SamplerState->Resource);
diff --git a/Engine/Source/Runtime/OpenGLDrv/Private/OpenGLTexture.cpp b/Engine/Source/Runtime/OpenGLDrv/Private/OpenGLTexture.cpp
--- a/Engine/Source/Runtime/OpenGLDrv/Private/OpenGLTexture.cpp
+++ b/Engine/Source/Runtime/OpenGLDrv/Private/OpenGLTexture.cpp
@@ -1492,8 +1492,7 @@ FTexture2DRHIRef FOpenGLDynamicRHI::AsyncReallocateTexture2D_RenderThread(class
 
 FTexture2DRHIRef FOpenGLDynamicRHI::RHIAsyncReallocateTexture2D(FRHITexture2D* Texture2DRHI, int32 NewMipCount, int32 NewSizeX, int32 NewSizeY, FThreadSafeCounter* RequestStatus)
 {
-	check(IsInRenderingThread());
-	FRHICommandListImmediate& RHICmdList = FRHICommandListExecutor::GetImmediateCommandList();
+	FRHICommandListImmediate& RHICmdList = FRHICommandListImmediate::Get();
 
 	FOpenGLTexture* OldTexture = ResourceCast(Texture2DRHI);
 
diff --git a/Engine/Source/Runtime/OpenGLDrv/Private/OpenGLUniformBuffer.cpp b/Engine/Source/Runtime/OpenGLDrv/Private/OpenGLUniformBuffer.cpp
--- a/Engine/Source/Runtime/OpenGLDrv/Private/OpenGLUniformBuffer.cpp
+++ b/Engine/Source/Runtime/OpenGLDrv/Private/OpenGLUniformBuffer.cpp
@@ -558,8 +558,7 @@ static FUniformBufferRHIRef CreateUniformBuffer(const void* Contents, const FRHI
 
 	bool bUseEmulatedUBs = GUseEmulatedUniformBuffers && !Layout->bNoEmulatedUniformBuffer;
 
-	check(IsInRenderingThread());
-	FRHICommandListImmediate& RHICmdList = FRHICommandListExecutor::GetImmediateCommandList();
+	FRHICommandListImmediate& RHICmdList = FRHICommandListImmediate::Get();
 
 	{
 		const bool bStreamDraw = (Usage == UniformBuffer_SingleDraw || Usage == UniformBuffer_SingleFrame);
@@ -673,8 +672,7 @@ static FOpenGLUniformBuffer* CreateUniformBufferView(FRHICommandListImmediate& R
 
 FUniformBufferRHIRef FOpenGLDynamicRHI::RHICreateUniformBuffer(const void* Contents, const FRHIUniformBufferLayout* Layout, EUniformBufferUsage Usage, EUniformBufferValidation Validation)
 {
-	check(IsInRenderingThread());
-	FRHICommandListImmediate& RHICmdList = FRHICommandListExecutor::GetImmediateCommandList();
+	FRHICommandListImmediate& RHICmdList = FRHICommandListImmediate::Get();
 
 	// This should really be synchronized, if there's a chance it'll be used from more than one buffer. Luckily, uniform buffers
 	// are only used for drawing/shader usage, not for loading resources or framebuffer blitting, so no synchronization primitives for now.
diff --git a/Engine/Source/Runtime/RHI/Private/PipelineStateCache.cpp b/Engine/Source/Runtime/RHI/Private/PipelineStateCache.cpp
--- a/Engine/Source/Runtime/RHI/Private/PipelineStateCache.cpp
+++ b/Engine/Source/Runtime/RHI/Private/PipelineStateCache.cpp
@@ -1039,10 +1039,7 @@ public:
 
 	void ProcessDelayedCleanup()
 	{
-		check(IsInRenderingThread());
-
-		FRHICommandListImmediate& RHICmdList = FRHICommandListExecutor::GetImmediateCommandList();
-		RHICmdList.EnqueueLambda([DeleteArray = MoveTemp(DeleteArray)](FRHICommandListImmediate& RHICmdList) mutable
+		FRHICommandListImmediate::Get().EnqueueLambda([DeleteArray = MoveTemp(DeleteArray)](FRHICommandListImmediate& RHICmdList) mutable
 		{
 			for (TMyValue& OldPipelineState : DeleteArray)
 			{
@@ -1061,14 +1058,11 @@ public:
 
 	int32 DiscardAndSwap()
 	{
-		check(IsInRenderingThread());
-
 		// the consolidate should always be run before the DiscardAndSwap.
 		// there should be no inuse pipeline states in the backfill map (because they should have been moved into the CurrentMap).
 		int32 Discarded = BackfillMap->Num();
 
-		FRHICommandListImmediate& RHICmdList = FRHICommandListExecutor::GetImmediateCommandList();
-		RHICmdList.EnqueueLambda([BackfillMap = MoveTemp(*BackfillMap)](FRHICommandListImmediate& RHICmdList) mutable
+		FRHICommandListImmediate::Get().EnqueueLambda([BackfillMap = MoveTemp(*BackfillMap)](FRHICommandListImmediate& RHICmdList) mutable
 		{
 			for (const auto& DiscardIterator :  BackfillMap)
 			{
diff --git a/Engine/Source/Runtime/RHI/Private/RHICommandList.cpp b/Engine/Source/Runtime/RHI/Private/RHICommandList.cpp
--- a/Engine/Source/Runtime/RHI/Private/RHICommandList.cpp
+++ b/Engine/Source/Runtime/RHI/Private/RHICommandList.cpp
@@ -1559,7 +1559,7 @@ void FRHICommandListImmediate::BeginScene()
 		// if we aren't running an RHIThread, there is no good reason to buffer this frame advance stuff and that complicates state management, so flush everything out now
 		QUICK_SCOPE_CYCLE_COUNTER(BeginScene_Flush);
 		CSV_SCOPED_TIMING_STAT(RHITFlushes, BeginScene);
-		FRHICommandListExecutor::GetImmediateCommandList().ImmediateFlush(EImmediateFlushType::FlushRHIThread);
+		ImmediateFlush(EImmediateFlushType::FlushRHIThread);
 	}
 }
 
@@ -1577,7 +1577,7 @@ void FRHICommandListImmediate::EndScene()
 		// if we aren't running an RHIThread, there is no good reason to buffer this frame advance stuff and that complicates state management, so flush everything out now
 		QUICK_SCOPE_CYCLE_COUNTER(EndScene_Flush);
 		CSV_SCOPED_TIMING_STAT(RHITFlushes, EndScene);
-		FRHICommandListExecutor::GetImmediateCommandList().ImmediateFlush(EImmediateFlushType::FlushRHIThread);
+		ImmediateFlush(EImmediateFlushType::FlushRHIThread);
 	}
 }
 
@@ -1595,7 +1595,7 @@ void FRHICommandListImmediate::BeginDrawingViewport(FRHIViewport* Viewport, FRHI
 		// if we aren't running an RHIThread, there is no good reason to buffer this frame advance stuff and that complicates state management, so flush everything out now
 		QUICK_SCOPE_CYCLE_COUNTER(BeginDrawingViewport_Flush);
 		CSV_SCOPED_TIMING_STAT(RHITFlushes, BeginDrawingViewport);
-		FRHICommandListExecutor::GetImmediateCommandList().ImmediateFlush(EImmediateFlushType::FlushRHIThread);
+		ImmediateFlush(EImmediateFlushType::FlushRHIThread);
 	}
 }
 
@@ -1604,7 +1604,7 @@ void FRHICommandListImmediate::EndDrawingViewport(FRHIViewport* Viewport, bool b
 	// Make sure all prior graphics and async compute work has been submitted.
 	// This is necessary because platform RHIs often submit additional work on the graphics queue during present, and we need to ensure we won't deadlock on async work that wasn't yet submitted by the renderer.
 	// In future, Present() itself should be an enqueued / recorded command, and platform RHIs should never implicitly submit graphics or async compute work.
-	FRHICommandListExecutor::GetImmediateCommandList().ImmediateFlush(EImmediateFlushType::DispatchToRHIThread);
+	ImmediateFlush(EImmediateFlushType::DispatchToRHIThread);
 
 	check(IsImmediate() && IsInRenderingThread());
 	if (Bypass())
@@ -1623,7 +1623,7 @@ void FRHICommandListImmediate::EndDrawingViewport(FRHIViewport* Viewport, bool b
 		// if we aren't running an RHIThread, there is no good reason to buffer this frame advance stuff and that complicates state management, so flush everything out now
 		{
 			QUICK_SCOPE_CYCLE_COUNTER(STAT_EndDrawingViewport_Dispatch);
-			FRHICommandListExecutor::GetImmediateCommandList().ImmediateFlush(EImmediateFlushType::DispatchToRHIThread);
+			ImmediateFlush(EImmediateFlushType::DispatchToRHIThread);
 		}
 	}
 
@@ -1649,7 +1649,7 @@ void FRHICommandListImmediate::BeginFrame()
 		CSV_SCOPED_TIMING_STAT(RHITFlushes, BeginFrame);
 
 		// Use the ERHISubmitFlags::ProcessStats flag to delineate stat data between the previous and next frames.
-		FRHICommandListExecutor::GetImmediateCommandList().ImmediateFlush(EImmediateFlushType::DispatchToRHIThread, ERHISubmitFlags::ProcessStats);
+		ImmediateFlush(EImmediateFlushType::DispatchToRHIThread, ERHISubmitFlags::ProcessStats);
 	}
 
 	GDynamicRHI->RHIBeginFrame(*this);
@@ -1681,11 +1681,11 @@ void FRHICommandListImmediate::EndFrame()
 		// if we aren't running an RHIThread, there is no good reason to buffer this frame advance stuff and that complicates state management, so flush everything out now
 		QUICK_SCOPE_CYCLE_COUNTER(EndFrame_Flush);
 		CSV_SCOPED_TIMING_STAT(RHITFlushes, EndFrame);
-		FRHICommandListExecutor::GetImmediateCommandList().ImmediateFlush(EImmediateFlushType::FlushRHIThread);
+		ImmediateFlush(EImmediateFlushType::FlushRHIThread);
 	}
 	else
 	{
-		FRHICommandListExecutor::GetImmediateCommandList().ImmediateFlush(EImmediateFlushType::DispatchToRHIThread);
+		ImmediateFlush(EImmediateFlushType::DispatchToRHIThread);
 	}
 }
 
diff --git a/Engine/Source/Runtime/RHI/Public/RHICommandList.h b/Engine/Source/Runtime/RHI/Public/RHICommandList.h
--- a/Engine/Source/Runtime/RHI/Public/RHICommandList.h
+++ b/Engine/Source/Runtime/RHI/Public/RHICommandList.h
@@ -4755,22 +4755,19 @@ FORCEINLINE_DEBUGGABLE FRHICommandListImmediate& FRHICommandListExecutor::GetImm
 UE_DEPRECATED(5.3, "RHICreateBuffer is deprecated. Use FRHICommandListBase::CreateBuffer instead.")
 FORCEINLINE FBufferRHIRef RHICreateBuffer(uint32 Size, EBufferUsageFlags Usage, uint32 Stride, ERHIAccess ResourceState, FRHIResourceCreateInfo& CreateInfo)
 {
-	check(IsInRenderingThread());
-	return FRHICommandListExecutor::GetImmediateCommandList().CreateBuffer(Size, Usage, Stride, ResourceState, CreateInfo);
+	return FRHICommandListImmediate::Get().CreateBuffer(Size, Usage, Stride, ResourceState, CreateInfo);
 }
 
 UE_DEPRECATED(5.3, "RHICreateIndexBuffer is deprecated. Use FRHICommandListBase::CreateIndexBuffer instead.")
 FORCEINLINE FBufferRHIRef RHICreateIndexBuffer(uint32 Stride, uint32 Size, EBufferUsageFlags Usage, ERHIAccess ResourceState, FRHIResourceCreateInfo& CreateInfo)
 {
-	check(IsInRenderingThread());
-	return FRHICommandListExecutor::GetImmediateCommandList().CreateIndexBuffer(Stride, Size, Usage, ResourceState, CreateInfo);
+	return FRHICommandListImmediate::Get().CreateIndexBuffer(Stride, Size, Usage, ResourceState, CreateInfo);
 }
 
 UE_DEPRECATED(5.3, "RHICreateIndexBuffer is deprecated. Use FRHICommandListBase::CreateIndexBuffer instead.")
 FORCEINLINE FBufferRHIRef RHICreateIndexBuffer(uint32 Stride, uint32 Size, EBufferUsageFlags Usage, FRHIResourceCreateInfo& CreateInfo)
 {
-	check(IsInRenderingThread());
-	return FRHICommandListExecutor::GetImmediateCommandList().CreateIndexBuffer(Stride, Size, Usage, CreateInfo);
+	return FRHICommandListImmediate::Get().CreateIndexBuffer(Stride, Size, Usage, CreateInfo);
 }
 
 UE_DEPRECATED(5.3, "RHIUpdateUniformBuffer is deprecated. Use FRHICommandListBase::UpdateUniformBuffer instead.")
@@ -4782,43 +4779,37 @@ FORCEINLINE void RHIUpdateUniformBuffer(FRHIUniformBuffer* UniformBufferRHI, con
 UE_DEPRECATED(5.3, "RHICreateVertexBuffer is deprecated. Use FRHICommandListBase::CreateVertexBuffer instead.")
 FORCEINLINE FBufferRHIRef RHICreateVertexBuffer(uint32 Size, EBufferUsageFlags Usage, ERHIAccess ResourceState, FRHIResourceCreateInfo& CreateInfo)
 {
-	check(IsInRenderingThread());
-	return FRHICommandListExecutor::GetImmediateCommandList().CreateVertexBuffer(Size, Usage, ResourceState, CreateInfo);
+	return FRHICommandListImmediate::Get().CreateVertexBuffer(Size, Usage, ResourceState, CreateInfo);
 }
 
 UE_DEPRECATED(5.3, "RHICreateVertexBuffer is deprecated. Use FRHICommandListBase::CreateVertexBuffer instead.")
 FORCEINLINE FBufferRHIRef RHICreateVertexBuffer(uint32 Size, EBufferUsageFlags Usage, FRHIResourceCreateInfo& CreateInfo)
 {
-	check(IsInRenderingThread());
-	return FRHICommandListExecutor::GetImmediateCommandList().CreateVertexBuffer(Size, Usage, CreateInfo);
+	return FRHICommandListImmediate::Get().CreateVertexBuffer(Size, Usage, CreateInfo);
 }
 
 UE_DEPRECATED(5.3, "RHICreateStructuredBuffer is deprecated. Use FRHICommandListBase::CreateStructuredBuffer instead.")
 FORCEINLINE FBufferRHIRef RHICreateStructuredBuffer(uint32 Stride, uint32 Size, EBufferUsageFlags Usage, ERHIAccess ResourceState, FRHIResourceCreateInfo& CreateInfo)
 {
-	check(IsInRenderingThread());
-	return FRHICommandListExecutor::GetImmediateCommandList().CreateStructuredBuffer(Stride, Size, Usage, ResourceState, CreateInfo);
+	return FRHICommandListImmediate::Get().CreateStructuredBuffer(Stride, Size, Usage, ResourceState, CreateInfo);
 }
 
 UE_DEPRECATED(5.3, "RHICreateStructuredBuffer is deprecated. Use FRHICommandListBase::CreateStructuredBuffer instead.")
 FORCEINLINE FBufferRHIRef RHICreateStructuredBuffer(uint32 Stride, uint32 Size, EBufferUsageFlags Usage, FRHIResourceCreateInfo& CreateInfo)
 {
-	check(IsInRenderingThread());
-	return FRHICommandListExecutor::GetImmediateCommandList().CreateStructuredBuffer(Stride, Size, Usage, CreateInfo);
+	return FRHICommandListImmediate::Get().CreateStructuredBuffer(Stride, Size, Usage, CreateInfo);
 }
 
 UE_DEPRECATED(5.3, "RHILockBuffer is deprecated. Use FRHICommandListBase::LockBuffer instead.")
 FORCEINLINE void* RHILockBuffer(FRHIBuffer* Buffer, uint32 Offset, uint32 SizeRHI, EResourceLockMode LockMode)
 {
-	check(IsInRenderingThread());
-	return FRHICommandListExecutor::GetImmediateCommandList().LockBuffer(Buffer, Offset, SizeRHI, LockMode);
+	return FRHICommandListImmediate::Get().LockBuffer(Buffer, Offset, SizeRHI, LockMode);
 }
 
 UE_DEPRECATED(5.3, "RHIUnlockBuffer is deprecated. Use FRHICommandListBase::UnlockBuffer instead.")
 FORCEINLINE void RHIUnlockBuffer(FRHIBuffer* Buffer)
 {
-	check(IsInRenderingThread());
-	FRHICommandListExecutor::GetImmediateCommandList().UnlockBuffer(Buffer);
+	FRHICommandListImmediate::Get().UnlockBuffer(Buffer);
 }
 
 UE_DEPRECATED(5.3, "RHICreateShaderResourceView is deprecated. Use FRHICommandListBase::CreateShaderResourceView instead.")
@@ -4912,8 +4903,7 @@ FORCEINLINE FTextureReferenceRHIRef RHICreateTextureReference(FRHITexture* InRef
 
 FORCEINLINE void RHIUpdateTextureReference(FRHITextureReference* TextureRef, FRHITexture* NewTexture)
 {
-	check(IsInRenderingThread());
-	FRHICommandListExecutor::GetImmediateCommandList().UpdateTextureReference(TextureRef, NewTexture);
+	FRHICommandListImmediate::Get().UpdateTextureReference(TextureRef, NewTexture);
 }
 
 FORCEINLINE FTextureRHIRef RHICreateTexture(const FRHITextureCreateDesc& CreateDesc)
diff --git a/Engine/Source/Runtime/RenderCore/Private/DumpGPU.cpp b/Engine/Source/Runtime/RenderCore/Private/DumpGPU.cpp
--- a/Engine/Source/Runtime/RenderCore/Private/DumpGPU.cpp
+++ b/Engine/Source/Runtime/RenderCore/Private/DumpGPU.cpp
@@ -1287,15 +1287,10 @@ public:
 	}
 
 	void DumpDrawTextureSubResource(
-		FRHICommandList& RHICmdList,
+		FRHICommandListImmediate& RHICmdList,
 		FRDGTextureSRVDesc SubresourceDesc,
 		ERHIAccess RHIAccessState)
 	{
-		check(IsInRenderingThread());
-
-		FRHICommandListImmediate& RHICmdListImmediate = FRHICommandListExecutor::GetImmediateCommandList();
-		check(&RHICmdListImmediate == &RHICmdList);
-
 		const FString UniqueResourceSubResourceName = GetUniqueSubResourceName(SubresourceDesc);
 		const FTextureSubresourceDumpDesc SubresourceDumpDesc = TranslateSubresourceDumpDesc(SubresourceDesc);
 
@@ -1310,11 +1305,11 @@ public:
 		if (SubresourceDumpDesc.bPreprocessForStaging)
 		{
 			SubResourceSRV = CreateSRVNoLifetimeExtension(RHICmdList, RHITexture, FRHITextureSRVCreateInfo(SubresourceDesc));
-			RHICmdListImmediate.Transition(FRHITransitionInfo(RHITexture, ERHIAccess::Unknown, ERHIAccess::SRVCompute));
+			RHICmdList.Transition(FRHITransitionInfo(RHITexture, ERHIAccess::Unknown, ERHIAccess::SRVCompute));
 		}
 		else
 		{
-			RHICmdListImmediate.Transition(FRHITransitionInfo(RHITexture, RHIAccessState, ERHIAccess::CopySrc));
+			RHICmdList.Transition(FRHITransitionInfo(RHITexture, RHIAccessState, ERHIAccess::CopySrc));
 		}
 
 		FString DumpFilePath = kResourcesDir / FString::Printf(
@@ -1324,7 +1319,7 @@ public:
 			DrawDumpCount);
 
 		DumpTextureSubResource(
-			RHICmdListImmediate,
+			RHICmdList,
 			SubresourceDesc.Texture->Name,
 			RHITexture,
 			SubResourceSRV,
@@ -1333,17 +1328,17 @@ public:
 
 		if (SubresourceDumpDesc.bPreprocessForStaging)
 		{
-			RHICmdListImmediate.Transition(FRHITransitionInfo(RHITexture, ERHIAccess::SRVCompute, RHIAccessState));
+			RHICmdList.Transition(FRHITransitionInfo(RHITexture, ERHIAccess::SRVCompute, RHIAccessState));
 		}
 		else
 		{
-			RHICmdListImmediate.Transition(FRHITransitionInfo(RHITexture, ERHIAccess::CopySrc, RHIAccessState));
+			RHICmdList.Transition(FRHITransitionInfo(RHITexture, ERHIAccess::CopySrc, RHIAccessState));
 		}
 
 		SubResourceSRV = nullptr;
 		if (!bStream)
 		{
-			ReleaseRHIResources(RHICmdListImmediate);
+			ReleaseRHIResources(RHICmdList);
 		}
 	}
 
@@ -2739,6 +2734,8 @@ void FRDGBuilder::DumpDraw(const FRDGEventName& DrawEventName)
 		return;
 	}
 
+	FRHICommandListImmediate& RHICmdList = FRHICommandListImmediate::Get();
+
 	FRDGResourceDumpContext* ResourceDumpContext = GRDGResourceDumpContext_RenderThread;
 
 	if (!ResourceDumpContext->DrawDumpingPass)
@@ -2748,8 +2745,6 @@ void FRDGBuilder::DumpDraw(const FRDGEventName& DrawEventName)
 
 	const FRDGPass* Pass = ResourceDumpContext->DrawDumpingPass;
 
-	FRHICommandListImmediate& RHICmdList = FRHICommandListExecutor::GetImmediateCommandList();
-
 	if (EnumHasAnyFlags(Pass->GetFlags(), ERDGPassFlags::Raster))
 	{
 		RHICmdList.EndRenderPass();
diff --git a/Engine/Source/Runtime/RenderCore/Private/RenderGraphEvent.cpp b/Engine/Source/Runtime/RenderCore/Private/RenderGraphEvent.cpp
--- a/Engine/Source/Runtime/RenderCore/Private/RenderGraphEvent.cpp
+++ b/Engine/Source/Runtime/RenderCore/Private/RenderGraphEvent.cpp
@@ -244,8 +244,7 @@ namespace DynamicRenderScaling
 
 	void EndFrame()
 	{
-		check(IsInRenderingThread());
-		GRDGTimingPool.EndFrame(FRHICommandListExecutor::GetImmediateCommandList());
+		GRDGTimingPool.EndFrame(FRHICommandListImmediate::Get());
 	}
 
 	const TMap<uint64>& GetLatestTimings()
diff --git a/Engine/Source/Runtime/RenderCore/Private/RenderResource.cpp b/Engine/Source/Runtime/RenderCore/Private/RenderResource.cpp
--- a/Engine/Source/Runtime/RenderCore/Private/RenderResource.cpp
+++ b/Engine/Source/Runtime/RenderCore/Private/RenderResource.cpp
@@ -122,8 +122,7 @@ void FRenderResource::ReleaseRHIForAllResources()
 /** Initialize all resources initialized before the RHI was initialized */
 void FRenderResource::InitPreRHIResources()
 {
-	check(IsInRenderingThread());
-	FRHICommandListImmediate& RHICmdList = FRHICommandListExecutor::GetImmediateCommandList();
+	FRHICommandListImmediate& RHICmdList = FRHICommandListImmediate::Get();
 	RHICmdList.InitializeImmediateContexts();
 
 	FRenderResourceList& PreResourceList = FRenderResourceList::Get<FRenderResource::EInitPhase::Pre>();
@@ -163,8 +162,7 @@ void FRenderResource::ChangeFeatureLevel(ERHIFeatureLevel::Type NewFeatureLevel)
 
 FRHICommandListBase& FRenderResource::GetImmediateCommandList()
 {
-	check(IsInRenderingThread());
-	return FRHICommandListExecutor::GetImmediateCommandList();
+	return FRHICommandListImmediate::Get();
 }
 
 void FRenderResource::InitResource(FRHICommandListBase& RHICmdList)
@@ -216,8 +214,7 @@ void FRenderResource::ReleaseResource()
 
 void FRenderResource::UpdateRHI()
 {
-	check(IsInRenderingThread());
-	UpdateRHI(FRHICommandListExecutor::GetImmediateCommandList());
+	UpdateRHI(FRHICommandListImmediate::Get());
 }
 
 void FRenderResource::UpdateRHI(FRHICommandListBase& RHICmdList)
diff --git a/Engine/Source/Runtime/Renderer/Private/ReflectionEnvironment.cpp b/Engine/Source/Runtime/Renderer/Private/ReflectionEnvironment.cpp
--- a/Engine/Source/Runtime/Renderer/Private/ReflectionEnvironment.cpp
+++ b/Engine/Source/Runtime/Renderer/Private/ReflectionEnvironment.cpp
@@ -530,9 +530,10 @@ void FReflectionEnvironmentSceneData::Reset(FScene* Scene)
 
 void FReflectionEnvironmentCubemapArray::ResizeCubemapArrayGPU(uint32 InMaxCubemaps, int32 InCubemapSize, const TArray<int32>& IndexRemapping)
 {
+	FRHICommandListImmediate& RHICmdList = FRHICommandListImmediate::Get();
+
 	check(InMaxCubemaps > 0);
 	check(InCubemapSize > 0);
-	check(IsInRenderingThread());
 	check(IsInitialized());
 	check(InCubemapSize == CubemapSize);
 
@@ -542,7 +543,6 @@ void FReflectionEnvironmentCubemapArray::ResizeCubemapArrayGPU(uint32 InMaxCubem
 	int OldMaxCubemaps = MaxCubemaps;
 	MaxCubemaps = InMaxCubemaps;
 
-	FRHICommandListImmediate& RHICmdList = FRHICommandListExecutor::GetImmediateCommandList();
 	InitRHI(RHICmdList);
 
 	FTextureRHIRef TexRef = OldReflectionEnvs->GetRHI();
diff --git a/Engine/Source/Runtime/Renderer/Private/ReflectionEnvironmentCapture.cpp b/Engine/Source/Runtime/Renderer/Private/ReflectionEnvironmentCapture.cpp
--- a/Engine/Source/Runtime/Renderer/Private/ReflectionEnvironmentCapture.cpp
+++ b/Engine/Source/Runtime/Renderer/Private/ReflectionEnvironmentCapture.cpp
@@ -530,7 +530,7 @@ void CaptureSceneToScratchCubemap(
 		AddPass(GraphBuilder, RDG_EVENT_NAME("FlushGPU"), [](FRHICommandListImmediate& InRHICmdList)
 		{
 			QUICK_SCOPE_CYCLE_COUNTER(STAT_CaptureSceneToScratchCubemap_Flush);
-			FRHICommandListExecutor::GetImmediateCommandList().ImmediateFlush(EImmediateFlushType::FlushRHIThread);
+			InRHICmdList.ImmediateFlush(EImmediateFlushType::FlushRHIThread);
 
 			// some platforms may not be able to keep enqueueing commands like crazy, this will
 			// allow them to restart their command buffers
diff --git a/Engine/Source/Runtime/VulkanRHI/Private/VulkanQuery.cpp b/Engine/Source/Runtime/VulkanRHI/Private/VulkanQuery.cpp
--- a/Engine/Source/Runtime/VulkanRHI/Private/VulkanQuery.cpp
+++ b/Engine/Source/Runtime/VulkanRHI/Private/VulkanQuery.cpp
@@ -443,7 +443,9 @@ bool FVulkanDynamicRHI::RHIGetRenderQueryResult(FRHIRenderQuery* QueryRHI, uint6
 		uint64 Microseconds = (uint64)((double(Timestamp) / Frequency) * 1000.0 * 1000.0);
 		return Microseconds;
 	};
-	check(IsInRenderingThread());
+
+	FRHICommandListImmediate& RHICmdList = FRHICommandListImmediate::Get();
+
 	FVulkanRenderQuery* BaseQuery = ResourceCast(QueryRHI);
 	if (BaseQuery->QueryType == RQT_Occlusion)
 	{
@@ -482,7 +484,7 @@ bool FVulkanDynamicRHI::RHIGetRenderQueryResult(FRHIRenderQuery* QueryRHI, uint6
 			}
 
 			// Blocking was requested, so, we need to wait for the RHI thread to catch up.
-			FRHICommandListExecutor::GetImmediateCommandList().ImmediateFlush(EImmediateFlushType::FlushRHIThread);
+			RHICmdList.ImmediateFlush(EImmediateFlushType::FlushRHIThread);
 			check(Query->Pool != nullptr);
 		}
 
@@ -523,7 +525,7 @@ bool FVulkanDynamicRHI::RHIGetRenderQueryResult(FRHIRenderQuery* QueryRHI, uint6
 					bool bWaitForStart = StartQuerySyncPoint.FenceCounter == StartQuerySyncPoint.CmdBuffer->GetFenceSignaledCounter();
 					if (bWaitForStart)
 					{
-						FRHICommandListExecutor::GetImmediateCommandList().ImmediateFlush(EImmediateFlushType::FlushRHIThread);
+						RHICmdList.ImmediateFlush(EImmediateFlushType::FlushRHIThread);
 
 						// Need to submit the open command lists.
 						Device->SubmitCommandsAndFlushGPU();
diff --git a/Engine/Source/Runtime/VulkanRHI/Private/VulkanRHI.cpp b/Engine/Source/Runtime/VulkanRHI/Private/VulkanRHI.cpp
--- a/Engine/Source/Runtime/VulkanRHI/Private/VulkanRHI.cpp
+++ b/Engine/Source/Runtime/VulkanRHI/Private/VulkanRHI.cpp
@@ -617,7 +617,7 @@ void FVulkanDynamicRHI::Shutdown()
 #endif // VULKAN_RHI_RAYTRACING
 
 		// Flush all pending deletes before destroying the device.
-		FRHICommandListExecutor::GetImmediateCommandList().ImmediateFlush(EImmediateFlushType::FlushRHIThreadFlushResources);
+		FRHICommandListImmediate::Get().ImmediateFlush(EImmediateFlushType::FlushRHIThreadFlushResources);
 	}
 
 	Device->Destroy();
diff --git a/Engine/Source/Runtime/VulkanRHI/Private/VulkanViewport.cpp b/Engine/Source/Runtime/VulkanRHI/Private/VulkanViewport.cpp
--- a/Engine/Source/Runtime/VulkanRHI/Private/VulkanViewport.cpp
+++ b/Engine/Source/Runtime/VulkanRHI/Private/VulkanViewport.cpp
@@ -1180,7 +1180,8 @@ void FVulkanDynamicRHI::RHITick(float DeltaTime)
 
 FTexture2DRHIRef FVulkanDynamicRHI::RHIGetViewportBackBuffer(FRHIViewport* ViewportRHI)
 {
-	check(IsInRenderingThread());
+	FRHICommandListImmediate& RHICmdList = FRHICommandListImmediate::Get();
+
 	check(ViewportRHI);
 	FVulkanViewport* Viewport = ResourceCast(ViewportRHI);
 
@@ -1189,15 +1190,15 @@ FTexture2DRHIRef FVulkanDynamicRHI::RHIGetViewportBackBuffer(FRHIViewport* Viewp
 		Viewport->SwapChain->RenderThreadPacing();
 	}
 
-	return Viewport->GetBackBuffer(FRHICommandListExecutor::GetImmediateCommandList());
+	return Viewport->GetBackBuffer(RHICmdList);
 }
 
 void FVulkanDynamicRHI::RHIAdvanceFrameForGetViewportBackBuffer(FRHIViewport* ViewportRHI)
 {
-	check(IsInRenderingThread());
+	FRHICommandListImmediate& RHICmdList = FRHICommandListImmediate::Get();
+
 	check(ViewportRHI);
 	FVulkanViewport* Viewport = ResourceCast(ViewportRHI);
-	FRHICommandListImmediate& RHICmdList = FRHICommandListExecutor::GetImmediateCommandList();
 
 	Viewport->AdvanceBackBufferFrame(RHICmdList);
 }
diff --git a/Engine/Source/Runtime/Windows/D3D11RHI/Private/D3D11Texture.cpp b/Engine/Source/Runtime/Windows/D3D11RHI/Private/D3D11Texture.cpp
--- a/Engine/Source/Runtime/Windows/D3D11RHI/Private/D3D11Texture.cpp
+++ b/Engine/Source/Runtime/Windows/D3D11RHI/Private/D3D11Texture.cpp
@@ -1214,7 +1214,7 @@ void* FD3D11DynamicRHI::LockTexture2D_RenderThread(
 	}
 	else if (LockMode == RLM_ReadOnly)
 	{
-		FRHICommandListExecutor::GetImmediateCommandList().ImmediateFlush(EImmediateFlushType::FlushRHIThread);
+		RHICmdList.ImmediateFlush(EImmediateFlushType::FlushRHIThread);
 		LockedTexture = RHILockTexture2D(Texture, MipIndex, LockMode, DestStride, bLockWithinMiptail, OutLockedByteCount);
 	}
 	else
@@ -1391,7 +1391,7 @@ void* FD3D11DynamicRHI::RHILockTextureCubeFace_RenderThread(
 	}
 	else if (LockMode == RLM_ReadOnly)
 	{
-		FRHICommandListExecutor::GetImmediateCommandList().ImmediateFlush(EImmediateFlushType::FlushRHIThread);
+		RHICmdList.ImmediateFlush(EImmediateFlushType::FlushRHIThread);
 		LockedTexture = RHILockTextureCubeFace(Texture, FaceIndex, ArrayIndex, MipIndex, LockMode, DestStride, bLockWithinMiptail);
 	}
 	else
