RHIBindAccelerationStructureMemory and RHIClearRayTracingBindings modify members of the FRHIRayTracingScene inside platform RHI implementations. It therefore needs the RHI lock fence to prevent races those members.

diff --git a/Engine/Source/Runtime/RHI/Public/RHICommandList.h b/Engine/Source/Runtime/RHI/Public/RHICommandList.h
--- a/Engine/Source/Runtime/RHI/Public/RHICommandList.h
+++ b/Engine/Source/Runtime/RHI/Public/RHICommandList.h
@@ -3038,6 +3038,10 @@ public:
 		else
 		{
 			ALLOC_COMMAND(FRHICommandBindAccelerationStructureMemory)(Scene, Buffer, BufferOffset);
+
+			// This RHI command modifies members of the FRHIRayTracingScene inside platform RHI implementations.
+			// It therefore needs the RHI lock fence to prevent races on those members.
+			RHIThreadFence(true);
 		}
 	}
 #endif
@@ -3605,6 +3609,10 @@ public:
 		else
 		{
 			ALLOC_COMMAND(FRHICommandClearRayTracingBindings)(Scene);
+
+			// This RHI command modifies members of the FRHIRayTracingScene inside platform RHI implementations.
+			// It therefore needs the RHI lock fence to prevent races on those members.
+			RHIThreadFence(true);
 		}
 	}
 
