Deprecate:

diff --git a/Engine/Source/Runtime/RHI/Private/RHICommandList.cpp b/Engine/Source/Runtime/RHI/Private/RHICommandList.cpp
--- a/Engine/Source/Runtime/RHI/Private/RHICommandList.cpp
+++ b/Engine/Source/Runtime/RHI/Private/RHICommandList.cpp
@@ -1927,7 +1927,7 @@ void FDynamicRHI::UnlockStagingBuffer_RenderThread(class FRHICommandListImmediat
 	}
 }
 
-FTexture2DRHIRef FDynamicRHI::AsyncReallocateTexture2D_RenderThread(class FRHICommandListImmediate& RHICmdList, FRHITexture2D* Texture2D, int32 NewMipCount, int32 NewSizeX, int32 NewSizeY, FThreadSafeCounter* RequestStatus)
+FTextureRHIRef FDynamicRHI::AsyncReallocateTexture2D_RenderThread(class FRHICommandListImmediate& RHICmdList, FRHITexture* Texture2D, int32 NewMipCount, int32 NewSizeX, int32 NewSizeY, FThreadSafeCounter* RequestStatus)
 {
 	QUICK_SCOPE_CYCLE_COUNTER(STAT_RHIMETHOD_AsyncReallocateTexture2D_Flush);
 	CSV_SCOPED_TIMING_STAT(RHITFlushes, AsyncReallocateTexture2D_RenderThread);
@@ -1935,21 +1935,21 @@ FTexture2DRHIRef FDynamicRHI::AsyncReallocateTexture2D_RenderThread(class FRHICo
 	return GDynamicRHI->RHIAsyncReallocateTexture2D(Texture2D, NewMipCount, NewSizeX, NewSizeY, RequestStatus);
 }
 
-ETextureReallocationStatus FDynamicRHI::FinalizeAsyncReallocateTexture2D_RenderThread(class FRHICommandListImmediate& RHICmdList, FRHITexture2D* Texture2D, bool bBlockUntilCompleted)
+ETextureReallocationStatus FDynamicRHI::FinalizeAsyncReallocateTexture2D_RenderThread(class FRHICommandListImmediate& RHICmdList, FRHITexture* Texture2D, bool bBlockUntilCompleted)
 {
 	CSV_SCOPED_TIMING_STAT(RHITStalls, FinalizeAsyncReallocateTexture2D_RenderThread);
 	FScopedRHIThreadStaller StallRHIThread(RHICmdList);
 	return GDynamicRHI->RHIFinalizeAsyncReallocateTexture2D(Texture2D, bBlockUntilCompleted);
 }
 
-ETextureReallocationStatus FDynamicRHI::CancelAsyncReallocateTexture2D_RenderThread(class FRHICommandListImmediate& RHICmdList, FRHITexture2D* Texture2D, bool bBlockUntilCompleted)
+ETextureReallocationStatus FDynamicRHI::CancelAsyncReallocateTexture2D_RenderThread(class FRHICommandListImmediate& RHICmdList, FRHITexture* Texture2D, bool bBlockUntilCompleted)
 {
 	CSV_SCOPED_TIMING_STAT(RHITStalls, CancelAsyncReallocateTexture2D_RenderThread);
 	FScopedRHIThreadStaller StallRHIThread(RHICmdList);
 	return GDynamicRHI->RHICancelAsyncReallocateTexture2D(Texture2D, bBlockUntilCompleted);
 }
 
-FUpdateTexture3DData FDynamicRHI::RHIBeginUpdateTexture3D(FRHICommandListBase& RHICmdList, FRHITexture3D* Texture, uint32 MipIndex, const struct FUpdateTextureRegion3D& UpdateRegion)
+FUpdateTexture3DData FDynamicRHI::RHIBeginUpdateTexture3D(FRHICommandListBase& RHICmdList, FRHITexture* Texture, uint32 MipIndex, const struct FUpdateTextureRegion3D& UpdateRegion)
 {
 	check(IsInParallelRenderingThread());
 
@@ -1980,7 +1980,7 @@ void FDynamicRHI::RHIEndMultiUpdateTexture3D(FRHICommandListBase& RHICmdList, TA
 	}
 }
 
-void* FDynamicRHI::LockTexture2D_RenderThread(class FRHICommandListImmediate& RHICmdList, FRHITexture2D* Texture, uint32 MipIndex, EResourceLockMode LockMode, uint32& DestStride, bool bLockWithinMiptail, bool bNeedsDefaultRHIFlush, uint64* OutLockedByteCount)
+void* FDynamicRHI::LockTexture2D_RenderThread(class FRHICommandListImmediate& RHICmdList, FRHITexture* Texture, uint32 MipIndex, EResourceLockMode LockMode, uint32& DestStride, bool bLockWithinMiptail, bool bNeedsDefaultRHIFlush, uint64* OutLockedByteCount)
 {
 	if (bNeedsDefaultRHIFlush) 
 	{
@@ -1994,7 +1994,7 @@ void* FDynamicRHI::LockTexture2D_RenderThread(class FRHICommandListImmediate& RH
 	return GDynamicRHI->RHILockTexture2D(Texture, MipIndex, LockMode, DestStride, bLockWithinMiptail, OutLockedByteCount);
 }
 
-void FDynamicRHI::UnlockTexture2D_RenderThread(class FRHICommandListImmediate& RHICmdList, FRHITexture2D* Texture, uint32 MipIndex, bool bLockWithinMiptail, bool bNeedsDefaultRHIFlush)
+void FDynamicRHI::UnlockTexture2D_RenderThread(class FRHICommandListImmediate& RHICmdList, FRHITexture* Texture, uint32 MipIndex, bool bLockWithinMiptail, bool bNeedsDefaultRHIFlush)
 {
 	if (bNeedsDefaultRHIFlush)
 	{
@@ -2009,7 +2009,7 @@ void FDynamicRHI::UnlockTexture2D_RenderThread(class FRHICommandListImmediate& R
 	GDynamicRHI->RHIUnlockTexture2D(Texture, MipIndex, bLockWithinMiptail);
 }
 
-void* FDynamicRHI::LockTexture2DArray_RenderThread(class FRHICommandListImmediate& RHICmdList, FRHITexture2DArray* Texture, uint32 ArrayIndex, uint32 MipIndex, EResourceLockMode LockMode, uint32& DestStride, bool bLockWithinMiptail)
+void* FDynamicRHI::LockTexture2DArray_RenderThread(class FRHICommandListImmediate& RHICmdList, FRHITexture* Texture, uint32 ArrayIndex, uint32 MipIndex, EResourceLockMode LockMode, uint32& DestStride, bool bLockWithinMiptail)
 {
 	QUICK_SCOPE_CYCLE_COUNTER(STAT_RHIMETHOD_LockTexture2DArray_Flush);
 	CSV_SCOPED_TIMING_STAT(RHITFlushes, LockTexture2DArray_RenderThread);
@@ -2017,7 +2017,7 @@ void* FDynamicRHI::LockTexture2DArray_RenderThread(class FRHICommandListImmediat
 	return GDynamicRHI->RHILockTexture2DArray(Texture, ArrayIndex, MipIndex, LockMode, DestStride, bLockWithinMiptail);
 }
 
-void FDynamicRHI::UnlockTexture2DArray_RenderThread(class FRHICommandListImmediate& RHICmdList, FRHITexture2DArray* Texture, uint32 ArrayIndex, uint32 MipIndex, bool bLockWithinMiptail)
+void FDynamicRHI::UnlockTexture2DArray_RenderThread(class FRHICommandListImmediate& RHICmdList, FRHITexture* Texture, uint32 ArrayIndex, uint32 MipIndex, bool bLockWithinMiptail)
 {
 	QUICK_SCOPE_CYCLE_COUNTER(STAT_RHIMETHOD_UnlockTexture2DArray_Flush);
 	CSV_SCOPED_TIMING_STAT(RHITFlushes, UnlockTexture2DArray_RenderThread);
@@ -2031,7 +2031,7 @@ FRHIShaderLibraryRef FDynamicRHI::RHICreateShaderLibrary_RenderThread(class FRHI
 	return GDynamicRHI->RHICreateShaderLibrary(Platform, FilePath, Name);
 }
 
-void* FDynamicRHI::RHILockTextureCubeFace_RenderThread(class FRHICommandListImmediate& RHICmdList, FRHITextureCube* Texture, uint32 FaceIndex, uint32 ArrayIndex, uint32 MipIndex, EResourceLockMode LockMode, uint32& DestStride, bool bLockWithinMiptail)
+void* FDynamicRHI::RHILockTextureCubeFace_RenderThread(class FRHICommandListImmediate& RHICmdList, FRHITexture* Texture, uint32 FaceIndex, uint32 ArrayIndex, uint32 MipIndex, EResourceLockMode LockMode, uint32& DestStride, bool bLockWithinMiptail)
 {
 	QUICK_SCOPE_CYCLE_COUNTER(STAT_RHIMETHOD_LockTextureCubeFace_Flush);
 	CSV_SCOPED_TIMING_STAT(RHITFlushes, RHILockTextureCubeFace_RenderThread);
@@ -2039,7 +2039,7 @@ void* FDynamicRHI::RHILockTextureCubeFace_RenderThread(class FRHICommandListImme
 	return GDynamicRHI->RHILockTextureCubeFace(Texture, FaceIndex, ArrayIndex, MipIndex, LockMode, DestStride, bLockWithinMiptail);
 }
 
-void FDynamicRHI::RHIUnlockTextureCubeFace_RenderThread(class FRHICommandListImmediate& RHICmdList, FRHITextureCube* Texture, uint32 FaceIndex, uint32 ArrayIndex, uint32 MipIndex, bool bLockWithinMiptail)
+void FDynamicRHI::RHIUnlockTextureCubeFace_RenderThread(class FRHICommandListImmediate& RHICmdList, FRHITexture* Texture, uint32 FaceIndex, uint32 ArrayIndex, uint32 MipIndex, bool bLockWithinMiptail)
 {
 	QUICK_SCOPE_CYCLE_COUNTER(STAT_RHIMETHOD_UnlockTextureCubeFace_Flush);
 	CSV_SCOPED_TIMING_STAT(RHITFlushes, RHIUnlockTextureCubeFace_RenderThread);
diff --git a/Engine/Source/Runtime/RHI/Public/DynamicRHI.h b/Engine/Source/Runtime/RHI/Public/DynamicRHI.h
--- a/Engine/Source/Runtime/RHI/Public/DynamicRHI.h
+++ b/Engine/Source/Runtime/RHI/Public/DynamicRHI.h
@@ -463,7 +463,7 @@ public:
 	*/
 	// FlushType: Flush RHI Thread
 	// NP: Note that no RHI currently implements this as an async call, we should simplify the API.
-	virtual FTexture2DRHIRef RHIAsyncReallocateTexture2D(FRHITexture2D* Texture2D, int32 NewMipCount, int32 NewSizeX, int32 NewSizeY, FThreadSafeCounter* RequestStatus) = 0;
+	virtual FTextureRHIRef RHIAsyncReallocateTexture2D(FRHITexture* Texture2D, int32 NewMipCount, int32 NewSizeX, int32 NewSizeY, FThreadSafeCounter* RequestStatus) = 0;
 
 	/**
 	* Finalizes an async reallocation request.
@@ -477,7 +477,7 @@ public:
 	*	TexRealloc_InProgress	Reallocation is still in progress, try again later
 	*/
 	// FlushType: Wait RHI Thread
-	virtual ETextureReallocationStatus RHIFinalizeAsyncReallocateTexture2D(FRHITexture2D* Texture2D, bool bBlockUntilCompleted) = 0;
+	virtual ETextureReallocationStatus RHIFinalizeAsyncReallocateTexture2D(FRHITexture* Texture2D, bool bBlockUntilCompleted) = 0;
 
 	/**
 	* Cancels an async reallocation for the specified texture.
@@ -488,7 +488,7 @@ public:
 	* @return						Reallocation status
 	*/
 	// FlushType: Wait RHI Thread
-	virtual ETextureReallocationStatus RHICancelAsyncReallocateTexture2D(FRHITexture2D* Texture2D, bool bBlockUntilCompleted) = 0;
+	virtual ETextureReallocationStatus RHICancelAsyncReallocateTexture2D(FRHITexture* Texture2D, bool bBlockUntilCompleted) = 0;
 
 	/**
 	* Locks an RHI texture's mip-map for read/write operations on the CPU
@@ -500,7 +500,7 @@ public:
 	* @return pointer to the CPU accessible resource data
 	*/
 	// FlushType: Flush RHI Thread
-	virtual void* RHILockTexture2D(FRHITexture2D* Texture, uint32 MipIndex, EResourceLockMode LockMode, uint32& DestStride, bool bLockWithinMiptail, uint64* OutLockedByteCount = nullptr) = 0;
+	virtual void* RHILockTexture2D(FRHITexture* Texture, uint32 MipIndex, EResourceLockMode LockMode, uint32& DestStride, bool bLockWithinMiptail, uint64* OutLockedByteCount = nullptr) = 0;
 
 	/**
 	* Unlocks a previously locked RHI texture resource
@@ -509,7 +509,7 @@ public:
 	* @param bLockWithinMiptail - for platforms that support packed miptails allow locking of individual mip levels within the miptail
 	*/
 	// FlushType: Flush RHI Thread
-	virtual void RHIUnlockTexture2D(FRHITexture2D* Texture, uint32 MipIndex, bool bLockWithinMiptail) = 0;
+	virtual void RHIUnlockTexture2D(FRHITexture* Texture, uint32 MipIndex, bool bLockWithinMiptail) = 0;
 
 	/**
 	* Locks an RHI texture's mip-map for read/write operations on the CPU
@@ -521,7 +521,7 @@ public:
 	* @return pointer to the CPU accessible resource data
 	*/
 	// FlushType: Flush RHI Thread
-	virtual void* RHILockTexture2DArray(FRHITexture2DArray* Texture, uint32 TextureIndex, uint32 MipIndex, EResourceLockMode LockMode, uint32& DestStride, bool bLockWithinMiptail) = 0;
+	virtual void* RHILockTexture2DArray(FRHITexture* Texture, uint32 TextureIndex, uint32 MipIndex, EResourceLockMode LockMode, uint32& DestStride, bool bLockWithinMiptail) = 0;
 
 	/**
 	* Unlocks a previously locked RHI texture resource
@@ -530,7 +530,7 @@ public:
 	* @param bLockWithinMiptail - for platforms that support packed miptails allow locking of individual mip levels within the miptail
 	*/
 	// FlushType: Flush RHI Thread
-	virtual void RHIUnlockTexture2DArray(FRHITexture2DArray* Texture, uint32 TextureIndex, uint32 MipIndex, bool bLockWithinMiptail) = 0;
+	virtual void RHIUnlockTexture2DArray(FRHITexture* Texture, uint32 TextureIndex, uint32 MipIndex, bool bLockWithinMiptail) = 0;
 
 	/**
 	* Updates a region of a 2D texture from system memory
@@ -540,7 +540,7 @@ public:
 	* @param SourcePitch - size in bytes of each row of the source image
 	* @param SourceData - source image data, starting at the upper left corner of the source rectangle (in same pixel format as texture)
 	*/
-	virtual void RHIUpdateTexture2D(FRHICommandListBase& RHICmdList, FRHITexture2D* Texture, uint32 MipIndex, const struct FUpdateTextureRegion2D& UpdateRegion, uint32 SourcePitch, const uint8* SourceData) = 0;
+	virtual void RHIUpdateTexture2D(FRHICommandListBase& RHICmdList, FRHITexture* Texture, uint32 MipIndex, const struct FUpdateTextureRegion2D& UpdateRegion, uint32 SourcePitch, const uint8* SourceData) = 0;
 
 	/**
 	* Updates a region of a 2D texture from GPU memory provided by the given buffer (may not be implemented on every platform)
@@ -550,7 +550,7 @@ public:
 	* @param SourcePitch - size in bytes of each row of the source image
 	* @param Buffer, BufferOffset - source image data, starting at the upper left corner of the source rectangle (in same pixel format as texture)
 	*/
-	virtual void RHIUpdateFromBufferTexture2D(FRHICommandListBase& RHICmdList, FRHITexture2D* Texture, uint32 MipIndex, const struct FUpdateTextureRegion2D& UpdateRegion, uint32 SourcePitch, FRHIBuffer* Buffer, uint32 BufferOffset)
+	virtual void RHIUpdateFromBufferTexture2D(FRHICommandListBase& RHICmdList, FRHITexture* Texture, uint32 MipIndex, const struct FUpdateTextureRegion2D& UpdateRegion, uint32 SourcePitch, FRHIBuffer* Buffer, uint32 BufferOffset)
 	{
 		checkNoEntry();
 	}
@@ -564,7 +564,7 @@ public:
 	* @param SourceDepthPitch - size in bytes of each depth slice of the source image, usually Bpp * SizeX * SizeY
 	* @param SourceData - source image data, starting at the upper left corner of the source rectangle (in same pixel format as texture)
 	*/
-	virtual void RHIUpdateTexture3D(FRHICommandListBase& RHICmdList, FRHITexture3D* Texture, uint32 MipIndex, const struct FUpdateTextureRegion3D& UpdateRegion, uint32 SourceRowPitch, uint32 SourceDepthPitch, const uint8* SourceData) = 0;
+	virtual void RHIUpdateTexture3D(FRHICommandListBase& RHICmdList, FRHITexture* Texture, uint32 MipIndex, const struct FUpdateTextureRegion3D& UpdateRegion, uint32 SourceRowPitch, uint32 SourceDepthPitch, const uint8* SourceData) = 0;
 
 	/**
 	* Locks an RHI texture's mip-map for read/write operations on the CPU
@@ -576,7 +576,7 @@ public:
 	* @return pointer to the CPU accessible resource data
 	*/
 	// FlushType: Flush RHI Thread
-	virtual void* RHILockTextureCubeFace(FRHITextureCube* Texture, uint32 FaceIndex, uint32 ArrayIndex, uint32 MipIndex, EResourceLockMode LockMode, uint32& DestStride, bool bLockWithinMiptail) = 0;
+	virtual void* RHILockTextureCubeFace(FRHITexture* Texture, uint32 FaceIndex, uint32 ArrayIndex, uint32 MipIndex, EResourceLockMode LockMode, uint32& DestStride, bool bLockWithinMiptail) = 0;
 
 	/**
 	* Unlocks a previously locked RHI texture resource
@@ -585,7 +585,7 @@ public:
 	* @param bLockWithinMiptail - for platforms that support packed miptails allow locking of individual mip levels within the miptail
 	*/
 	// FlushType: Flush RHI Thread
-	virtual void RHIUnlockTextureCubeFace(FRHITextureCube* Texture, uint32 FaceIndex, uint32 ArrayIndex, uint32 MipIndex, bool bLockWithinMiptail) = 0;
+	virtual void RHIUnlockTextureCubeFace(FRHITexture* Texture, uint32 FaceIndex, uint32 ArrayIndex, uint32 MipIndex, bool bLockWithinMiptail) = 0;
 
 	// FlushType: Thread safe
 	virtual void RHIBindDebugLabelName(FRHICommandListBase& RHICmdList,FRHITexture* Texture, const TCHAR* Name) = 0;
@@ -643,7 +643,7 @@ public:
 
 	// With RHI thread, this is the current backbuffer from the perspective of the render thread.
 	// FlushType: Thread safe
-	virtual FTexture2DRHIRef RHIGetViewportBackBuffer(FRHIViewport* Viewport) = 0;
+	virtual FTextureRHIRef RHIGetViewportBackBuffer(FRHIViewport* Viewport) = 0;
 
 	virtual FUnorderedAccessViewRHIRef RHIGetViewportBackBufferUAV(FRHIViewport* ViewportRHI)
 	{
@@ -777,7 +777,7 @@ public:
 	* @param FirstMip - the first mip that should be in memory
 	*/
 	// FlushType: Wait RHI Thread
-	RHI_API virtual void RHIVirtualTextureSetFirstMipInMemory(class FRHICommandListImmediate& RHICmdList, FRHITexture2D* Texture, uint32 FirstMip);
+	RHI_API virtual void RHIVirtualTextureSetFirstMipInMemory(class FRHICommandListImmediate& RHICmdList, FRHITexture* Texture, uint32 FirstMip);
 
 	/**
 	* Function that can be used to update which is the first visible mip to the GPU.
@@ -785,7 +785,7 @@ public:
 	* @param FirstMip - the first mip that should be visible to the GPU
 	*/
 	// FlushType: Wait RHI Thread
-	RHI_API virtual void RHIVirtualTextureSetFirstMipVisible(class FRHICommandListImmediate& RHICmdList, FRHITexture2D* Texture, uint32 FirstMip);
+	RHI_API virtual void RHIVirtualTextureSetFirstMipVisible(class FRHICommandListImmediate& RHICmdList, FRHITexture* Texture, uint32 FirstMip);
 
 	/**
 	* Provides access to the native device. Generally this should be avoided but is useful for third party plugins.
@@ -881,25 +881,25 @@ public:
 	//
 	virtual void RHIProcessDeleteQueue() {}
 
-	RHI_API virtual FTexture2DRHIRef AsyncReallocateTexture2D_RenderThread(class FRHICommandListImmediate& RHICmdList, FRHITexture2D* Texture2D, int32 NewMipCount, int32 NewSizeX, int32 NewSizeY, FThreadSafeCounter* RequestStatus);
-	RHI_API virtual ETextureReallocationStatus FinalizeAsyncReallocateTexture2D_RenderThread(class FRHICommandListImmediate& RHICmdList, FRHITexture2D* Texture2D, bool bBlockUntilCompleted);
-	RHI_API virtual ETextureReallocationStatus CancelAsyncReallocateTexture2D_RenderThread(class FRHICommandListImmediate& RHICmdList, FRHITexture2D* Texture2D, bool bBlockUntilCompleted);
+	RHI_API virtual FTextureRHIRef AsyncReallocateTexture2D_RenderThread(class FRHICommandListImmediate& RHICmdList, FRHITexture* Texture2D, int32 NewMipCount, int32 NewSizeX, int32 NewSizeY, FThreadSafeCounter* RequestStatus);
+	RHI_API virtual ETextureReallocationStatus FinalizeAsyncReallocateTexture2D_RenderThread(class FRHICommandListImmediate& RHICmdList, FRHITexture* Texture2D, bool bBlockUntilCompleted);
+	RHI_API virtual ETextureReallocationStatus CancelAsyncReallocateTexture2D_RenderThread(class FRHICommandListImmediate& RHICmdList, FRHITexture* Texture2D, bool bBlockUntilCompleted);
 
-	RHI_API virtual void* LockTexture2D_RenderThread(class FRHICommandListImmediate& RHICmdList, FRHITexture2D* Texture, uint32 MipIndex, EResourceLockMode LockMode, uint32& DestStride, bool bLockWithinMiptail, bool bNeedsDefaultRHIFlush = true, uint64* OutLockedByteCount = nullptr);
-	RHI_API virtual void UnlockTexture2D_RenderThread(class FRHICommandListImmediate& RHICmdList, FRHITexture2D* Texture, uint32 MipIndex, bool bLockWithinMiptail, bool bNeedsDefaultRHIFlush = true);
+	RHI_API virtual void* LockTexture2D_RenderThread(class FRHICommandListImmediate& RHICmdList, FRHITexture* Texture, uint32 MipIndex, EResourceLockMode LockMode, uint32& DestStride, bool bLockWithinMiptail, bool bNeedsDefaultRHIFlush = true, uint64* OutLockedByteCount = nullptr);
+	RHI_API virtual void UnlockTexture2D_RenderThread(class FRHICommandListImmediate& RHICmdList, FRHITexture* Texture, uint32 MipIndex, bool bLockWithinMiptail, bool bNeedsDefaultRHIFlush = true);
 
-	RHI_API virtual void* LockTexture2DArray_RenderThread(class FRHICommandListImmediate& RHICmdList, FRHITexture2DArray* Texture, uint32 ArrayIndex, uint32 MipIndex, EResourceLockMode LockMode, uint32& DestStride, bool bLockWithinMiptail);
-	RHI_API virtual void UnlockTexture2DArray_RenderThread(class FRHICommandListImmediate& RHICmdList, FRHITexture2DArray* Texture, uint32 ArrayIndex, uint32 MipIndex, bool bLockWithinMiptail);
+	RHI_API virtual void* LockTexture2DArray_RenderThread(class FRHICommandListImmediate& RHICmdList, FRHITexture* Texture, uint32 ArrayIndex, uint32 MipIndex, EResourceLockMode LockMode, uint32& DestStride, bool bLockWithinMiptail);
+	RHI_API virtual void UnlockTexture2DArray_RenderThread(class FRHICommandListImmediate& RHICmdList, FRHITexture* Texture, uint32 ArrayIndex, uint32 MipIndex, bool bLockWithinMiptail);
 
-	RHI_API virtual FUpdateTexture3DData RHIBeginUpdateTexture3D(FRHICommandListBase& RHICmdList, FRHITexture3D* Texture, uint32 MipIndex, const struct FUpdateTextureRegion3D& UpdateRegion);
+	RHI_API virtual FUpdateTexture3DData RHIBeginUpdateTexture3D(FRHICommandListBase& RHICmdList, FRHITexture* Texture, uint32 MipIndex, const struct FUpdateTextureRegion3D& UpdateRegion);
 	RHI_API virtual void RHIEndUpdateTexture3D(FRHICommandListBase& RHICmdList, FUpdateTexture3DData& UpdateData);
 
 	RHI_API virtual void RHIEndMultiUpdateTexture3D(FRHICommandListBase& RHICmdList, TArray<FUpdateTexture3DData>& UpdateDataArray);
 
 	RHI_API virtual FRHIShaderLibraryRef RHICreateShaderLibrary_RenderThread(class FRHICommandListImmediate& RHICmdList, EShaderPlatform Platform, FString FilePath, FString Name);
 
-	RHI_API virtual void* RHILockTextureCubeFace_RenderThread(class FRHICommandListImmediate& RHICmdList, FRHITextureCube* Texture, uint32 FaceIndex, uint32 ArrayIndex, uint32 MipIndex, EResourceLockMode LockMode, uint32& DestStride, bool bLockWithinMiptail);
-	RHI_API virtual void RHIUnlockTextureCubeFace_RenderThread(class FRHICommandListImmediate& RHICmdList, FRHITextureCube* Texture, uint32 FaceIndex, uint32 ArrayIndex, uint32 MipIndex, bool bLockWithinMiptail);
+	RHI_API virtual void* RHILockTextureCubeFace_RenderThread(class FRHICommandListImmediate& RHICmdList, FRHITexture* Texture, uint32 FaceIndex, uint32 ArrayIndex, uint32 MipIndex, EResourceLockMode LockMode, uint32& DestStride, bool bLockWithinMiptail);
+	RHI_API virtual void RHIUnlockTextureCubeFace_RenderThread(class FRHICommandListImmediate& RHICmdList, FRHITexture* Texture, uint32 FaceIndex, uint32 ArrayIndex, uint32 MipIndex, bool bLockWithinMiptail);
 
 	RHI_API virtual void RHIMapStagingSurface_RenderThread(class FRHICommandListImmediate& RHICmdList, FRHITexture* Texture, uint32 GPUIndex, FRHIGPUFence* Fence, void*& OutData, int32& OutWidth, int32& OutHeight);
 	RHI_API virtual void RHIUnmapStagingSurface_RenderThread(class FRHICommandListImmediate& RHICmdList, FRHITexture* Texture, uint32 GPUIndex);
@@ -1298,7 +1298,7 @@ FORCEINLINE uint32 RHIGetViewportNextPresentGPUIndex(FRHIViewport* Viewport)
 	return GDynamicRHI->RHIGetViewportNextPresentGPUIndex(Viewport);
 }
 
-FORCEINLINE FTexture2DRHIRef RHIGetViewportBackBuffer(FRHIViewport* Viewport)
+FORCEINLINE FTextureRHIRef RHIGetViewportBackBuffer(FRHIViewport* Viewport)
 {
 	return GDynamicRHI->RHIGetViewportBackBuffer(Viewport);
 }
diff --git a/Engine/Source/Runtime/RHI/Public/RHICommandList.h b/Engine/Source/Runtime/RHI/Public/RHICommandList.h
--- a/Engine/Source/Runtime/RHI/Public/RHICommandList.h
+++ b/Engine/Source/Runtime/RHI/Public/RHICommandList.h
@@ -664,7 +664,7 @@ public:
 		GDynamicRHI->RHIUpdateUniformBuffer(*this, UniformBufferRHI, Contents);
 	}
 
-	FORCEINLINE void UpdateTexture2D(FRHITexture2D* Texture, uint32 MipIndex, const struct FUpdateTextureRegion2D& UpdateRegion, uint32 SourcePitch, const uint8* SourceData)
+	FORCEINLINE void UpdateTexture2D(FRHITexture* Texture, uint32 MipIndex, const struct FUpdateTextureRegion2D& UpdateRegion, uint32 SourcePitch, const uint8* SourceData)
 	{
 		checkf(UpdateRegion.DestX + UpdateRegion.Width <= Texture->GetSizeX(), TEXT("UpdateTexture2D out of bounds on X. Texture: %s, %i, %i, %i"), *Texture->GetName().ToString(), UpdateRegion.DestX, UpdateRegion.Width, Texture->GetSizeX());
 		checkf(UpdateRegion.DestY + UpdateRegion.Height <= Texture->GetSizeY(), TEXT("UpdateTexture2D out of bounds on Y. Texture: %s, %i, %i, %i"), *Texture->GetName().ToString(), UpdateRegion.DestY, UpdateRegion.Height, Texture->GetSizeY());
@@ -690,7 +690,7 @@ public:
 		return GDynamicRHI->RHICreateTexture(*this, CreateDesc);
 	}
 
-	FORCEINLINE void UpdateFromBufferTexture2D(FRHITexture2D* Texture, uint32 MipIndex, const struct FUpdateTextureRegion2D& UpdateRegion, uint32 SourcePitch, FRHIBuffer* Buffer, uint32 BufferOffset)
+	FORCEINLINE void UpdateFromBufferTexture2D(FRHITexture* Texture, uint32 MipIndex, const struct FUpdateTextureRegion2D& UpdateRegion, uint32 SourcePitch, FRHIBuffer* Buffer, uint32 BufferOffset)
 	{
 		checkf(UpdateRegion.DestX + UpdateRegion.Width <= Texture->GetSizeX(), TEXT("UpdateFromBufferTexture2D out of bounds on X. Texture: %s, %i, %i, %i"), *Texture->GetName().ToString(), UpdateRegion.DestX, UpdateRegion.Width, Texture->GetSizeX());
 		checkf(UpdateRegion.DestY + UpdateRegion.Height <= Texture->GetSizeY(), TEXT("UpdateFromBufferTexture2D out of bounds on Y. Texture: %s, %i, %i, %i"), *Texture->GetName().ToString(), UpdateRegion.DestY, UpdateRegion.Height, Texture->GetSizeY());
@@ -700,7 +700,7 @@ public:
 		GDynamicRHI->RHIUpdateFromBufferTexture2D(*this, Texture, MipIndex, UpdateRegion, SourcePitch, Buffer, BufferOffset);
 	}
 
-	FORCEINLINE void UpdateTexture3D(FRHITexture3D* Texture, uint32 MipIndex, const struct FUpdateTextureRegion3D& UpdateRegion, uint32 SourceRowPitch, uint32 SourceDepthPitch, const uint8* SourceData)
+	FORCEINLINE void UpdateTexture3D(FRHITexture* Texture, uint32 MipIndex, const struct FUpdateTextureRegion3D& UpdateRegion, uint32 SourceRowPitch, uint32 SourceDepthPitch, const uint8* SourceData)
 	{
 		checkf(UpdateRegion.DestX + UpdateRegion.Width <= Texture->GetSizeX(), TEXT("UpdateTexture3D out of bounds on X. Texture: %s, %i, %i, %i"), *Texture->GetName().ToString(), UpdateRegion.DestX, UpdateRegion.Width, Texture->GetSizeX());
 		checkf(UpdateRegion.DestY + UpdateRegion.Height <= Texture->GetSizeY(), TEXT("UpdateTexture3D out of bounds on Y. Texture: %s, %i, %i, %i"), *Texture->GetName().ToString(), UpdateRegion.DestY, UpdateRegion.Height, Texture->GetSizeY());
@@ -870,7 +870,7 @@ public:
 	}
 
 	//UE_DEPRECATED(5.3, "Use the CreateShaderResourceView function that takes an FRHIViewDesc.")
-	FORCEINLINE FShaderResourceViewRHIRef CreateShaderResourceViewWriteMask(FRHITexture2D* Texture2DRHI)
+	FORCEINLINE FShaderResourceViewRHIRef CreateShaderResourceViewWriteMask(FRHITexture* Texture2DRHI)
 	{
 		return CreateShaderResourceView(Texture2DRHI, FRHIViewDesc::CreateTextureSRV()
 			.SetDimensionFromTexture(Texture2DRHI)
@@ -879,7 +879,7 @@ public:
 	}
 
 	//UE_DEPRECATED(5.3, "Use the CreateShaderResourceView function that takes an FRHIViewDesc.")
-	FORCEINLINE FShaderResourceViewRHIRef CreateShaderResourceViewFMask(FRHITexture2D* Texture2DRHI)
+	FORCEINLINE FShaderResourceViewRHIRef CreateShaderResourceViewFMask(FRHITexture* Texture2DRHI)
 	{
 		return CreateShaderResourceView(Texture2DRHI, FRHIViewDesc::CreateTextureSRV()
 			.SetDimensionFromTexture(Texture2DRHI)
@@ -1856,9 +1856,9 @@ FRHICOMMAND_MACRO(FRHICommandCopyTexture)
 
 FRHICOMMAND_MACRO(FRHICommandResummarizeHTile)
 {
-	FRHITexture2D* DepthTexture;
+	FRHITexture* DepthTexture;
 
-	FORCEINLINE_DEBUGGABLE FRHICommandResummarizeHTile(FRHITexture2D* InDepthTexture)
+	FORCEINLINE_DEBUGGABLE FRHICommandResummarizeHTile(FRHITexture* InDepthTexture)
 	: DepthTexture(InDepthTexture)
 	{
 	}
@@ -3304,7 +3304,7 @@ public:
 		{
 			GraphicsPSOInit.DepthStencilTargetFormat = PersistentState.CachedDepthStencilTarget.Texture->GetFormat();
 			GraphicsPSOInit.DepthStencilTargetFlag = PersistentState.CachedDepthStencilTarget.Texture->GetFlags();
-			const FRHITexture2DArray* TextureArray = PersistentState.CachedDepthStencilTarget.Texture->GetTexture2DArray();
+			const FRHITexture* TextureArray = PersistentState.CachedDepthStencilTarget.Texture->GetTexture2DArray();
 		}
 		else
 		{
@@ -3457,7 +3457,7 @@ public:
 		ALLOC_COMMAND(FRHICommandCopyTexture)(SourceTextureRHI, DestTextureRHI, CopyInfo);
 	}
 
-	FORCEINLINE_DEBUGGABLE void ResummarizeHTile(FRHITexture2D* DepthTexture)
+	FORCEINLINE_DEBUGGABLE void ResummarizeHTile(FRHITexture* DepthTexture)
 	{
 		if (Bypass())
 		{
@@ -4037,47 +4037,47 @@ public:
 		return GDynamicRHI->RHIGenerateMips(Texture);
 	}
 	
-	FORCEINLINE FTexture2DRHIRef AsyncReallocateTexture2D(FRHITexture2D* Texture2D, int32 NewMipCount, int32 NewSizeX, int32 NewSizeY, FThreadSafeCounter* RequestStatus)
+	FORCEINLINE FTextureRHIRef AsyncReallocateTexture2D(FRHITexture* Texture2D, int32 NewMipCount, int32 NewSizeX, int32 NewSizeY, FThreadSafeCounter* RequestStatus)
 	{
 		LLM_SCOPE(ELLMTag::Textures);
 		return GDynamicRHI->AsyncReallocateTexture2D_RenderThread(*this, Texture2D, NewMipCount, NewSizeX, NewSizeY, RequestStatus);
 	}
 	
-	FORCEINLINE ETextureReallocationStatus FinalizeAsyncReallocateTexture2D(FRHITexture2D* Texture2D, bool bBlockUntilCompleted)
+	FORCEINLINE ETextureReallocationStatus FinalizeAsyncReallocateTexture2D(FRHITexture* Texture2D, bool bBlockUntilCompleted)
 	{
 		LLM_SCOPE(ELLMTag::Textures);
 		return GDynamicRHI->FinalizeAsyncReallocateTexture2D_RenderThread(*this, Texture2D, bBlockUntilCompleted);
 	}
 	
-	FORCEINLINE ETextureReallocationStatus CancelAsyncReallocateTexture2D(FRHITexture2D* Texture2D, bool bBlockUntilCompleted)
+	FORCEINLINE ETextureReallocationStatus CancelAsyncReallocateTexture2D(FRHITexture* Texture2D, bool bBlockUntilCompleted)
 	{
 		return GDynamicRHI->CancelAsyncReallocateTexture2D_RenderThread(*this, Texture2D, bBlockUntilCompleted);
 	}
 	
-	FORCEINLINE void* LockTexture2D(FRHITexture2D* Texture, uint32 MipIndex, EResourceLockMode LockMode, uint32& DestStride, bool bLockWithinMiptail, bool bFlushRHIThread = true, uint64* OutLockedByteCount = nullptr)
+	FORCEINLINE void* LockTexture2D(FRHITexture* Texture, uint32 MipIndex, EResourceLockMode LockMode, uint32& DestStride, bool bLockWithinMiptail, bool bFlushRHIThread = true, uint64* OutLockedByteCount = nullptr)
 	{
 		LLM_SCOPE(ELLMTag::Textures);
 		return GDynamicRHI->LockTexture2D_RenderThread(*this, Texture, MipIndex, LockMode, DestStride, bLockWithinMiptail, bFlushRHIThread, OutLockedByteCount);
 	}
 	
-	FORCEINLINE void UnlockTexture2D(FRHITexture2D* Texture, uint32 MipIndex, bool bLockWithinMiptail, bool bFlushRHIThread = true)
+	FORCEINLINE void UnlockTexture2D(FRHITexture* Texture, uint32 MipIndex, bool bLockWithinMiptail, bool bFlushRHIThread = true)
 	{		
 		GDynamicRHI->UnlockTexture2D_RenderThread(*this, Texture, MipIndex, bLockWithinMiptail, bFlushRHIThread);
 	}
 	
-	FORCEINLINE void* LockTexture2DArray(FRHITexture2DArray* Texture, uint32 TextureIndex, uint32 MipIndex, EResourceLockMode LockMode, uint32& DestStride, bool bLockWithinMiptail)
+	FORCEINLINE void* LockTexture2DArray(FRHITexture* Texture, uint32 TextureIndex, uint32 MipIndex, EResourceLockMode LockMode, uint32& DestStride, bool bLockWithinMiptail)
 	{
 		LLM_SCOPE(ELLMTag::Textures);
 		return GDynamicRHI->LockTexture2DArray_RenderThread(*this, Texture, TextureIndex, MipIndex, LockMode, DestStride, bLockWithinMiptail);
 	}
 	
-	FORCEINLINE void UnlockTexture2DArray(FRHITexture2DArray* Texture, uint32 TextureIndex, uint32 MipIndex, bool bLockWithinMiptail)
+	FORCEINLINE void UnlockTexture2DArray(FRHITexture* Texture, uint32 TextureIndex, uint32 MipIndex, bool bLockWithinMiptail)
 	{
 		LLM_SCOPE(ELLMTag::Textures);
 		GDynamicRHI->UnlockTexture2DArray_RenderThread(*this, Texture, TextureIndex, MipIndex, bLockWithinMiptail);
 	}
 	
-	FORCEINLINE FUpdateTexture3DData BeginUpdateTexture3D(FRHITexture3D* Texture, uint32 MipIndex, const struct FUpdateTextureRegion3D& UpdateRegion)
+	FORCEINLINE FUpdateTexture3DData BeginUpdateTexture3D(FRHITexture* Texture, uint32 MipIndex, const struct FUpdateTextureRegion3D& UpdateRegion)
 	{
 		checkf(UpdateRegion.DestX + UpdateRegion.Width <= Texture->GetSizeX(), TEXT("UpdateTexture3D out of bounds on X. Texture: %s, %i, %i, %i"), *Texture->GetName().ToString(), UpdateRegion.DestX, UpdateRegion.Width, Texture->GetSizeX());
 		checkf(UpdateRegion.DestY + UpdateRegion.Height <= Texture->GetSizeY(), TEXT("UpdateTexture3D out of bounds on Y. Texture: %s, %i, %i, %i"), *Texture->GetName().ToString(), UpdateRegion.DestY, UpdateRegion.Height, Texture->GetSizeY());
@@ -4226,12 +4226,12 @@ public:
 		return GDynamicRHI->RHIIsRenderingSuspended();
 	}
 	
-	FORCEINLINE void VirtualTextureSetFirstMipInMemory(FRHITexture2D* Texture, uint32 FirstMip)
+	FORCEINLINE void VirtualTextureSetFirstMipInMemory(FRHITexture* Texture, uint32 FirstMip)
 	{
 		GDynamicRHI->RHIVirtualTextureSetFirstMipInMemory(*this, Texture, FirstMip);
 	}
 	
-	FORCEINLINE void VirtualTextureSetFirstMipVisible(FRHITexture2D* Texture, uint32 FirstMip)
+	FORCEINLINE void VirtualTextureSetFirstMipVisible(FRHITexture* Texture, uint32 FirstMip)
 	{
 		GDynamicRHI->RHIVirtualTextureSetFirstMipVisible(*this, Texture, FirstMip);
 	}
@@ -4920,14 +4920,14 @@ FORCEINLINE FTextureRHIRef RHICreateTexture(const FRHITextureCreateDesc& CreateD
 }
 
 UE_DEPRECATED(5.4, "Use the RHIAsyncCreateTexture2D function that takes an DebugName.")
-FORCEINLINE FTexture2DRHIRef RHIAsyncCreateTexture2D(uint32 SizeX, uint32 SizeY, uint8 Format, uint32 NumMips, ETextureCreateFlags Flags, ERHIAccess InResourceState, void** InitialMipData, uint32 NumInitialMips, FGraphEventRef& OutCompletionEvent)
+FORCEINLINE FTextureRHIRef RHIAsyncCreateTexture2D(uint32 SizeX, uint32 SizeY, uint8 Format, uint32 NumMips, ETextureCreateFlags Flags, ERHIAccess InResourceState, void** InitialMipData, uint32 NumInitialMips, FGraphEventRef& OutCompletionEvent)
 {
 	LLM_SCOPE(EnumHasAnyFlags(Flags, TexCreate_RenderTargetable | TexCreate_DepthStencilTargetable) ? ELLMTag::RenderTargets : ELLMTag::Textures);
 	const ERHIAccess ResourceState = InResourceState == ERHIAccess::Unknown ? RHIGetDefaultResourceState((ETextureCreateFlags)Flags, InitialMipData != nullptr) : InResourceState;
 	return GDynamicRHI->RHIAsyncCreateTexture2D(SizeX, SizeY, Format, NumMips, Flags, ResourceState, InitialMipData, NumInitialMips, TEXT("RHIAsyncCreateTexture2D"), OutCompletionEvent);
 }
 
-FORCEINLINE FTexture2DRHIRef RHIAsyncCreateTexture2D(uint32 SizeX, uint32 SizeY, uint8 Format, uint32 NumMips, ETextureCreateFlags Flags, ERHIAccess InResourceState, void** InitialMipData, uint32 NumInitialMips, const TCHAR* DebugName, FGraphEventRef& OutCompletionEvent)
+FORCEINLINE FTextureRHIRef RHIAsyncCreateTexture2D(uint32 SizeX, uint32 SizeY, uint8 Format, uint32 NumMips, ETextureCreateFlags Flags, ERHIAccess InResourceState, void** InitialMipData, uint32 NumInitialMips, const TCHAR* DebugName, FGraphEventRef& OutCompletionEvent)
 {
 	LLM_SCOPE(EnumHasAnyFlags(Flags, TexCreate_RenderTargetable | TexCreate_DepthStencilTargetable) ? ELLMTag::RenderTargets : ELLMTag::Textures);
 	const ERHIAccess ResourceState = InResourceState == ERHIAccess::Unknown ? RHIGetDefaultResourceState((ETextureCreateFlags)Flags, InitialMipData != nullptr) : InResourceState;
@@ -4935,54 +4935,54 @@ FORCEINLINE FTexture2DRHIRef RHIAsyncCreateTexture2D(uint32 SizeX, uint32 SizeY,
 }
 
 UE_DEPRECATED(5.4, "Use the RHIAsyncCreateTexture2D function that takes an InResourceState and DebugName.")
-FORCEINLINE FTexture2DRHIRef RHIAsyncCreateTexture2D(uint32 SizeX, uint32 SizeY, uint8 Format, uint32 NumMips, ETextureCreateFlags Flags, void** InitialMipData, uint32 NumInitialMips, FGraphEventRef& OutCompletionEvent)
+FORCEINLINE FTextureRHIRef RHIAsyncCreateTexture2D(uint32 SizeX, uint32 SizeY, uint8 Format, uint32 NumMips, ETextureCreateFlags Flags, void** InitialMipData, uint32 NumInitialMips, FGraphEventRef& OutCompletionEvent)
 {
 PRAGMA_DISABLE_DEPRECATION_WARNINGS
 	return RHIAsyncCreateTexture2D(SizeX, SizeY, Format, NumMips, Flags, ERHIAccess::Unknown, InitialMipData, NumInitialMips, OutCompletionEvent);
 PRAGMA_ENABLE_DEPRECATION_WARNINGS
 }
 
-FORCEINLINE FTexture2DRHIRef RHIAsyncReallocateTexture2D(FRHITexture2D* Texture2D, int32 NewMipCount, int32 NewSizeX, int32 NewSizeY, FThreadSafeCounter* RequestStatus)
+FORCEINLINE FTextureRHIRef RHIAsyncReallocateTexture2D(FRHITexture* Texture2D, int32 NewMipCount, int32 NewSizeX, int32 NewSizeY, FThreadSafeCounter* RequestStatus)
 {
 	return FRHICommandListExecutor::GetImmediateCommandList().AsyncReallocateTexture2D(Texture2D, NewMipCount, NewSizeX, NewSizeY, RequestStatus);
 }
 
-FORCEINLINE ETextureReallocationStatus RHIFinalizeAsyncReallocateTexture2D(FRHITexture2D* Texture2D, bool bBlockUntilCompleted)
+FORCEINLINE ETextureReallocationStatus RHIFinalizeAsyncReallocateTexture2D(FRHITexture* Texture2D, bool bBlockUntilCompleted)
 {
 	return FRHICommandListExecutor::GetImmediateCommandList().FinalizeAsyncReallocateTexture2D(Texture2D, bBlockUntilCompleted);
 }
 
-FORCEINLINE ETextureReallocationStatus RHICancelAsyncReallocateTexture2D(FRHITexture2D* Texture2D, bool bBlockUntilCompleted)
+FORCEINLINE ETextureReallocationStatus RHICancelAsyncReallocateTexture2D(FRHITexture* Texture2D, bool bBlockUntilCompleted)
 {
 	return FRHICommandListExecutor::GetImmediateCommandList().CancelAsyncReallocateTexture2D(Texture2D, bBlockUntilCompleted);
 }
 
-FORCEINLINE void* RHILockTexture2D(FRHITexture2D* Texture, uint32 MipIndex, EResourceLockMode LockMode, uint32& DestStride, bool bLockWithinMiptail, bool bFlushRHIThread = true, uint64* OutLockedByteCount = nullptr)
+FORCEINLINE void* RHILockTexture2D(FRHITexture* Texture, uint32 MipIndex, EResourceLockMode LockMode, uint32& DestStride, bool bLockWithinMiptail, bool bFlushRHIThread = true, uint64* OutLockedByteCount = nullptr)
 {
 	return FRHICommandListExecutor::GetImmediateCommandList().LockTexture2D(Texture, MipIndex, LockMode, DestStride, bLockWithinMiptail, bFlushRHIThread, OutLockedByteCount);
 }
 
-FORCEINLINE void RHIUnlockTexture2D(FRHITexture2D* Texture, uint32 MipIndex, bool bLockWithinMiptail, bool bFlushRHIThread = true)
+FORCEINLINE void RHIUnlockTexture2D(FRHITexture* Texture, uint32 MipIndex, bool bLockWithinMiptail, bool bFlushRHIThread = true)
 {
 	 FRHICommandListExecutor::GetImmediateCommandList().UnlockTexture2D(Texture, MipIndex, bLockWithinMiptail, bFlushRHIThread);
 }
 
-FORCEINLINE void* RHILockTexture2DArray(FRHITexture2DArray* Texture, uint32 TextureIndex, uint32 MipIndex, EResourceLockMode LockMode, uint32& DestStride, bool bLockWithinMiptail)
+FORCEINLINE void* RHILockTexture2DArray(FRHITexture* Texture, uint32 TextureIndex, uint32 MipIndex, EResourceLockMode LockMode, uint32& DestStride, bool bLockWithinMiptail)
 {
 	return FRHICommandListExecutor::GetImmediateCommandList().LockTexture2DArray(Texture, TextureIndex, MipIndex, LockMode, DestStride, bLockWithinMiptail);
 }
 
-FORCEINLINE void RHIUnlockTexture2DArray(FRHITexture2DArray* Texture, uint32 TextureIndex, uint32 MipIndex, bool bLockWithinMiptail)
+FORCEINLINE void RHIUnlockTexture2DArray(FRHITexture* Texture, uint32 TextureIndex, uint32 MipIndex, bool bLockWithinMiptail)
 {
 	 FRHICommandListExecutor::GetImmediateCommandList().UnlockTexture2DArray(Texture, TextureIndex, MipIndex, bLockWithinMiptail);
 }
 
-FORCEINLINE void RHIUpdateTexture2D(FRHITexture2D* Texture, uint32 MipIndex, const struct FUpdateTextureRegion2D& UpdateRegion, uint32 SourcePitch, const uint8* SourceData)
+FORCEINLINE void RHIUpdateTexture2D(FRHITexture* Texture, uint32 MipIndex, const struct FUpdateTextureRegion2D& UpdateRegion, uint32 SourcePitch, const uint8* SourceData)
 {
 	 FRHICommandListExecutor::GetImmediateCommandList().UpdateTexture2D(Texture, MipIndex, UpdateRegion, SourcePitch, SourceData);
 }
 
-FORCEINLINE FUpdateTexture3DData RHIBeginUpdateTexture3D(FRHITexture3D* Texture, uint32 MipIndex, const struct FUpdateTextureRegion3D& UpdateRegion)
+FORCEINLINE FUpdateTexture3DData RHIBeginUpdateTexture3D(FRHITexture* Texture, uint32 MipIndex, const struct FUpdateTextureRegion3D& UpdateRegion)
 {
 	return FRHICommandListExecutor::GetImmediateCommandList().BeginUpdateTexture3D(Texture, MipIndex, UpdateRegion);
 }
@@ -4997,17 +4997,17 @@ FORCEINLINE void RHIEndMultiUpdateTexture3D(TArray<FUpdateTexture3DData>& Update
 	FRHICommandListExecutor::GetImmediateCommandList().EndMultiUpdateTexture3D(UpdateDataArray);
 }
 
-FORCEINLINE void RHIUpdateTexture3D(FRHITexture3D* Texture, uint32 MipIndex, const struct FUpdateTextureRegion3D& UpdateRegion, uint32 SourceRowPitch, uint32 SourceDepthPitch, const uint8* SourceData)
+FORCEINLINE void RHIUpdateTexture3D(FRHITexture* Texture, uint32 MipIndex, const struct FUpdateTextureRegion3D& UpdateRegion, uint32 SourceRowPitch, uint32 SourceDepthPitch, const uint8* SourceData)
 {
 	 FRHICommandListExecutor::GetImmediateCommandList().UpdateTexture3D(Texture, MipIndex, UpdateRegion, SourceRowPitch, SourceDepthPitch, SourceData);
 }
 
-FORCEINLINE void* RHILockTextureCubeFace(FRHITextureCube* Texture, uint32 FaceIndex, uint32 ArrayIndex, uint32 MipIndex, EResourceLockMode LockMode, uint32& DestStride, bool bLockWithinMiptail)
+FORCEINLINE void* RHILockTextureCubeFace(FRHITexture* Texture, uint32 FaceIndex, uint32 ArrayIndex, uint32 MipIndex, EResourceLockMode LockMode, uint32& DestStride, bool bLockWithinMiptail)
 {
 	return FRHICommandListExecutor::GetImmediateCommandList().LockTextureCubeFace(Texture, FaceIndex, ArrayIndex, MipIndex, LockMode, DestStride, bLockWithinMiptail);
 }
 
-FORCEINLINE void RHIUnlockTextureCubeFace(FRHITextureCube* Texture, uint32 FaceIndex, uint32 ArrayIndex, uint32 MipIndex, bool bLockWithinMiptail)
+FORCEINLINE void RHIUnlockTextureCubeFace(FRHITexture* Texture, uint32 FaceIndex, uint32 ArrayIndex, uint32 MipIndex, bool bLockWithinMiptail)
 {
 	 FRHICommandListExecutor::GetImmediateCommandList().UnlockTextureCubeFace(Texture, FaceIndex, ArrayIndex, MipIndex, bLockWithinMiptail);
 }
@@ -5025,12 +5025,12 @@ FORCEINLINE void RHIFlushResources()
 	return FRHICommandListExecutor::GetImmediateCommandList().FlushResources();
 }
 
-FORCEINLINE void RHIVirtualTextureSetFirstMipInMemory(FRHITexture2D* Texture, uint32 FirstMip)
+FORCEINLINE void RHIVirtualTextureSetFirstMipInMemory(FRHITexture* Texture, uint32 FirstMip)
 {
 	 FRHICommandListExecutor::GetImmediateCommandList().VirtualTextureSetFirstMipInMemory(Texture, FirstMip);
 }
 
-FORCEINLINE void RHIVirtualTextureSetFirstMipVisible(FRHITexture2D* Texture, uint32 FirstMip)
+FORCEINLINE void RHIVirtualTextureSetFirstMipVisible(FRHITexture* Texture, uint32 FirstMip)
 {
 	 FRHICommandListExecutor::GetImmediateCommandList().VirtualTextureSetFirstMipVisible(Texture, FirstMip);
 }
diff --git a/Engine/Source/Runtime/RHI/Public/RHIContext.h b/Engine/Source/Runtime/RHI/Public/RHIContext.h
--- a/Engine/Source/Runtime/RHI/Public/RHIContext.h
+++ b/Engine/Source/Runtime/RHI/Public/RHIContext.h
@@ -165,7 +165,7 @@ struct FTransferResourceParams
 {
 	FTransferResourceParams() {}
 
-	FTransferResourceParams(FRHITexture2D* InTexture, const FIntRect& InRect, uint32 InSrcGPUIndex, uint32 InDestGPUIndex, bool InPullData, bool InLockStepGPUs)
+	FTransferResourceParams(FRHITexture* InTexture, const FIntRect& InRect, uint32 InSrcGPUIndex, uint32 InDestGPUIndex, bool InPullData, bool InLockStepGPUs)
 		: Texture(InTexture), Buffer(nullptr), Min(InRect.Min.X, InRect.Min.Y, 0), Max(InRect.Max.X, InRect.Max.Y, 1), SrcGPUIndex(InSrcGPUIndex), DestGPUIndex(InDestGPUIndex), bPullData(InPullData), bLockStepGPUs(InLockStepGPUs)
 	{
 		check(InTexture);
@@ -608,7 +608,7 @@ public:
 	* Rebuilds the depth target HTILE meta data (on supported platforms).
 	* @param DepthTexture - the depth surface to resummarize.
 	*/
-	virtual void RHIResummarizeHTile(FRHITexture2D* DepthTexture)
+	virtual void RHIResummarizeHTile(FRHITexture* DepthTexture)
 	{
 		/* empty default implementation */
 	}
diff --git a/Engine/Source/Runtime/RHI/Public/RHIFwd.h b/Engine/Source/Runtime/RHI/Public/RHIFwd.h
--- a/Engine/Source/Runtime/RHI/Public/RHIFwd.h
+++ b/Engine/Source/Runtime/RHI/Public/RHIFwd.h
@@ -126,11 +126,11 @@ using FVertexShaderRHIRef              = TRefCountPtr<FRHIVertexShader>;
 using FViewportRHIRef                  = TRefCountPtr<FRHIViewport>;
 
 // Deprecated typenames
-using FRHITexture2D                    = FRHITexture;
-using FRHITexture2DArray               = FRHITexture;
-using FRHITexture3D                    = FRHITexture;
-using FRHITextureCube                  = FRHITexture;
-using FTexture2DRHIRef                 = FTextureRHIRef;
-using FTexture2DArrayRHIRef            = FTextureRHIRef;
-using FTexture3DRHIRef                 = FTextureRHIRef;
-using FTextureCubeRHIRef               = FTextureRHIRef;
+using FRHITexture2D         UE_DEPRECATED(5.5, "The separate RHI texture types were unified in UE 5.1. FRHITexture2D is now deprecated. Use FRHITexture instead."           ) = FRHITexture;
+using FRHITexture2DArray    UE_DEPRECATED(5.5, "The separate RHI texture types were unified in UE 5.1. FRHITexture2DArray is now deprecated. Use FRHITexture instead."      ) = FRHITexture;
+using FRHITexture3D         UE_DEPRECATED(5.5, "The separate RHI texture types were unified in UE 5.1. FRHITexture3D is now deprecated. Use FRHITexture instead."           ) = FRHITexture;
+using FRHITextureCube       UE_DEPRECATED(5.5, "The separate RHI texture types were unified in UE 5.1. FRHITextureCube is now deprecated. Use FRHITexture instead."         ) = FRHITexture;
+using FTexture2DRHIRef      UE_DEPRECATED(5.5, "The separate RHI texture types were unified in UE 5.1. FTexture2DRHIRef is now deprecated. Use FTextureRHIRef instead."     ) = FTextureRHIRef;
+using FTexture2DArrayRHIRef UE_DEPRECATED(5.5, "The separate RHI texture types were unified in UE 5.1. FTexture2DArrayRHIRef is now deprecated. Use FTextureRHIRef instead.") = FTextureRHIRef;
+using FTexture3DRHIRef      UE_DEPRECATED(5.5, "The separate RHI texture types were unified in UE 5.1. FTexture3DRHIRef is now deprecated. Use FTextureRHIRef instead."     ) = FTextureRHIRef;
+using FTextureCubeRHIRef    UE_DEPRECATED(5.5, "The separate RHI texture types were unified in UE 5.1. FTextureCubeRHIRef is now deprecated. Use FTextureRHIRef instead."   ) = FTextureRHIRef;
diff --git a/Engine/Source/Runtime/RHI/Public/RHIResources.h b/Engine/Source/Runtime/RHI/Public/RHIResources.h
--- a/Engine/Source/Runtime/RHI/Public/RHIResources.h
+++ b/Engine/Source/Runtime/RHI/Public/RHIResources.h
@@ -1828,13 +1828,13 @@ public:
 	/// 
 
 	//UE_DEPRECATED(5.1, "FRHITexture2D is deprecated, please use FRHITexture directly")
-	inline FRHITexture2D* GetTexture2D() { return TextureDesc.Dimension == ETextureDimension::Texture2D ? this : nullptr; }
+	inline FRHITexture* GetTexture2D() { return TextureDesc.Dimension == ETextureDimension::Texture2D ? this : nullptr; }
 	//UE_DEPRECATED(5.1, "FRHITexture2DArray is deprecated, please use FRHITexture directly")
-	inline FRHITexture2DArray* GetTexture2DArray() { return TextureDesc.Dimension == ETextureDimension::Texture2DArray ? this : nullptr; }
+	inline FRHITexture* GetTexture2DArray() { return TextureDesc.Dimension == ETextureDimension::Texture2DArray ? this : nullptr; }
 	//UE_DEPRECATED(5.1, "FRHITexture3D is deprecated, please use FRHITexture directly")
-	inline FRHITexture3D* GetTexture3D() { return TextureDesc.Dimension == ETextureDimension::Texture3D ? this : nullptr; }
+	inline FRHITexture* GetTexture3D() { return TextureDesc.Dimension == ETextureDimension::Texture3D ? this : nullptr; }
 	//UE_DEPRECATED(5.1, "FRHITextureCube is deprecated, please use FRHITexture directly")
-	inline FRHITextureCube* GetTextureCube() { return TextureDesc.IsTextureCube() ? this : nullptr; }
+	inline FRHITexture* GetTextureCube() { return TextureDesc.IsTextureCube() ? this : nullptr; }
 
 	//UE_DEPRECATED(5.1, "GetSizeX() is deprecated, please use GetDesc().Extent.X instead")
 	uint32 GetSizeX() const { return GetDesc().Extent.X; }
diff --git a/Engine/Source/Runtime/RHI/Public/RHIValidation.h b/Engine/Source/Runtime/RHI/Public/RHIValidation.h
--- a/Engine/Source/Runtime/RHI/Public/RHIValidation.h
+++ b/Engine/Source/Runtime/RHI/Public/RHIValidation.h
@@ -582,12 +582,12 @@ public:
 	*/
 	// FlushType: Flush RHI Thread
 	// NP: Note that no RHI currently implements this as an async call, we should simplify the API.
-	virtual FTexture2DRHIRef RHIAsyncReallocateTexture2D(FRHITexture2D* Texture2D, int32 NewMipCount, int32 NewSizeX, int32 NewSizeY, FThreadSafeCounter* RequestStatus) override final
+	virtual FTextureRHIRef RHIAsyncReallocateTexture2D(FRHITexture* Texture2D, int32 NewMipCount, int32 NewSizeX, int32 NewSizeY, FThreadSafeCounter* RequestStatus) override final
 	{
 		// TODO: find proper state for new texture
 		ERHIAccess ResourceState = ERHIAccess::SRVMask;
 
-		FTexture2DRHIRef NewTexture2D = RHI->RHIAsyncReallocateTexture2D(Texture2D, NewMipCount, NewSizeX, NewSizeY, RequestStatus);
+		FTextureRHIRef NewTexture2D = RHI->RHIAsyncReallocateTexture2D(Texture2D, NewMipCount, NewSizeX, NewSizeY, RequestStatus);
 		NewTexture2D->InitBarrierTracking(NewMipCount, 1, NewTexture2D->GetFormat(), NewTexture2D->GetFlags(), ResourceState, NewTexture2D->GetTrackerResource()->GetDebugName()); // @todo the threading of GetDebugName() is wrong.
 		return NewTexture2D;
 	}
@@ -604,7 +604,7 @@ public:
 	*	TexRealloc_InProgress	Reallocation is still in progress, try again later
 	*/
 	// FlushType: Wait RHI Thread
-	virtual ETextureReallocationStatus RHIFinalizeAsyncReallocateTexture2D(FRHITexture2D* Texture2D, bool bBlockUntilCompleted) override final
+	virtual ETextureReallocationStatus RHIFinalizeAsyncReallocateTexture2D(FRHITexture* Texture2D, bool bBlockUntilCompleted) override final
 	{
 		return RHI->RHIFinalizeAsyncReallocateTexture2D(Texture2D, bBlockUntilCompleted);
 	}
@@ -618,7 +618,7 @@ public:
 	* @return						Reallocation status
 	*/
 	// FlushType: Wait RHI Thread
-	virtual ETextureReallocationStatus RHICancelAsyncReallocateTexture2D(FRHITexture2D* Texture2D, bool bBlockUntilCompleted) override final
+	virtual ETextureReallocationStatus RHICancelAsyncReallocateTexture2D(FRHITexture* Texture2D, bool bBlockUntilCompleted) override final
 	{
 		return RHI->RHICancelAsyncReallocateTexture2D(Texture2D, bBlockUntilCompleted);
 	}
@@ -633,7 +633,7 @@ public:
 	* @return pointer to the CPU accessible resource data
 	*/
 	// FlushType: Flush RHI Thread
-	virtual void* RHILockTexture2D(FRHITexture2D* Texture, uint32 MipIndex, EResourceLockMode LockMode, uint32& DestStride, bool bLockWithinMiptail, uint64* OutLockedByteCount) override final
+	virtual void* RHILockTexture2D(FRHITexture* Texture, uint32 MipIndex, EResourceLockMode LockMode, uint32& DestStride, bool bLockWithinMiptail, uint64* OutLockedByteCount) override final
 	{
 		return RHI->RHILockTexture2D(Texture, MipIndex, LockMode, DestStride, bLockWithinMiptail, OutLockedByteCount);
 	}
@@ -645,7 +645,7 @@ public:
 	* @param bLockWithinMiptail - for platforms that support packed miptails allow locking of individual mip levels within the miptail
 	*/
 	// FlushType: Flush RHI Thread
-	virtual void RHIUnlockTexture2D(FRHITexture2D* Texture, uint32 MipIndex, bool bLockWithinMiptail) override final
+	virtual void RHIUnlockTexture2D(FRHITexture* Texture, uint32 MipIndex, bool bLockWithinMiptail) override final
 	{
 		RHI->RHIUnlockTexture2D(Texture, MipIndex, bLockWithinMiptail);
 	}
@@ -660,7 +660,7 @@ public:
 	* @return pointer to the CPU accessible resource data
 	*/
 	// FlushType: Flush RHI Thread
-	virtual void* RHILockTexture2DArray(FRHITexture2DArray* Texture, uint32 TextureIndex, uint32 MipIndex, EResourceLockMode LockMode, uint32& DestStride, bool bLockWithinMiptail) override final
+	virtual void* RHILockTexture2DArray(FRHITexture* Texture, uint32 TextureIndex, uint32 MipIndex, EResourceLockMode LockMode, uint32& DestStride, bool bLockWithinMiptail) override final
 	{
 		return RHI->RHILockTexture2DArray(Texture, TextureIndex, MipIndex, LockMode, DestStride, bLockWithinMiptail);
 	}
@@ -672,7 +672,7 @@ public:
 	* @param bLockWithinMiptail - for platforms that support packed miptails allow locking of individual mip levels within the miptail
 	*/
 	// FlushType: Flush RHI Thread
-	virtual void RHIUnlockTexture2DArray(FRHITexture2DArray* Texture, uint32 TextureIndex, uint32 MipIndex, bool bLockWithinMiptail) override final
+	virtual void RHIUnlockTexture2DArray(FRHITexture* Texture, uint32 TextureIndex, uint32 MipIndex, bool bLockWithinMiptail) override final
 	{
 		RHI->RHIUnlockTexture2DArray(Texture, TextureIndex, MipIndex, bLockWithinMiptail);
 	}
@@ -686,12 +686,12 @@ public:
 	* @param SourceData - source image data, starting at the upper left corner of the source rectangle (in same pixel format as texture)
 	*/
 	// FlushType: Flush RHI Thread
-	virtual void RHIUpdateTexture2D(FRHICommandListBase& RHICmdList, FRHITexture2D* Texture, uint32 MipIndex, const struct FUpdateTextureRegion2D& UpdateRegion, uint32 SourcePitch, const uint8* SourceData) override final
+	virtual void RHIUpdateTexture2D(FRHICommandListBase& RHICmdList, FRHITexture* Texture, uint32 MipIndex, const struct FUpdateTextureRegion2D& UpdateRegion, uint32 SourcePitch, const uint8* SourceData) override final
 	{
 		RHI->RHIUpdateTexture2D(RHICmdList, Texture, MipIndex, UpdateRegion, SourcePitch, SourceData);
 	}
 
-	virtual void RHIUpdateFromBufferTexture2D(FRHICommandListBase& RHICmdList, FRHITexture2D* Texture, uint32 MipIndex, const struct FUpdateTextureRegion2D& UpdateRegion, uint32 SourcePitch, FRHIBuffer* Buffer, uint32 BufferOffset) override final
+	virtual void RHIUpdateFromBufferTexture2D(FRHICommandListBase& RHICmdList, FRHITexture* Texture, uint32 MipIndex, const struct FUpdateTextureRegion2D& UpdateRegion, uint32 SourcePitch, FRHIBuffer* Buffer, uint32 BufferOffset) override final
 	{
 		RHI->RHIUpdateFromBufferTexture2D(RHICmdList, Texture, MipIndex, UpdateRegion, SourcePitch, Buffer, BufferOffset);
 	}
@@ -706,7 +706,7 @@ public:
 	* @param SourceData - source image data, starting at the upper left corner of the source rectangle (in same pixel format as texture)
 	*/
 	// FlushType: Flush RHI Thread
-	virtual void RHIUpdateTexture3D(FRHICommandListBase& RHICmdList, FRHITexture3D* Texture, uint32 MipIndex, const struct FUpdateTextureRegion3D& UpdateRegion, uint32 SourceRowPitch, uint32 SourceDepthPitch, const uint8* SourceData) override final
+	virtual void RHIUpdateTexture3D(FRHICommandListBase& RHICmdList, FRHITexture* Texture, uint32 MipIndex, const struct FUpdateTextureRegion3D& UpdateRegion, uint32 SourceRowPitch, uint32 SourceDepthPitch, const uint8* SourceData) override final
 	{
 		RHI->RHIUpdateTexture3D(RHICmdList, Texture, MipIndex, UpdateRegion, SourceRowPitch, SourceDepthPitch, SourceData);
 	}
@@ -721,7 +721,7 @@ public:
 	* @return pointer to the CPU accessible resource data
 	*/
 	// FlushType: Flush RHI Thread
-	virtual void* RHILockTextureCubeFace(FRHITextureCube* Texture, uint32 FaceIndex, uint32 ArrayIndex, uint32 MipIndex, EResourceLockMode LockMode, uint32& DestStride, bool bLockWithinMiptail) override final
+	virtual void* RHILockTextureCubeFace(FRHITexture* Texture, uint32 FaceIndex, uint32 ArrayIndex, uint32 MipIndex, EResourceLockMode LockMode, uint32& DestStride, bool bLockWithinMiptail) override final
 	{
 		return RHI->RHILockTextureCubeFace(Texture, FaceIndex, ArrayIndex, MipIndex, LockMode, DestStride, bLockWithinMiptail);
 	}
@@ -733,7 +733,7 @@ public:
 	* @param bLockWithinMiptail - for platforms that support packed miptails allow locking of individual mip levels within the miptail
 	*/
 	// FlushType: Flush RHI Thread
-	virtual void RHIUnlockTextureCubeFace(FRHITextureCube* Texture, uint32 FaceIndex, uint32 ArrayIndex, uint32 MipIndex, bool bLockWithinMiptail) override final
+	virtual void RHIUnlockTextureCubeFace(FRHITexture* Texture, uint32 FaceIndex, uint32 ArrayIndex, uint32 MipIndex, bool bLockWithinMiptail) override final
 	{
 		RHI->RHIUnlockTextureCubeFace(Texture, FaceIndex, ArrayIndex, MipIndex, bLockWithinMiptail);
 	}
@@ -832,9 +832,9 @@ public:
 
 	// With RHI thread, this is the current backbuffer from the perspective of the render thread.
 	// FlushType: Thread safe
-	virtual FTexture2DRHIRef RHIGetViewportBackBuffer(FRHIViewport* Viewport) override final
+	virtual FTextureRHIRef RHIGetViewportBackBuffer(FRHIViewport* Viewport) override final
 	{
-		FTexture2DRHIRef Texture = RHI->RHIGetViewportBackBuffer(Viewport);
+		FTextureRHIRef Texture = RHI->RHIGetViewportBackBuffer(Viewport);
 		if (!Texture->GetTrackerResource()->IsBarrierTrackingInitialized())
 		{
 			// Assume present and renderer needs to perform transition to RTV if needed
@@ -1025,7 +1025,7 @@ public:
 	* @param FirstMip - the first mip that should be in memory
 	*/
 	// FlushType: Wait RHI Thread
-	virtual void RHIVirtualTextureSetFirstMipInMemory(FRHICommandListImmediate& RHICmdList, FRHITexture2D* Texture, uint32 FirstMip) override final
+	virtual void RHIVirtualTextureSetFirstMipInMemory(FRHICommandListImmediate& RHICmdList, FRHITexture* Texture, uint32 FirstMip) override final
 	{
 		RHI->RHIVirtualTextureSetFirstMipInMemory(RHICmdList, Texture, FirstMip);
 	}
@@ -1036,7 +1036,7 @@ public:
 	* @param FirstMip - the first mip that should be visible to the GPU
 	*/
 	// FlushType: Wait RHI Thread
-	virtual void RHIVirtualTextureSetFirstMipVisible(FRHICommandListImmediate& RHICmdList, FRHITexture2D* Texture, uint32 FirstMip) override final
+	virtual void RHIVirtualTextureSetFirstMipVisible(FRHICommandListImmediate& RHICmdList, FRHITexture* Texture, uint32 FirstMip) override final
 	{
 		RHI->RHIVirtualTextureSetFirstMipVisible(RHICmdList, Texture, FirstMip);
 	}
@@ -1110,22 +1110,22 @@ public:
 		return RHI->RHIGetMinimumAlignmentForBufferBackedSRV(Format);
 	}
 
-	virtual FTexture2DRHIRef AsyncReallocateTexture2D_RenderThread(class FRHICommandListImmediate& RHICmdList, FRHITexture2D* Texture2D, int32 NewMipCount, int32 NewSizeX, int32 NewSizeY, FThreadSafeCounter* RequestStatus) override final
+	virtual FTextureRHIRef AsyncReallocateTexture2D_RenderThread(class FRHICommandListImmediate& RHICmdList, FRHITexture* Texture2D, int32 NewMipCount, int32 NewSizeX, int32 NewSizeY, FThreadSafeCounter* RequestStatus) override final
 	{
 		// TODO: find proper state for new texture
 		ERHIAccess ResourceState = ERHIAccess::SRVMask;
 
-		FTexture2DRHIRef NewTexture2D = RHI->AsyncReallocateTexture2D_RenderThread(RHICmdList, Texture2D, NewMipCount, NewSizeX, NewSizeY, RequestStatus);
+		FTextureRHIRef NewTexture2D = RHI->AsyncReallocateTexture2D_RenderThread(RHICmdList, Texture2D, NewMipCount, NewSizeX, NewSizeY, RequestStatus);
 		NewTexture2D->InitBarrierTracking(NewMipCount, 1, NewTexture2D->GetFormat(), NewTexture2D->GetFlags(), ResourceState, NewTexture2D->GetTrackerResource()->GetDebugName()); // @todo the threading of GetDebugName() is wrong.
 		return NewTexture2D;
 	}
 
-	virtual ETextureReallocationStatus FinalizeAsyncReallocateTexture2D_RenderThread(class FRHICommandListImmediate& RHICmdList, FRHITexture2D* Texture2D, bool bBlockUntilCompleted) override final
+	virtual ETextureReallocationStatus FinalizeAsyncReallocateTexture2D_RenderThread(class FRHICommandListImmediate& RHICmdList, FRHITexture* Texture2D, bool bBlockUntilCompleted) override final
 	{
 		return RHI->FinalizeAsyncReallocateTexture2D_RenderThread(RHICmdList, Texture2D, bBlockUntilCompleted);
 	}
 
-	virtual ETextureReallocationStatus CancelAsyncReallocateTexture2D_RenderThread(class FRHICommandListImmediate& RHICmdList, FRHITexture2D* Texture2D, bool bBlockUntilCompleted) override final
+	virtual ETextureReallocationStatus CancelAsyncReallocateTexture2D_RenderThread(class FRHICommandListImmediate& RHICmdList, FRHITexture* Texture2D, bool bBlockUntilCompleted) override final
 	{
 		return RHI->CancelAsyncReallocateTexture2D_RenderThread(RHICmdList, Texture2D, bBlockUntilCompleted);
 	}
@@ -1141,27 +1141,27 @@ public:
 		RHI->UnlockBuffer_BottomOfPipe(RHICmdList, Buffer);
 	}
 
-	virtual void* LockTexture2D_RenderThread(class FRHICommandListImmediate& RHICmdList, FRHITexture2D* Texture, uint32 MipIndex, EResourceLockMode LockMode, uint32& DestStride, bool bLockWithinMiptail, bool bNeedsDefaultRHIFlush = true, uint64* OutLockedByteCount = nullptr) override final
+	virtual void* LockTexture2D_RenderThread(class FRHICommandListImmediate& RHICmdList, FRHITexture* Texture, uint32 MipIndex, EResourceLockMode LockMode, uint32& DestStride, bool bLockWithinMiptail, bool bNeedsDefaultRHIFlush = true, uint64* OutLockedByteCount = nullptr) override final
 	{
 		return RHI->LockTexture2D_RenderThread(RHICmdList, Texture, MipIndex, LockMode, DestStride, bLockWithinMiptail, bNeedsDefaultRHIFlush, OutLockedByteCount);
 	}
 
-	virtual void UnlockTexture2D_RenderThread(class FRHICommandListImmediate& RHICmdList, FRHITexture2D* Texture, uint32 MipIndex, bool bLockWithinMiptail, bool bNeedsDefaultRHIFlush = true) override final
+	virtual void UnlockTexture2D_RenderThread(class FRHICommandListImmediate& RHICmdList, FRHITexture* Texture, uint32 MipIndex, bool bLockWithinMiptail, bool bNeedsDefaultRHIFlush = true) override final
 	{
 		RHI->UnlockTexture2D_RenderThread(RHICmdList, Texture, MipIndex, bLockWithinMiptail, bNeedsDefaultRHIFlush);
 	}
 
-	virtual void* LockTexture2DArray_RenderThread(class FRHICommandListImmediate& RHICmdList, FRHITexture2DArray* Texture, uint32 ArrayIndex, uint32 MipIndex, EResourceLockMode LockMode, uint32& DestStride, bool bLockWithinMiptail) override final
+	virtual void* LockTexture2DArray_RenderThread(class FRHICommandListImmediate& RHICmdList, FRHITexture* Texture, uint32 ArrayIndex, uint32 MipIndex, EResourceLockMode LockMode, uint32& DestStride, bool bLockWithinMiptail) override final
 	{
 		return RHI->LockTexture2DArray_RenderThread(RHICmdList, Texture, ArrayIndex, MipIndex, LockMode, DestStride, bLockWithinMiptail);
 	}
 
-	virtual void UnlockTexture2DArray_RenderThread(class FRHICommandListImmediate& RHICmdList, FRHITexture2DArray* Texture, uint32 ArrayIndex, uint32 MipIndex, bool bLockWithinMiptail) override final
+	virtual void UnlockTexture2DArray_RenderThread(class FRHICommandListImmediate& RHICmdList, FRHITexture* Texture, uint32 ArrayIndex, uint32 MipIndex, bool bLockWithinMiptail) override final
 	{
 		RHI->UnlockTexture2DArray_RenderThread(RHICmdList, Texture, ArrayIndex, MipIndex, bLockWithinMiptail);
 	}
 
-	virtual FUpdateTexture3DData RHIBeginUpdateTexture3D(FRHICommandListBase& RHICmdList, FRHITexture3D* Texture, uint32 MipIndex, const struct FUpdateTextureRegion3D& UpdateRegion) override final
+	virtual FUpdateTexture3DData RHIBeginUpdateTexture3D(FRHICommandListBase& RHICmdList, FRHITexture* Texture, uint32 MipIndex, const struct FUpdateTextureRegion3D& UpdateRegion) override final
 	{
 		return RHI->RHIBeginUpdateTexture3D(RHICmdList, Texture, MipIndex, UpdateRegion);
 	}
@@ -1181,11 +1181,11 @@ public:
 		return RHI->RHICreateShaderLibrary_RenderThread(RHICmdList, Platform, FilePath, Name);
 	}
 
-	virtual void* RHILockTextureCubeFace_RenderThread(class FRHICommandListImmediate& RHICmdList, FRHITextureCube* Texture, uint32 FaceIndex, uint32 ArrayIndex, uint32 MipIndex, EResourceLockMode LockMode, uint32& DestStride, bool bLockWithinMiptail) override final
+	virtual void* RHILockTextureCubeFace_RenderThread(class FRHICommandListImmediate& RHICmdList, FRHITexture* Texture, uint32 FaceIndex, uint32 ArrayIndex, uint32 MipIndex, EResourceLockMode LockMode, uint32& DestStride, bool bLockWithinMiptail) override final
 	{
 		return RHI->RHILockTextureCubeFace_RenderThread(RHICmdList, Texture, FaceIndex, ArrayIndex, MipIndex, LockMode, DestStride, bLockWithinMiptail);
 	}
-	virtual void RHIUnlockTextureCubeFace_RenderThread(class FRHICommandListImmediate& RHICmdList, FRHITextureCube* Texture, uint32 FaceIndex, uint32 ArrayIndex, uint32 MipIndex, bool bLockWithinMiptail) override final
+	virtual void RHIUnlockTextureCubeFace_RenderThread(class FRHICommandListImmediate& RHICmdList, FRHITexture* Texture, uint32 FaceIndex, uint32 ArrayIndex, uint32 MipIndex, bool bLockWithinMiptail) override final
 	{
 		RHI->RHIUnlockTextureCubeFace_RenderThread(RHICmdList, Texture, FaceIndex, ArrayIndex, MipIndex, bLockWithinMiptail);
 	}
diff --git a/Engine/Source/Runtime/RHI/Public/RHIValidationContext.h b/Engine/Source/Runtime/RHI/Public/RHIValidationContext.h
--- a/Engine/Source/Runtime/RHI/Public/RHIValidationContext.h
+++ b/Engine/Source/Runtime/RHI/Public/RHIValidationContext.h
@@ -540,7 +540,7 @@ public:
 		RHIContext->RHIEndUAVOverlap(UAVs);
 	}
 
-	virtual void RHIResummarizeHTile(FRHITexture2D* DepthTexture) override final
+	virtual void RHIResummarizeHTile(FRHITexture* DepthTexture) override final
 	{
 		Tracker->Assert(DepthTexture->GetWholeResourceIdentity(), ERHIAccess::DSVWrite);
 		RHIContext->RHIResummarizeHTile(DepthTexture);
diff --git a/Engine/Source/Runtime/RenderCore/Private/GlobalRenderResources.cpp b/Engine/Source/Runtime/RenderCore/Private/GlobalRenderResources.cpp
--- a/Engine/Source/Runtime/RenderCore/Private/GlobalRenderResources.cpp
+++ b/Engine/Source/Runtime/RenderCore/Private/GlobalRenderResources.cpp
@@ -366,7 +366,7 @@ public:
 			.SetFlags(ETextureCreateFlags::ShaderResource)
 			.SetClassName(TEXT("FSolidColorTextureCube"));
 
-		FTextureCubeRHIRef TextureCube = RHICreateTexture(Desc);
+		FTextureRHIRef TextureCube = RHICreateTexture(Desc);
 		TextureRHI = TextureCube;
 
 		// Write the contents of the texture.
diff --git a/Engine/Source/Runtime/RenderCore/Private/MediaShaders.cpp b/Engine/Source/Runtime/RenderCore/Private/MediaShaders.cpp
--- a/Engine/Source/Runtime/RenderCore/Private/MediaShaders.cpp
+++ b/Engine/Source/Runtime/RenderCore/Private/MediaShaders.cpp
@@ -152,7 +152,7 @@ IMPLEMENT_GLOBAL_SHADER_PARAMETER_STRUCT(FAYUVConvertUB, "AYUVConvertUB");
 IMPLEMENT_SHADER_TYPE(, FAYUVConvertPS, TEXT("/Engine/Private/MediaShaders.usf"), TEXT("AYUVConvertPS"), SF_Pixel);
 
 
-void FAYUVConvertPS::SetParameters(FRHIBatchedShaderParameters& BatchedParameters, TRefCountPtr<FRHITexture2D> AYUVTexture, const FMatrix& ColorTransform, const FVector& YUVOffset, bool SrgbToLinear)
+void FAYUVConvertPS::SetParameters(FRHIBatchedShaderParameters& BatchedParameters, TRefCountPtr<FRHITexture> AYUVTexture, const FMatrix& ColorTransform, const FVector& YUVOffset, bool SrgbToLinear)
 {
 	FAYUVConvertUB UB;
 	{
@@ -181,7 +181,7 @@ IMPLEMENT_GLOBAL_SHADER_PARAMETER_STRUCT(FBMPConvertUB, "BMPConvertUB");
 IMPLEMENT_SHADER_TYPE(, FBMPConvertPS, TEXT("/Engine/Private/MediaShaders.usf"), TEXT("BMPConvertPS"), SF_Pixel);
 
 
-void FBMPConvertPS::SetParameters(FRHIBatchedShaderParameters& BatchedParameters, TRefCountPtr<FRHITexture2D> BMPTexture, const FIntPoint& OutputDimensions, bool SrgbToLinear)
+void FBMPConvertPS::SetParameters(FRHIBatchedShaderParameters& BatchedParameters, TRefCountPtr<FRHITexture> BMPTexture, const FIntPoint& OutputDimensions, bool SrgbToLinear)
 {
 	FBMPConvertUB UB;
 	{
@@ -272,7 +272,7 @@ END_GLOBAL_SHADER_PARAMETER_STRUCT()
 IMPLEMENT_GLOBAL_SHADER_PARAMETER_STRUCT(FNV12ConvertAsBytesUB, "NV12ConvertAsBytesUB");
 IMPLEMENT_SHADER_TYPE(, FNV12ConvertAsBytesPS, TEXT("/Engine/Private/MediaShaders.usf"), TEXT("NV12ConvertAsBytesPS"), SF_Pixel);
 
-void FNV12ConvertAsBytesPS::SetParameters(FRHIBatchedShaderParameters& BatchedParameters, TRefCountPtr<FRHITexture2D> NV12Texture, const FIntPoint& OutputDimensions, const FMatrix44f& ColorTransform, UE::Color::EEncoding Encoding, const FMatrix44f& CSTransform, bool bSwapChroma)
+void FNV12ConvertAsBytesPS::SetParameters(FRHIBatchedShaderParameters& BatchedParameters, TRefCountPtr<FRHITexture> NV12Texture, const FIntPoint& OutputDimensions, const FMatrix44f& ColorTransform, UE::Color::EEncoding Encoding, const FMatrix44f& CSTransform, bool bSwapChroma)
 {
 	FNV12ConvertAsBytesUB UB;
 	{
@@ -309,7 +309,7 @@ IMPLEMENT_GLOBAL_SHADER_PARAMETER_STRUCT(FNV21ConvertUB, "NV21ConvertUB");
 IMPLEMENT_SHADER_TYPE(, FNV21ConvertPS, TEXT("/Engine/Private/MediaShaders.usf"), TEXT("NV21ConvertPS"), SF_Pixel);
 
 
-void FNV21ConvertPS::SetParameters(FRHIBatchedShaderParameters& BatchedParameters, TRefCountPtr<FRHITexture2D> NV21Texture, const FIntPoint& OutputDimensions, const FMatrix& ColorTransform, const FVector& YUVOffset, bool SrgbToLinear)
+void FNV21ConvertPS::SetParameters(FRHIBatchedShaderParameters& BatchedParameters, TRefCountPtr<FRHITexture> NV21Texture, const FIntPoint& OutputDimensions, const FMatrix& ColorTransform, const FVector& YUVOffset, bool SrgbToLinear)
 {
 	FNV21ConvertUB UB;
 	{
@@ -380,7 +380,7 @@ END_GLOBAL_SHADER_PARAMETER_STRUCT()
 IMPLEMENT_GLOBAL_SHADER_PARAMETER_STRUCT(FP010ConvertAsUINT16sUB, "P010ConvertAsUINT16sUB");
 IMPLEMENT_SHADER_TYPE(, FP010ConvertAsUINT16sPS, TEXT("/Engine/Private/MediaShaders.usf"), TEXT("P010ConvertAsUINT16sPS"), SF_Pixel);
 
-void FP010ConvertAsUINT16sPS::SetParameters(FRHIBatchedShaderParameters& BatchedParameters, const FIntPoint& TexDim, TRefCountPtr<FRHITexture2D> NV12Texture, const FIntPoint& OutputDimensions, const FMatrix44f& ColorTransform, const FMatrix44f& CSTransform, UE::Color::EEncoding Encoding)
+void FP010ConvertAsUINT16sPS::SetParameters(FRHIBatchedShaderParameters& BatchedParameters, const FIntPoint& TexDim, TRefCountPtr<FRHITexture> NV12Texture, const FIntPoint& OutputDimensions, const FMatrix44f& ColorTransform, const FMatrix44f& CSTransform, UE::Color::EEncoding Encoding)
 {
 	FP010ConvertAsUINT16sUB UB;
 	{
@@ -460,7 +460,7 @@ IMPLEMENT_GLOBAL_SHADER_PARAMETER_STRUCT(FRGBConvertUB, "RGBConvertUB");
 IMPLEMENT_SHADER_TYPE(, FRGBConvertPS, TEXT("/Engine/Private/MediaShaders.usf"), TEXT("RGBConvertPS"), SF_Pixel);
 
 
-void FRGBConvertPS::SetParameters(FRHIBatchedShaderParameters& BatchedParameters, TRefCountPtr<FRHITexture2D> RGBTexture, const FIntPoint& OutputDimensions, UE::Color::EEncoding Encoding, const FMatrix44f& CSTransform)
+void FRGBConvertPS::SetParameters(FRHIBatchedShaderParameters& BatchedParameters, TRefCountPtr<FRHITexture> RGBTexture, const FIntPoint& OutputDimensions, UE::Color::EEncoding Encoding, const FMatrix44f& CSTransform)
 {
 	FRGBConvertUB UB;
 	{
@@ -491,7 +491,7 @@ IMPLEMENT_GLOBAL_SHADER_PARAMETER_STRUCT(FYCoCgConvertUB, "YCoCgConvertUB");
 IMPLEMENT_SHADER_TYPE(, FYCoCgConvertPS, TEXT("/Engine/Private/MediaShaders.usf"), TEXT("YCoCgConvertPS"), SF_Pixel);
 
 
-void FYCoCgConvertPS::SetParameters(FRHIBatchedShaderParameters& BatchedParameters, TRefCountPtr<FRHITexture2D> RGBTexture, const FIntPoint& OutputDimensions, UE::Color::EEncoding Encoding, const FMatrix44f& CSTransform)
+void FYCoCgConvertPS::SetParameters(FRHIBatchedShaderParameters& BatchedParameters, TRefCountPtr<FRHITexture> RGBTexture, const FIntPoint& OutputDimensions, UE::Color::EEncoding Encoding, const FMatrix44f& CSTransform)
 {
 	FYCoCgConvertUB UB;
 	{
@@ -524,7 +524,7 @@ IMPLEMENT_SHADER_TYPE(, FYCbCrConvertPS, TEXT("/Engine/Private/MediaShaders.usf"
 IMPLEMENT_SHADER_TYPE(, FYCbCrConvertPS_4x4Matrix, TEXT("/Engine/Private/MediaShaders.usf"), TEXT("YCbCrConvertPS_4x4Matrix"), SF_Pixel);
 
 
-void FYCbCrConvertPS::SetParameters(FRHIBatchedShaderParameters& BatchedParameters, TRefCountPtr<FRHITexture2D> LumaTexture, TRefCountPtr<FRHITexture2D> CbCrTexture, const FMatrix& ColorTransform, const FVector& YUVOffset, bool SrgbToLinear)
+void FYCbCrConvertPS::SetParameters(FRHIBatchedShaderParameters& BatchedParameters, TRefCountPtr<FRHITexture> LumaTexture, TRefCountPtr<FRHITexture> CbCrTexture, const FMatrix& ColorTransform, const FVector& YUVOffset, bool SrgbToLinear)
 {
 	FYCbCrConvertUB UB;
 	{
@@ -560,7 +560,7 @@ IMPLEMENT_GLOBAL_SHADER_PARAMETER_STRUCT(FUYVYConvertUB, "UYVYConvertUB");
 IMPLEMENT_SHADER_TYPE(, FUYVYConvertPS, TEXT("/Engine/Private/MediaShaders.usf"), TEXT("UYVYConvertPS"), SF_Pixel);
 
 
-void FUYVYConvertPS::SetParameters(FRHIBatchedShaderParameters& BatchedParameters, TRefCountPtr<FRHITexture2D> UYVYTexture, const FMatrix& ColorTransform, const FVector& YUVOffset, bool SrgbToLinear)
+void FUYVYConvertPS::SetParameters(FRHIBatchedShaderParameters& BatchedParameters, TRefCountPtr<FRHITexture> UYVYTexture, const FMatrix& ColorTransform, const FVector& YUVOffset, bool SrgbToLinear)
 {
 	FUYVYConvertUB UB;
 	{
@@ -596,7 +596,7 @@ IMPLEMENT_GLOBAL_SHADER_PARAMETER_STRUCT(FYUVConvertUB, "YUVConvertUB");
 IMPLEMENT_SHADER_TYPE(, FYUVConvertPS, TEXT("/Engine/Private/MediaShaders.usf"), TEXT("YUVConvertPS"), SF_Pixel);
 
 
-void FYUVConvertPS::SetParameters(FRHIBatchedShaderParameters& BatchedParameters, TRefCountPtr<FRHITexture2D> YTexture, TRefCountPtr<FRHITexture2D> UTexture, TRefCountPtr<FRHITexture2D> VTexture, const FIntPoint& OutputDimensions, const FMatrix& ColorTransform, const FVector& YUVOffset, bool SrgbToLinear)
+void FYUVConvertPS::SetParameters(FRHIBatchedShaderParameters& BatchedParameters, TRefCountPtr<FRHITexture> YTexture, TRefCountPtr<FRHITexture> UTexture, TRefCountPtr<FRHITexture> VTexture, const FIntPoint& OutputDimensions, const FMatrix& ColorTransform, const FVector& YUVOffset, bool SrgbToLinear)
 {
 	FYUVConvertUB UB;
 	{
@@ -635,7 +635,7 @@ IMPLEMENT_GLOBAL_SHADER_PARAMETER_STRUCT(FYUVv216ConvertUB, "YUVv216ConvertUB");
 IMPLEMENT_SHADER_TYPE(, FYUVv216ConvertPS, TEXT("/Engine/Private/MediaShaders.usf"), TEXT("YUVv216ConvertPS"), SF_Pixel);
 
 
-void FYUVv216ConvertPS::SetParameters(FRHIBatchedShaderParameters& BatchedParameters, TRefCountPtr<FRHITexture2D> YUVTexture, const FIntPoint& OutputDimensions, const FMatrix44f& ColorTransform, UE::Color::EEncoding Encoding, const FMatrix44f& CSTransform, bool bIsCbY0CrY1, bool bIsARGBFmt, bool bSwapChroma)
+void FYUVv216ConvertPS::SetParameters(FRHIBatchedShaderParameters& BatchedParameters, TRefCountPtr<FRHITexture> YUVTexture, const FIntPoint& OutputDimensions, const FMatrix44f& ColorTransform, UE::Color::EEncoding Encoding, const FMatrix44f& CSTransform, bool bIsCbY0CrY1, bool bIsARGBFmt, bool bSwapChroma)
 {
 	FYUVv216ConvertUB UB;
 	{
@@ -672,7 +672,7 @@ IMPLEMENT_GLOBAL_SHADER_PARAMETER_STRUCT(FYUVv210ConvertUB, "YUVv210ConvertUB");
 IMPLEMENT_SHADER_TYPE(, FYUVv210ConvertPS, TEXT("/Engine/Private/MediaShaders.usf"), TEXT("YUVv210ConvertPS"), SF_Pixel);
 
 
-void FYUVv210ConvertPS::SetParameters(FRHIBatchedShaderParameters& BatchedParameters, TRefCountPtr<FRHITexture2D> YUVTexture, const FIntPoint& OutputDimensions, const FMatrix44f& ColorTransform, UE::Color::EEncoding Encoding, const FMatrix44f& CSTransform, bool bIsCbY0CrY1)
+void FYUVv210ConvertPS::SetParameters(FRHIBatchedShaderParameters& BatchedParameters, TRefCountPtr<FRHITexture> YUVTexture, const FIntPoint& OutputDimensions, const FMatrix44f& ColorTransform, UE::Color::EEncoding Encoding, const FMatrix44f& CSTransform, bool bIsCbY0CrY1)
 {
 	FYUVv210ConvertUB UB;
 	{
@@ -769,7 +769,7 @@ IMPLEMENT_GLOBAL_SHADER_PARAMETER_STRUCT(FYUY2ConvertUB, "YUY2ConvertUB");
 IMPLEMENT_SHADER_TYPE(, FYUY2ConvertPS, TEXT("/Engine/Private/MediaShaders.usf"), TEXT("YUY2ConvertPS"), SF_Pixel);
 
 
-void FYUY2ConvertPS::SetParameters(FRHIBatchedShaderParameters& BatchedParameters, TRefCountPtr<FRHITexture2D> YUY2Texture, const FIntPoint& OutputDimensions, const FMatrix& ColorTransform, const FVector& YUVOffset, bool SrgbToLinear)
+void FYUY2ConvertPS::SetParameters(FRHIBatchedShaderParameters& BatchedParameters, TRefCountPtr<FRHITexture> YUY2Texture, const FIntPoint& OutputDimensions, const FMatrix& ColorTransform, const FVector& YUVOffset, bool SrgbToLinear)
 {
 	FYUY2ConvertUB UB;
 	{
@@ -802,7 +802,7 @@ IMPLEMENT_GLOBAL_SHADER_PARAMETER_STRUCT(FYVYUConvertUB, "YVYUConvertUB");
 IMPLEMENT_SHADER_TYPE(, FYVYUConvertPS, TEXT("/Engine/Private/MediaShaders.usf"), TEXT("YVYUConvertPS"), SF_Pixel);
 
 
-void FYVYUConvertPS::SetParameters(FRHIBatchedShaderParameters& BatchedParameters, TRefCountPtr<FRHITexture2D> YVYUTexture, const FMatrix& ColorTransform, const FVector& YUVOffset, bool SrgbToLinear)
+void FYVYUConvertPS::SetParameters(FRHIBatchedShaderParameters& BatchedParameters, TRefCountPtr<FRHITexture> YVYUTexture, const FMatrix& ColorTransform, const FVector& YUVOffset, bool SrgbToLinear)
 {
 	FYVYUConvertUB UB;
 	{
@@ -833,7 +833,7 @@ IMPLEMENT_GLOBAL_SHADER_PARAMETER_STRUCT(FRGB8toY8ConvertUB, "RGB8toY8ConvertUB"
 IMPLEMENT_SHADER_TYPE(, FRGB8toY8ConvertPS, TEXT("/Engine/Private/MediaShaders.usf"), TEXT("RGB8toY8ConvertPS"), SF_Pixel);
 
 
-void FRGB8toY8ConvertPS::SetParameters(FRHIBatchedShaderParameters& BatchedParameters, TRefCountPtr<FRHITexture2D> RGBATexture, const FVector4f& ColorTransform, bool LinearToSrgb)
+void FRGB8toY8ConvertPS::SetParameters(FRHIBatchedShaderParameters& BatchedParameters, TRefCountPtr<FRHITexture> RGBATexture, const FVector4f& ColorTransform, bool LinearToSrgb)
 {
 	FRGB8toY8ConvertUB UB;
 	{
diff --git a/Engine/Source/Runtime/RenderCore/Private/RenderGraphUtils.cpp b/Engine/Source/Runtime/RenderCore/Private/RenderGraphUtils.cpp
--- a/Engine/Source/Runtime/RenderCore/Private/RenderGraphUtils.cpp
+++ b/Engine/Source/Runtime/RenderCore/Private/RenderGraphUtils.cpp
@@ -799,7 +799,7 @@ void AddResummarizeHTilePass(FRDGBuilder& GraphBuilder, FRDGTextureRef Texture)
 	GraphBuilder.AddPass(RDG_EVENT_NAME("ResummarizeHTile (%s)", Texture->Name), PassParameters, ERDGPassFlags::Raster,
 		[Texture](FRHICommandList& RHICmdList)
 	{
-		RHICmdList.ResummarizeHTile(static_cast<FRHITexture2D*>(Texture->GetRHI()));
+		RHICmdList.ResummarizeHTile(static_cast<FRHITexture*>(Texture->GetRHI()));
 	});
 }
 
diff --git a/Engine/Source/Runtime/RenderCore/Public/MediaShaders.h b/Engine/Source/Runtime/RenderCore/Public/MediaShaders.h
--- a/Engine/Source/Runtime/RenderCore/Public/MediaShaders.h
+++ b/Engine/Source/Runtime/RenderCore/Public/MediaShaders.h
@@ -195,7 +195,7 @@ public:
 		: FGlobalShader(Initializer)
 	{ }
 
-	RENDERCORE_API void SetParameters(FRHIBatchedShaderParameters& BatchedParameters, TRefCountPtr<FRHITexture2D> AYUVTexture, const FMatrix& ColorTransform, const FVector& YUVOffset, bool SrgbToLinear);
+	RENDERCORE_API void SetParameters(FRHIBatchedShaderParameters& BatchedParameters, TRefCountPtr<FRHITexture> AYUVTexture, const FMatrix& ColorTransform, const FVector& YUVOffset, bool SrgbToLinear);
 };
 
 
@@ -216,7 +216,7 @@ public:
 		: FGlobalShader(Initializer)
 	{ }
 
-	RENDERCORE_API void SetParameters(FRHIBatchedShaderParameters& BatchedParameters, TRefCountPtr<FRHITexture2D> BMPTexture, const FIntPoint& OutputDimensions, bool SrgbToLinear);
+	RENDERCORE_API void SetParameters(FRHIBatchedShaderParameters& BatchedParameters, TRefCountPtr<FRHITexture> BMPTexture, const FIntPoint& OutputDimensions, bool SrgbToLinear);
 };
 
 
@@ -263,7 +263,7 @@ public:
 		: FGlobalShader(Initializer)
 	{ }
 
-	RENDERCORE_API void SetParameters(FRHIBatchedShaderParameters& BatchedParameters, TRefCountPtr<FRHITexture2D> NV12Texture, const FIntPoint& OutputDimensions, const FMatrix44f& ColorTransform, UE::Color::EEncoding Encoding, const FMatrix44f& CSTransform, bool bSwapChroma);
+	RENDERCORE_API void SetParameters(FRHIBatchedShaderParameters& BatchedParameters, TRefCountPtr<FRHITexture> NV12Texture, const FIntPoint& OutputDimensions, const FMatrix44f& ColorTransform, UE::Color::EEncoding Encoding, const FMatrix44f& CSTransform, bool bSwapChroma);
 };
 
 
@@ -286,7 +286,7 @@ public:
 		: FGlobalShader(Initializer)
 	{ }
 
-	RENDERCORE_API void SetParameters(FRHIBatchedShaderParameters& BatchedParameters, TRefCountPtr<FRHITexture2D> NV21Texture, const FIntPoint& OutputDimensions, const FMatrix& ColorTransform, const FVector& YUVOffset, bool SrgbToLinear);
+	RENDERCORE_API void SetParameters(FRHIBatchedShaderParameters& BatchedParameters, TRefCountPtr<FRHITexture> NV21Texture, const FIntPoint& OutputDimensions, const FMatrix& ColorTransform, const FVector& YUVOffset, bool SrgbToLinear);
 };
 
 /**
@@ -324,7 +324,7 @@ public:
 		: FGlobalShader(Initializer)
 	{ }
 
-	RENDERCORE_API void SetParameters(FRHIBatchedShaderParameters& BatchedParameters, const FIntPoint& TexDim, TRefCountPtr<FRHITexture2D> NV12Texture, const FIntPoint& OutputDimensions, const FMatrix44f& ColorTransform, const FMatrix44f& CSTransform, UE::Color::EEncoding Encoding);
+	RENDERCORE_API void SetParameters(FRHIBatchedShaderParameters& BatchedParameters, const FIntPoint& TexDim, TRefCountPtr<FRHITexture> NV12Texture, const FIntPoint& OutputDimensions, const FMatrix44f& ColorTransform, const FMatrix44f& CSTransform, UE::Color::EEncoding Encoding);
 };
 
 /**
@@ -365,7 +365,7 @@ public:
 	{ }
 
 
-	RENDERCORE_API void SetParameters(FRHIBatchedShaderParameters& BatchedParameters, TRefCountPtr<FRHITexture2D> RGBTexture, const FIntPoint& OutputDimensions, UE::Color::EEncoding Encoding, const FMatrix44f& CSTransform);
+	RENDERCORE_API void SetParameters(FRHIBatchedShaderParameters& BatchedParameters, TRefCountPtr<FRHITexture> RGBTexture, const FIntPoint& OutputDimensions, UE::Color::EEncoding Encoding, const FMatrix44f& CSTransform);
 };
 
 /**
@@ -387,7 +387,7 @@ public:
 	{ }
 
 
-	RENDERCORE_API void SetParameters(FRHIBatchedShaderParameters& BatchedParameters, TRefCountPtr<FRHITexture2D> YCoCgTexture, const FIntPoint& OutputDimensions, UE::Color::EEncoding Encoding, const FMatrix44f& CSTransform);
+	RENDERCORE_API void SetParameters(FRHIBatchedShaderParameters& BatchedParameters, TRefCountPtr<FRHITexture> YCoCgTexture, const FIntPoint& OutputDimensions, UE::Color::EEncoding Encoding, const FMatrix44f& CSTransform);
 };
 
 
@@ -412,7 +412,7 @@ public:
 		: FGlobalShader(Initializer)
 	{ }
 
-	RENDERCORE_API void SetParameters(FRHIBatchedShaderParameters& BatchedParameters, TRefCountPtr<FRHITexture2D> LumaTexture, TRefCountPtr<FRHITexture2D> CbCrTexture, const FMatrix& ColorTransform, const FVector& YUVOffset, bool SrgbToLinear);
+	RENDERCORE_API void SetParameters(FRHIBatchedShaderParameters& BatchedParameters, TRefCountPtr<FRHITexture> LumaTexture, TRefCountPtr<FRHITexture> CbCrTexture, const FMatrix& ColorTransform, const FVector& YUVOffset, bool SrgbToLinear);
 };
 
 class FYCbCrConvertPS_4x4Matrix : public FYCbCrConvertPS
@@ -448,7 +448,7 @@ public:
 		: FGlobalShader(Initializer)
 	{ }
 
-	RENDERCORE_API void SetParameters(FRHIBatchedShaderParameters& BatchedParameters, TRefCountPtr<FRHITexture2D> UYVYTexture, const FMatrix& ColorTransform, const FVector& YUVOffset, bool SrgbToLinear);
+	RENDERCORE_API void SetParameters(FRHIBatchedShaderParameters& BatchedParameters, TRefCountPtr<FRHITexture> UYVYTexture, const FMatrix& ColorTransform, const FVector& YUVOffset, bool SrgbToLinear);
 };
 
 
@@ -471,7 +471,7 @@ public:
 	{ }
 
 
-	RENDERCORE_API void SetParameters(FRHIBatchedShaderParameters& BatchedParameters, TRefCountPtr<FRHITexture2D> YTexture, TRefCountPtr<FRHITexture2D> UTexture, TRefCountPtr<FRHITexture2D> VTexture, const FIntPoint& OutputDimensions, const FMatrix& ColorTransform, const FVector& YUVOffset, bool SrgbToLinear);
+	RENDERCORE_API void SetParameters(FRHIBatchedShaderParameters& BatchedParameters, TRefCountPtr<FRHITexture> YTexture, TRefCountPtr<FRHITexture> UTexture, TRefCountPtr<FRHITexture> VTexture, const FIntPoint& OutputDimensions, const FMatrix& ColorTransform, const FVector& YUVOffset, bool SrgbToLinear);
 };
 
 
@@ -492,7 +492,7 @@ public:
 		: FGlobalShader(Initializer)
 	{ }
 
-	RENDERCORE_API void SetParameters(FRHIBatchedShaderParameters& BatchedParameters, TRefCountPtr<FRHITexture2D> YUVTexture, const FIntPoint& OutputDimensions, const FMatrix44f& ColorTransform, UE::Color::EEncoding Encoding, const FMatrix44f& CSTransform, bool bIsCbY0CrY1, bool bIsARGBFmt, bool bSwapChroma);
+	RENDERCORE_API void SetParameters(FRHIBatchedShaderParameters& BatchedParameters, TRefCountPtr<FRHITexture> YUVTexture, const FIntPoint& OutputDimensions, const FMatrix44f& ColorTransform, UE::Color::EEncoding Encoding, const FMatrix44f& CSTransform, bool bIsCbY0CrY1, bool bIsARGBFmt, bool bSwapChroma);
 };
 
 
@@ -514,7 +514,7 @@ public:
 		: FGlobalShader(Initializer)
 	{ }
 
-	RENDERCORE_API void SetParameters(FRHIBatchedShaderParameters& BatchedParameters, TRefCountPtr<FRHITexture2D> YUVTexture, const FIntPoint& OutputDimensions, const FMatrix44f& ColorTransform, UE::Color::EEncoding Encoding, const FMatrix44f& CSTransform, bool bIsCbY0CrY1);
+	RENDERCORE_API void SetParameters(FRHIBatchedShaderParameters& BatchedParameters, TRefCountPtr<FRHITexture> YUVTexture, const FIntPoint& OutputDimensions, const FMatrix44f& ColorTransform, UE::Color::EEncoding Encoding, const FMatrix44f& CSTransform, bool bIsCbY0CrY1);
 };
 
 
@@ -558,7 +558,7 @@ public:
 		: FGlobalShader(Initializer)
 	{ }
 
-	RENDERCORE_API void SetParameters(FRHIBatchedShaderParameters& BatchedParameters, TRefCountPtr<FRHITexture2D> YUY2Texture, const FIntPoint& OutputDimensions, const FMatrix& ColorTransform, const FVector& YUVOffset, bool SrgbToLinear);
+	RENDERCORE_API void SetParameters(FRHIBatchedShaderParameters& BatchedParameters, TRefCountPtr<FRHITexture> YUY2Texture, const FIntPoint& OutputDimensions, const FMatrix& ColorTransform, const FVector& YUVOffset, bool SrgbToLinear);
 };
 
 /**
@@ -602,7 +602,7 @@ public:
 		: FGlobalShader(Initializer)
 	{ }
 
-	RENDERCORE_API void SetParameters(FRHIBatchedShaderParameters& BatchedParameters, TRefCountPtr<FRHITexture2D> YVYUTexture, const FMatrix& ColorTransform, const FVector& YUVOffset, bool SrgbToLinear);
+	RENDERCORE_API void SetParameters(FRHIBatchedShaderParameters& BatchedParameters, TRefCountPtr<FRHITexture> YVYUTexture, const FMatrix& ColorTransform, const FVector& YUVOffset, bool SrgbToLinear);
 };
 
 
@@ -623,7 +623,7 @@ public:
 		: FGlobalShader(Initializer)
 	{ }
 
-	RENDERCORE_API void SetParameters(FRHIBatchedShaderParameters& BatchedParameters, TRefCountPtr<FRHITexture2D> RGBATexture, const FVector4f& ColorTransform, bool LinearToSrgb);
+	RENDERCORE_API void SetParameters(FRHIBatchedShaderParameters& BatchedParameters, TRefCountPtr<FRHITexture> RGBATexture, const FVector4f& ColorTransform, bool LinearToSrgb);
 };
 
 
diff --git a/Engine/Source/Runtime/Renderer/Private/PostProcess/DiaphragmDOF.h b/Engine/Source/Runtime/Renderer/Private/PostProcess/DiaphragmDOF.h
--- a/Engine/Source/Runtime/Renderer/Private/PostProcess/DiaphragmDOF.h
+++ b/Engine/Source/Runtime/Renderer/Private/PostProcess/DiaphragmDOF.h
@@ -38,7 +38,7 @@ struct FPhysicalCocModel
 	float InFocusRadius;
 
 	/** Radius offset lookup table */
-	FRHITexture2D* DynamicRadiusOffsetLUT;
+	FRHITexture* DynamicRadiusOffsetLUT;
 
 	// Resolution less minimal foreground coc radius < 0.
 	float MinForegroundCocRadius;
diff --git a/Engine/Source/Runtime/Renderer/Private/PostProcess/PostProcessBufferInspector.cpp b/Engine/Source/Runtime/Renderer/Private/PostProcess/PostProcessBufferInspector.cpp
--- a/Engine/Source/Runtime/Renderer/Private/PostProcess/PostProcessBufferInspector.cpp
+++ b/Engine/Source/Runtime/Renderer/Private/PostProcess/PostProcessBufferInspector.cpp
@@ -57,7 +57,7 @@ void ProcessPixelInspectorRequests(
 					0
 				);
 
-				const FTexture2DRHIRef &DestinationBufferDepth = PixelInspectorData.RenderTargetBufferDepth[PixelInspectorRequest->BufferIndex]->GetRenderTargetTexture();
+				const FTextureRHIRef &DestinationBufferDepth = PixelInspectorData.RenderTargetBufferDepth[PixelInspectorRequest->BufferIndex]->GetRenderTargetTexture();
 				if (DestinationBufferDepth.IsValid())
 				{
 					FRHITexture* SourceBufferSceneDepth = Parameters.SceneDepth->GetRHI();
@@ -73,7 +73,7 @@ void ProcessPixelInspectorRequests(
 
 			//////////////////////////////////////////////////////////////////////////
 			// FINAL COLOR
-			const FTexture2DRHIRef &DestinationBufferFinalColor = PixelInspectorData.RenderTargetBufferFinalColor[PixelInspectorRequest->BufferIndex]->GetRenderTargetTexture();
+			const FTextureRHIRef &DestinationBufferFinalColor = PixelInspectorData.RenderTargetBufferFinalColor[PixelInspectorRequest->BufferIndex]->GetRenderTargetTexture();
 			if (DestinationBufferFinalColor.IsValid())
 			{
 				const FIntVector SourcePoint(
@@ -114,7 +114,7 @@ void ProcessPixelInspectorRequests(
 
 			//////////////////////////////////////////////////////////////////////////
 			// ORIGINAL SCENE COLOR
-			const FTexture2DRHIRef& DestinationBufferSceneColor = PixelInspectorData.RenderTargetBufferSceneColor[PixelInspectorRequest->BufferIndex]->GetRenderTargetTexture();
+			const FTextureRHIRef& DestinationBufferSceneColor = PixelInspectorData.RenderTargetBufferSceneColor[PixelInspectorRequest->BufferIndex]->GetRenderTargetTexture();
 			if (DestinationBufferSceneColor.IsValid())
 			{
 				const FIntVector SourcePoint(
@@ -135,7 +135,7 @@ void ProcessPixelInspectorRequests(
 
 			//////////////////////////////////////////////////////////////////////////
 			// HDR
-			const FTexture2DRHIRef &DestinationBufferHDR = PixelInspectorData.RenderTargetBufferHDR[PixelInspectorRequest->BufferIndex]->GetRenderTargetTexture();
+			const FTextureRHIRef &DestinationBufferHDR = PixelInspectorData.RenderTargetBufferHDR[PixelInspectorRequest->BufferIndex]->GetRenderTargetTexture();
 			if (DestinationBufferHDR.IsValid())
 			{
 				const FIntVector SourcePoint(
@@ -167,7 +167,7 @@ void ProcessPixelInspectorRequests(
 					0
 				);
 
-				const FTexture2DRHIRef &DestinationBufferA = PixelInspectorData.RenderTargetBufferA[PixelInspectorRequest->BufferIndex]->GetRenderTargetTexture();
+				const FTextureRHIRef &DestinationBufferA = PixelInspectorData.RenderTargetBufferA[PixelInspectorRequest->BufferIndex]->GetRenderTargetTexture();
 				if (DestinationBufferA.IsValid() && Parameters.GBufferA)
 				{
 					FRHITexture* SourceBufferA = Parameters.GBufferA->GetRHI();
@@ -183,7 +183,7 @@ void ProcessPixelInspectorRequests(
 
 			//////////////////////////////////////////////////////////////////////////
 			// GBuffer BCDEF
-			const FTexture2DRHIRef &DestinationBufferBCDEF = PixelInspectorData.RenderTargetBufferBCDEF[PixelInspectorRequest->BufferIndex]->GetRenderTargetTexture();
+			const FTextureRHIRef &DestinationBufferBCDEF = PixelInspectorData.RenderTargetBufferBCDEF[PixelInspectorRequest->BufferIndex]->GetRenderTargetTexture();
 			if (DestinationBufferBCDEF.IsValid())
 			{
 				const FIntVector SourcePoint(
diff --git a/Engine/Source/Runtime/Renderer/Private/ReflectionEnvironmentCapture.cpp b/Engine/Source/Runtime/Renderer/Private/ReflectionEnvironmentCapture.cpp
--- a/Engine/Source/Runtime/Renderer/Private/ReflectionEnvironmentCapture.cpp
+++ b/Engine/Source/Runtime/Renderer/Private/ReflectionEnvironmentCapture.cpp
@@ -1241,9 +1241,9 @@ void CaptureSceneIntoScratchCubemap(
 	public:
 		FDummyRenderTarget() = default;
 
-		const FTexture2DRHIRef& GetRenderTargetTexture() const override
+		const FTextureRHIRef& GetRenderTargetTexture() const override
 		{
-			static FTexture2DRHIRef DummyTexture;
+			static FTextureRHIRef DummyTexture;
 			return DummyTexture;
 		}
 
diff --git a/Engine/Source/Runtime/Renderer/Private/RendererScene.cpp b/Engine/Source/Runtime/Renderer/Private/RendererScene.cpp
--- a/Engine/Source/Runtime/Renderer/Private/RendererScene.cpp
+++ b/Engine/Source/Runtime/Renderer/Private/RendererScene.cpp
@@ -720,7 +720,7 @@ FFXSystemInterface* FScene::GetFXSystem()
 	return FXSystem;
 }
 
-static uint64 GetTextureGPUSizeBytes(const FTexture2DRHIRef& Target, bool bLogSizes)
+static uint64 GetTextureGPUSizeBytes(const FTextureRHIRef& Target, bool bLogSizes)
 {
 	uint64 Size = Target.IsValid() ? Target->GetDesc().CalcMemorySizeEstimate() : 0;
 	if (bLogSizes && Size)
diff --git a/Engine/Source/Runtime/Renderer/Private/SceneOcclusion.cpp b/Engine/Source/Runtime/Renderer/Private/SceneOcclusion.cpp
--- a/Engine/Source/Runtime/Renderer/Private/SceneOcclusion.cpp
+++ b/Engine/Source/Runtime/Renderer/Private/SceneOcclusion.cpp
@@ -1011,8 +1011,8 @@ void FHZBOcclusionTester::Submit(FRDGBuilder& GraphBuilder, const FViewInfo& Vie
 				const int32 BlockY = BlockIndex / SizeInBlocksY;
 
 				FUpdateTextureRegion2D Region(BlockX * BlockSize, BlockY * BlockSize, 0, 0, BlockSize, BlockSize);
-				RHIUpdateTexture2D((FRHITexture2D*)BoundsCenterTexture->GetRHI(), 0, Region, BlockStride, (uint8*)CenterBuffer);
-				RHIUpdateTexture2D((FRHITexture2D*)BoundsExtentTexture->GetRHI(), 0, Region, BlockStride, (uint8*)ExtentBuffer);
+				RHIUpdateTexture2D((FRHITexture*)BoundsCenterTexture->GetRHI(), 0, Region, BlockStride, (uint8*)CenterBuffer);
+				RHIUpdateTexture2D((FRHITexture*)BoundsExtentTexture->GetRHI(), 0, Region, BlockStride, (uint8*)ExtentBuffer);
 			}
 		});
 	}
diff --git a/Engine/Source/Runtime/Renderer/Private/SceneRendering.h b/Engine/Source/Runtime/Renderer/Private/SceneRendering.h
--- a/Engine/Source/Runtime/Renderer/Private/SceneRendering.h
+++ b/Engine/Source/Runtime/Renderer/Private/SceneRendering.h
@@ -1550,7 +1550,7 @@ public:
 
 	/** Only one of the resources(TextureBuffer or Texture2D) will be used depending on the Mobile.UseGPUSceneTexture cvar */
 
-	FTexture2DRHIRef PrimitiveSceneDataTextureOverrideRHI;
+	FTextureRHIRef PrimitiveSceneDataTextureOverrideRHI;
 
 	FShaderPrintData ShaderPrintData;
 private:
diff --git a/Engine/Source/Runtime/Renderer/Private/SceneTextures.cpp b/Engine/Source/Runtime/Renderer/Private/SceneTextures.cpp
--- a/Engine/Source/Runtime/Renderer/Private/SceneTextures.cpp
+++ b/Engine/Source/Runtime/Renderer/Private/SceneTextures.cpp
@@ -130,13 +130,13 @@ static IStereoRenderTargetManager* FindStereoRenderTargetManager()
 	return GEngine->StereoRenderingDevice->GetRenderTargetManager();
 }
 
-static TRefCountPtr<FRHITexture2D> FindStereoDepthTexture(uint32 bSupportsXRDepth, FIntPoint TextureExtent, ETextureCreateFlags RequestedCreateFlags)
+static TRefCountPtr<FRHITexture> FindStereoDepthTexture(uint32 bSupportsXRDepth, FIntPoint TextureExtent, ETextureCreateFlags RequestedCreateFlags)
 {
 	if (bSupportsXRDepth == 1)
 	{
 		if (IStereoRenderTargetManager* StereoRenderTargetManager = FindStereoRenderTargetManager())
 		{
-			TRefCountPtr<FRHITexture2D> DepthTex, SRTex;
+			TRefCountPtr<FRHITexture> DepthTex, SRTex;
 			constexpr uint32 NumSamples = 1;
 			StereoRenderTargetManager->AllocateDepthTexture(0, TextureExtent.X, TextureExtent.Y, PF_DepthStencil, 1, RequestedCreateFlags, TexCreate_DepthStencilTargetable | TexCreate_ShaderResource | TexCreate_InputAttachmentRead, DepthTex, SRTex, NumSamples);
 			return MoveTemp(SRTex);
@@ -444,7 +444,7 @@ void FMinimalSceneTextures::InitializeViewFamily(FRDGBuilder& GraphBuilder, FVie
 	// Scene Depth
 
 	// If not using MSAA, we need to make sure to grab the stereo depth texture if appropriate.
-	FTexture2DRHIRef StereoDepthRHI;
+	FTextureRHIRef StereoDepthRHI;
 	if (Config.NumSamples == 1 && (StereoDepthRHI = FindStereoDepthTexture(Config.bSupportsXRTargetManagerDepthAlloc, Config.Extent, ETextureCreateFlags::None)) != nullptr)
 	{
 		SceneTextures.Depth = RegisterExternalTexture(GraphBuilder, StereoDepthRHI, TEXT("SceneDepthZ"));
diff --git a/Engine/Source/Runtime/Renderer/Private/Substrate/Glint/GlintShadingLUTs.h b/Engine/Source/Runtime/Renderer/Private/Substrate/Glint/GlintShadingLUTs.h
--- a/Engine/Source/Runtime/Renderer/Private/Substrate/Glint/GlintShadingLUTs.h
+++ b/Engine/Source/Runtime/Renderer/Private/Substrate/Glint/GlintShadingLUTs.h
@@ -14,7 +14,7 @@ class FViewInfo;
 struct FGlintShadingLUTsStateData
 {
 	TRefCountPtr<IPooledRenderTarget> GlintShadingLUTs = nullptr;
-	FRHITexture2DArray* RHIGlintShadingLUTs = nullptr;
+	FRHITexture* RHIGlintShadingLUTs = nullptr;
 
 	float Dictionary_Alpha = 0.0f;
 	int32 Dictionary_NDistributionsPerChannel = 0;
diff --git a/Engine/Source/Runtime/Renderer/Private/SystemTextures.cpp b/Engine/Source/Runtime/Renderer/Private/SystemTextures.cpp
--- a/Engine/Source/Runtime/Renderer/Private/SystemTextures.cpp
+++ b/Engine/Source/Runtime/Renderer/Private/SystemTextures.cpp
@@ -111,7 +111,7 @@ void FSystemTextures::InitializeCommonTextures(FRHICommandListImmediate& RHICmdL
 	// Create a WhiteDummy texture
 	{
 		WhiteDummy = CreateRenderTarget(GWhiteTexture->TextureRHI, TEXT("WhiteDummy"));
-		WhiteDummySRV = RHICmdList.CreateShaderResourceView((FRHITexture2D*)WhiteDummy->GetRHI(), 0);
+		WhiteDummySRV = RHICmdList.CreateShaderResourceView((FRHITexture*)WhiteDummy->GetRHI(), 0);
 	}
 
 	// Create a BlackDummy texture
@@ -709,7 +709,7 @@ void FSystemTextures::InitializeFeatureLevelDependentTextures(FRHICommandListImm
 					.SetFlags(ETextureCreateFlags::ShaderResource | ETextureCreateFlags::FastVRAM)
 					.SetClassName(SystemTexturesName);
 
-				FTexture2DRHIRef Texture = RHICreateTexture(Desc);
+				FTextureRHIRef Texture = RHICreateTexture(Desc);
 
 				// Write the contents of the texture.
 				uint32 DestStride;
@@ -738,7 +738,7 @@ void FSystemTextures::InitializeFeatureLevelDependentTextures(FRHICommandListImm
 					FRHITextureCreateDesc::Create2D(TEXT("Sheen.LTC"), Sheen_LTC_Size, Sheen_LTC_Size, PF_FloatRGBA)
 					.SetFlags(ETextureCreateFlags::ShaderResource | ETextureCreateFlags::FastVRAM);
 
-				FTexture2DRHIRef Texture = RHICreateTexture(Desc);
+				FTextureRHIRef Texture = RHICreateTexture(Desc);
 
 				// Write the contents of the texture.
 				uint32 DestStride;
@@ -908,7 +908,7 @@ void FSystemTextures::InitializeFeatureLevelDependentTextures(FRHICommandListImm
 			FUpdateTextureRegion3D Region(0, 0, 0, 0, 0, 0, Extent, Extent, Extent);
 
 			RHICmdList.UpdateTexture3D(
-				(FRHITexture3D*)Texture.GetReference(),
+				(FRHITexture*)Texture.GetReference(),
 				0,
 				Region,
 				Extent * sizeof(FFloat16),
@@ -1374,7 +1374,7 @@ void InitializeData(const TInType& InData, EPixelFormat InFormat, uint8* OutData
 }
 
 template <typename DataType>
-void SetDefaultTextureData2D(FRHITexture2D* Texture, const DataType& InData)
+void SetDefaultTextureData2D(FRHITexture* Texture, const DataType& InData)
 {
 	uint8 SrcData[16] = { 0,0,0,0, 0,0,0,0, 0,0,0,0, 0,0,0,0, };
 	uint32 SrcByteCount = 0;
@@ -1388,7 +1388,7 @@ void SetDefaultTextureData2D(FRHITexture2D* Texture, const DataType& InData)
 }
 
 template <typename DataType>
-void SetDefaultTextureData2DArray(FRHITexture2DArray* Texture, const DataType& InData)
+void SetDefaultTextureData2DArray(FRHITexture* Texture, const DataType& InData)
 {
 	uint8 SrcData[16] = { 0,0,0,0, 0,0,0,0, 0,0,0,0, 0,0,0,0, };
 	uint32 SrcByteCount = 0;
@@ -1402,7 +1402,7 @@ void SetDefaultTextureData2DArray(FRHITexture2DArray* Texture, const DataType& I
 }
 
 template <typename DataType>
-void SetDefaultTextureData3D(FRHICommandListImmediate& RHICmdList, FRHITexture3D* Texture, const DataType& InData)
+void SetDefaultTextureData3D(FRHICommandListImmediate& RHICmdList, FRHITexture* Texture, const DataType& InData)
 {
 	uint8 SrcData[16] = { 0,0,0,0, 0,0,0,0, 0,0,0,0, 0,0,0,0, };
 	uint32 SrcByteCount = 0;
@@ -1423,7 +1423,7 @@ void SetDefaultTextureData3D(FRHICommandListImmediate& RHICmdList, FRHITexture3D
 }
 
 template <typename DataType>
-void SetDefaultTextureDataCube(FRHITextureCube* Texture, const DataType& InData)
+void SetDefaultTextureDataCube(FRHITexture* Texture, const DataType& InData)
 {
 	uint8 SrcData[16] = { 0,0,0,0, 0,0,0,0, 0,0,0,0, 0,0,0,0, };
 	uint32 SrcByteCount = 0;
@@ -1481,7 +1481,7 @@ FRDGTextureRef GetInternalDefaultTexture(
 			FRHITextureCreateDesc::Create2D(TEXT("DefaultTexture2D"), 1, 1, Format)
 			.SetFlags(ETextureCreateFlags::ShaderResource);
 
-		FTexture2DRHIRef Texture = RHICreateTexture(Desc);
+		FTextureRHIRef Texture = RHICreateTexture(Desc);
 		SetDefaultTextureData2D(Texture, Value);
 		Entry.Texture = CreateRenderTarget(Texture, Desc.DebugName);
 	}
@@ -1491,7 +1491,7 @@ FRDGTextureRef GetInternalDefaultTexture(
 			FRHITextureCreateDesc::Create2DArray(TEXT("DefaultTexture2DArray"), 1, 1, 1, Format)
 			.SetFlags(ETextureCreateFlags::ShaderResource);
 
-		FTexture2DArrayRHIRef Texture = RHICreateTexture(Desc);
+		FTextureRHIRef Texture = RHICreateTexture(Desc);
 		SetDefaultTextureData2DArray(Texture, Value);
 		Entry.Texture = CreateRenderTarget(Texture, Desc.DebugName);
 	}
@@ -1501,7 +1501,7 @@ FRDGTextureRef GetInternalDefaultTexture(
 			FRHITextureCreateDesc::Create3D(TEXT("DefaultTexture3D"), 1, 1, 1, Format)
 			.SetFlags(ETextureCreateFlags::ShaderResource);
 
-		FTexture3DRHIRef Texture = RHICreateTexture(Desc);
+		FTextureRHIRef Texture = RHICreateTexture(Desc);
 		SetDefaultTextureData3D(GraphBuilder.RHICmdList, Texture, Value);
 		Entry.Texture = CreateRenderTarget(Texture, Desc.DebugName);
 	}
@@ -1511,7 +1511,7 @@ FRDGTextureRef GetInternalDefaultTexture(
 			FRHITextureCreateDesc::CreateCube(TEXT("DefaultTextureCube"), 1, Format)
 			.SetFlags(ETextureCreateFlags::ShaderResource);
 
-		FTextureCubeRHIRef Texture = RHICreateTexture(Desc);
+		FTextureRHIRef Texture = RHICreateTexture(Desc);
 		SetDefaultTextureDataCube(Texture, Value);
 		Entry.Texture = CreateRenderTarget(Texture, Desc.DebugName);
 	}
@@ -1521,7 +1521,7 @@ FRDGTextureRef GetInternalDefaultTexture(
 			FRHITextureCreateDesc::CreateCubeArray(TEXT("DefaultTextureCubeArray"), 1, 1, Format)
 			.SetFlags(ETextureCreateFlags::ShaderResource);
 
-		FTextureCubeRHIRef Texture = RHICreateTexture(Desc);
+		FTextureRHIRef Texture = RHICreateTexture(Desc);
 		SetDefaultTextureDataCube(Texture, Value);
 		Entry.Texture = CreateRenderTarget(Texture, Desc.DebugName);
 	}
