Couple of improvements to Alloc helper functions in FRHICommandListBase.

diff --git a/Engine/Source/Runtime/RHI/Public/RHICommandList.h b/Engine/Source/Runtime/RHI/Public/RHICommandList.h
--- a/Engine/Source/Runtime/RHI/Public/RHICommandList.h
+++ b/Engine/Source/Runtime/RHI/Public/RHICommandList.h
@@ -465,9 +465,9 @@ public:
 	}
 
 	template <typename T>
-	FORCEINLINE_DEBUGGABLE void* Alloc()
+	FORCEINLINE_DEBUGGABLE T* Alloc()
 	{
-		return Alloc(sizeof(T), alignof(T));
+		return (T*)Alloc(sizeof(T), alignof(T));
 	}
 
 	template <typename T>
@@ -477,7 +477,7 @@ public:
 	}
 
 	template <typename T>
-	FORCEINLINE_DEBUGGABLE const TArrayView<T> AllocArray(const TArrayView<T> InArray)
+	FORCEINLINE_DEBUGGABLE const TArrayView<T> AllocArray(TConstArrayView<T> InArray)
 	{
 		if (InArray.Num() == 0)
 		{
