Remove RHISetRayTracingHitGroup/RHISetRayTracingCallableShader/RHISetRayTracingMissShader from IRHICommandContext.
- Same functionality is provided via the more general RHISetRayTracingBindings.

diff --git a/Engine/Source/Runtime/Apple/MetalRHI/Private/MetalRayTracing.cpp b/Engine/Source/Runtime/Apple/MetalRHI/Private/MetalRayTracing.cpp
--- a/Engine/Source/Runtime/Apple/MetalRHI/Private/MetalRayTracing.cpp
+++ b/Engine/Source/Runtime/Apple/MetalRHI/Private/MetalRayTracing.cpp
@@ -843,34 +843,6 @@ void FMetalRHICommandContext::RHIRayTraceDispatchIndirect(FRHIRayTracingPipeline
 	checkNoEntry();
 }
 
-void FMetalRHICommandContext::RHISetRayTracingHitGroup(
-	   FRHIRayTracingScene* InScene, uint32 InstanceIndex, uint32 SegmentIndex, uint32 ShaderSlot,
-	   FRHIRayTracingPipelineState* InPipeline, uint32 HitGroupIndex,
-	   uint32 NumUniformBuffers, FRHIUniformBuffer* const* UniformBuffers,
-	   uint32 LooseParameterDataSize, const void* LooseParameterData,
-	   uint32 UserData)
-{
-	checkNoEntry();
-}
-
-void FMetalRHICommandContext::RHISetRayTracingCallableShader(
-	   FRHIRayTracingScene* Scene, uint32 ShaderSlotInScene,
-	   FRHIRayTracingPipelineState* Pipeline, uint32 ShaderIndexInPipeline,
-	   uint32 NumUniformBuffers, FRHIUniformBuffer* const* UniformBuffers,
-	   uint32 UserData)
-{
-	checkNoEntry();
-}
-
-void FMetalRHICommandContext::RHISetRayTracingMissShader(
-	   FRHIRayTracingScene* InScene, uint32 ShaderSlotInScene,
-	   FRHIRayTracingPipelineState* InPipeline, uint32 ShaderIndexInPipeline,
-	   uint32 NumUniformBuffers, FRHIUniformBuffer* const* UniformBuffers,
-	   uint32 UserData)
-{
-	checkNoEntry();
-}
-
 void FMetalRHICommandContext::RHISetRayTracingBindings(
 	FRHIRayTracingScene* InScene, FRHIRayTracingPipelineState* InPipeline,
 	uint32 NumBindings, const FRayTracingLocalShaderBindings* Bindings,
diff --git a/Engine/Source/Runtime/Apple/MetalRHI/Public/MetalRHIContext.h b/Engine/Source/Runtime/Apple/MetalRHI/Public/MetalRHIContext.h
--- a/Engine/Source/Runtime/Apple/MetalRHI/Public/MetalRHIContext.h
+++ b/Engine/Source/Runtime/Apple/MetalRHI/Public/MetalRHIContext.h
@@ -187,22 +187,6 @@ public:
 		FRHIRayTracingScene* Scene,
 		const FRayTracingShaderBindings& GlobalResourceBindings,
 		FRHIBuffer* ArgumentBuffer, uint32 ArgumentOffset) final override;
-	virtual void RHISetRayTracingHitGroup(
-		FRHIRayTracingScene* Scene, uint32 InstanceIndex, uint32 SegmentIndex, uint32 ShaderSlot,
-		FRHIRayTracingPipelineState* Pipeline, uint32 HitGroupIndex,
-		uint32 NumUniformBuffers, FRHIUniformBuffer* const* UniformBuffers,
-		uint32 LooseParameterDataSize, const void* LooseParameterData,
-		uint32 UserData) final override;
-	virtual void RHISetRayTracingCallableShader(
-		FRHIRayTracingScene* Scene, uint32 ShaderSlotInScene,
-		FRHIRayTracingPipelineState* Pipeline, uint32 ShaderIndexInPipeline,
-		uint32 NumUniformBuffers, FRHIUniformBuffer* const* UniformBuffers,
-		uint32 UserData) final override;
-	virtual void RHISetRayTracingMissShader(
-		FRHIRayTracingScene* Scene, uint32 ShaderSlotInScene,
-		FRHIRayTracingPipelineState* Pipeline, uint32 ShaderIndexInPipeline,
-		uint32 NumUniformBuffers, FRHIUniformBuffer* const* UniformBuffers,
-		uint32 UserData) final override;
 	virtual void RHISetRayTracingBindings(
 		FRHIRayTracingScene* Scene, FRHIRayTracingPipelineState* Pipeline,
 		uint32 NumBindings, const FRayTracingLocalShaderBindings* Bindings,
diff --git a/Engine/Source/Runtime/D3D12RHI/Private/D3D12CommandContext.h b/Engine/Source/Runtime/D3D12RHI/Private/D3D12CommandContext.h
--- a/Engine/Source/Runtime/D3D12RHI/Private/D3D12CommandContext.h
+++ b/Engine/Source/Runtime/D3D12RHI/Private/D3D12CommandContext.h
@@ -708,22 +708,6 @@ public:
 	virtual void RHISetRayTracingBindings(
 		FRHIRayTracingScene* Scene, FRHIRayTracingPipelineState* Pipeline,
 		uint32 NumBindings, const FRayTracingLocalShaderBindings* Bindings, ERayTracingBindingType BindingType) final override;
-	virtual void RHISetRayTracingHitGroup(
-		FRHIRayTracingScene* Scene, uint32 InstanceIndex, uint32 SegmentIndex, uint32 ShaderSlot,
-		FRHIRayTracingPipelineState* Pipeline, uint32 HitGroupIndex,
-		uint32 NumUniformBuffers, FRHIUniformBuffer* const* UniformBuffers,
-		uint32 LooseParameterDataSize, const void* LooseParameterData,
-		uint32 UserData) final override;
-	virtual void RHISetRayTracingCallableShader(
-		FRHIRayTracingScene* Scene, uint32 ShaderSlotInScene,
-		FRHIRayTracingPipelineState* Pipeline, uint32 ShaderIndexInPipeline,
-		uint32 NumUniformBuffers, FRHIUniformBuffer* const* UniformBuffers,
-		uint32 UserData) final override;
-	virtual void RHISetRayTracingMissShader(
-		FRHIRayTracingScene* Scene, uint32 ShaderSlotInScene,
-		FRHIRayTracingPipelineState* Pipeline, uint32 ShaderIndexInPipeline,
-		uint32 NumUniformBuffers, FRHIUniformBuffer* const* UniformBuffers,
-		uint32 UserData) final override;
 #endif // D3D12_RHI_RAYTRACING
 
 	template<typename ObjectType, typename RHIType>
@@ -1065,39 +1049,11 @@ public:
 		ContextRedirect(RHIRayTraceDispatchIndirect(RayTracingPipelineState, RayGenShader, Scene, GlobalResourceBindings, ArgumentBuffer, ArgumentOffset));
 	}
 
-	virtual void RHISetRayTracingHitGroup(
-		FRHIRayTracingScene* Scene, uint32 InstanceIndex, uint32 SegmentIndex, uint32 ShaderSlot,
-		FRHIRayTracingPipelineState* Pipeline, uint32 HitGroupIndex,
-		uint32 NumUniformBuffers, FRHIUniformBuffer* const* UniformBuffers,
-		uint32 LooseParameterDataSize, const void* LooseParameterData,
-		uint32 UserData) final override
-	{
-		ContextRedirect(RHISetRayTracingHitGroup(Scene, InstanceIndex, SegmentIndex, ShaderSlot, Pipeline, HitGroupIndex, NumUniformBuffers, UniformBuffers, LooseParameterDataSize, LooseParameterData, UserData));
-	}
-
 	virtual void RHISetRayTracingBindings(FRHIRayTracingScene* Scene, FRHIRayTracingPipelineState* Pipeline, uint32 NumBindings, const FRayTracingLocalShaderBindings* Bindings, ERayTracingBindingType BindingType) final override
 	{
 		ContextRedirect(RHISetRayTracingBindings(Scene, Pipeline, NumBindings, Bindings, BindingType));
 	}
 
-	virtual void RHISetRayTracingCallableShader(
-		FRHIRayTracingScene* Scene, uint32 ShaderSlotInScene,
-		FRHIRayTracingPipelineState* Pipeline, uint32 ShaderIndexInPipeline,
-		uint32 NumUniformBuffers, FRHIUniformBuffer* const* UniformBuffers,
-		uint32 UserData) final override
-	{
-		ContextRedirect(RHISetRayTracingCallableShader(Scene, ShaderSlotInScene, Pipeline, ShaderIndexInPipeline, NumUniformBuffers, UniformBuffers, UserData));
-	}
-
-	virtual void RHISetRayTracingMissShader(
-		FRHIRayTracingScene* Scene, uint32 ShaderSlotInScene,
-		FRHIRayTracingPipelineState* Pipeline, uint32 ShaderIndexInPipeline,
-		uint32 NumUniformBuffers, FRHIUniformBuffer* const* UniformBuffers,
-		uint32 UserData) final override
-	{
-		ContextRedirect(RHISetRayTracingMissShader(Scene, ShaderSlotInScene, Pipeline, ShaderIndexInPipeline, NumUniformBuffers, UniformBuffers, UserData));
-	}
-
 	virtual void RHIClearRayTracingBindings(FRHIRayTracingScene* Scene) final override
 	{
 		ContextRedirect(RHIClearRayTracingBindings(Scene));
diff --git a/Engine/Source/Runtime/D3D12RHI/Private/D3D12RayTracing.cpp b/Engine/Source/Runtime/D3D12RHI/Private/D3D12RayTracing.cpp
--- a/Engine/Source/Runtime/D3D12RHI/Private/D3D12RayTracing.cpp
+++ b/Engine/Source/Runtime/D3D12RHI/Private/D3D12RayTracing.cpp
@@ -4958,64 +4958,4 @@ void FD3D12CommandContext::RHISetRayTracingBindings(
 	ParallelForWithExistingTaskContext(TEXT("SetRayTracingBindings"), MakeArrayView(TaskContexts), NumBindings, ItemsPerTask, BindingTask);
 }
 
-void FD3D12CommandContext::RHISetRayTracingHitGroup(
-	FRHIRayTracingScene* InScene, uint32 InstanceIndex, uint32 SegmentIndex, uint32 ShaderSlot,
-	FRHIRayTracingPipelineState* InPipeline, uint32 HitGroupIndex,
-	uint32 NumUniformBuffers, FRHIUniformBuffer* const* UniformBuffers,
-	uint32 LooseParameterDataSize, const void* LooseParameterData,
-	uint32 UserData)
-{
-	FD3D12RayTracingScene* Scene = FD3D12DynamicRHI::ResourceCast(InScene);
-	FD3D12RayTracingPipelineState* Pipeline = FD3D12DynamicRHI::ResourceCast(InPipeline);
-
-	checkf(ShaderSlot < Scene->Initializer.ShaderSlotsPerGeometrySegment, TEXT("Shader slot is invalid. Make sure that ShaderSlotsPerGeometrySegment is correct on FRayTracingSceneInitializer."));
-
-	FD3D12RayTracingShaderTable* ShaderTable = Scene->FindOrCreateShaderTable(Pipeline, GetParentDevice());
-
-	const uint32 WorkerIndex = 0;
-	SetRayTracingHitGroup(GetParentDevice(), ShaderTable, Scene, Pipeline,
-		InstanceIndex,
-		SegmentIndex,
-		ShaderSlot,
-		HitGroupIndex,
-		NumUniformBuffers,
-		UniformBuffers,
-		LooseParameterDataSize,
-		LooseParameterData,
-		UserData,
-		WorkerIndex);
-}
-
-void FD3D12CommandContext::RHISetRayTracingCallableShader(
-	FRHIRayTracingScene* InScene, uint32 ShaderSlotInScene,
-	FRHIRayTracingPipelineState* InPipeline, uint32 ShaderIndexInPipeline,
-	uint32 NumUniformBuffers, FRHIUniformBuffer* const* UniformBuffers,
-	uint32 UserData)
-{
-	FD3D12RayTracingScene* Scene = FD3D12DynamicRHI::ResourceCast(InScene);
-	FD3D12RayTracingPipelineState* Pipeline = FD3D12DynamicRHI::ResourceCast(InPipeline);
-	FD3D12RayTracingShaderTable* ShaderTable = Scene->FindOrCreateShaderTable(Pipeline, GetParentDevice());
-	const uint32 WorkerIndex = 0;
-
-	SetRayTracingCallableShader(GetParentDevice(), ShaderTable, Scene, Pipeline, ShaderSlotInScene, ShaderIndexInPipeline, NumUniformBuffers, UniformBuffers, 0, nullptr, UserData, WorkerIndex);
-}
-
-void FD3D12CommandContext::RHISetRayTracingMissShader(
-	FRHIRayTracingScene* InScene, uint32 ShaderSlotInScene,
-	FRHIRayTracingPipelineState* InPipeline, uint32 ShaderIndexInPipeline,
-	uint32 NumUniformBuffers, FRHIUniformBuffer* const* UniformBuffers,
-	uint32 UserData)
-{
-	FD3D12RayTracingScene* Scene = FD3D12DynamicRHI::ResourceCast(InScene);
-	FD3D12RayTracingPipelineState* Pipeline = FD3D12DynamicRHI::ResourceCast(InPipeline);
-	FD3D12RayTracingShaderTable* ShaderTable = Scene->FindOrCreateShaderTable(Pipeline, GetParentDevice());
-	const uint32 WorkerIndex = 0;
-	SetRayTracingMissShader(GetParentDevice(), ShaderTable, Scene, Pipeline,
-		ShaderSlotInScene, ShaderIndexInPipeline,
-		NumUniformBuffers, UniformBuffers,
-		0, nullptr, // Loose parameters
-		UserData,
-		WorkerIndex);
-}
-
 #endif // D3D12_RHI_RAYTRACING
diff --git a/Engine/Source/Runtime/RHI/Public/RHICommandList.h b/Engine/Source/Runtime/RHI/Public/RHICommandList.h
--- a/Engine/Source/Runtime/RHI/Public/RHICommandList.h
+++ b/Engine/Source/Runtime/RHI/Public/RHICommandList.h
@@ -2229,68 +2229,18 @@ FRHICOMMAND_MACRO(FRHICommandRayTraceDispatch)
 FRHICOMMAND_MACRO(FRHICommandSetRayTracingBindings)
 {
 	FRHIRayTracingScene* Scene = nullptr;
-	ERayTracingBindingType BindingType = ERayTracingBindingType::HitGroup;
-	uint32 InstanceIndex = 0;
-	uint32 SegmentIndex = 0;
-	uint32 ShaderSlot = 0;
 	FRayTracingPipelineState* Pipeline = nullptr;
-	uint32 ShaderIndex = 0;
-	uint32 NumUniformBuffers = 0;
-	FRHIUniformBuffer* const* UniformBuffers = nullptr; // Pointer to an array of uniform buffers, allocated inline within the command list
-	uint32 LooseParameterDataSize = 0;
-	const void* LooseParameterData = nullptr;
-	uint32 UserData = 0;
-
-	// Batched bindings
 	int32 NumBindings = -1;
 	const FRayTracingLocalShaderBindings* Bindings = nullptr;
-
-	// Hit group bindings
-	FRHICommandSetRayTracingBindings(FRHIRayTracingScene* InScene, uint32 InInstanceIndex, uint32 InSegmentIndex, uint32 InShaderSlot,
-		FRayTracingPipelineState* InPipeline, uint32 InHitGroupIndex, uint32 InNumUniformBuffers, FRHIUniformBuffer* const* InUniformBuffers,
-		uint32 InLooseParameterDataSize, const void* InLooseParameterData,
-		uint32 InUserData)
-		: Scene(InScene)
-		, BindingType(ERayTracingBindingType::HitGroup)
-		, InstanceIndex(InInstanceIndex)
-		, SegmentIndex(InSegmentIndex)
-		, ShaderSlot(InShaderSlot)
-		, Pipeline(InPipeline)
-		, ShaderIndex(InHitGroupIndex)
-		, NumUniformBuffers(InNumUniformBuffers)
-		, UniformBuffers(InUniformBuffers)
-		, LooseParameterDataSize(InLooseParameterDataSize)
-		, LooseParameterData(InLooseParameterData)
-		, UserData(InUserData)
-	{
-	}
-
-	// Callable and Miss shader bindings
-	FRHICommandSetRayTracingBindings(FRHIRayTracingScene* InScene, uint32 InShaderSlot,
-		FRayTracingPipelineState* InPipeline, uint32 InShaderIndex,
-		uint32 InNumUniformBuffers, FRHIUniformBuffer* const* InUniformBuffers,
-		uint32 InUserData, ERayTracingBindingType InBindingType)
-		: Scene(InScene)
-		, BindingType(InBindingType)
-		, InstanceIndex(0)
-		, SegmentIndex(0)
-		, ShaderSlot(InShaderSlot)
-		, Pipeline(InPipeline)
-		, ShaderIndex(InShaderIndex)
-		, NumUniformBuffers(InNumUniformBuffers)
-		, UniformBuffers(InUniformBuffers)
-		, UserData(InUserData)
-	{
-		checkf(InBindingType != ERayTracingBindingType::HitGroup, TEXT("Hit group bindings must specify Instance and Segment Index."));
-	}
+	ERayTracingBindingType BindingType = ERayTracingBindingType::HitGroup;
 
 	// Bindings Batch
 	FRHICommandSetRayTracingBindings(FRHIRayTracingScene* InScene, FRayTracingPipelineState* InPipeline, uint32 InNumBindings, const FRayTracingLocalShaderBindings* InBindings, ERayTracingBindingType InBindingType)
 		: Scene(InScene)
-		, BindingType(InBindingType)
 		, Pipeline(InPipeline)
 		, NumBindings(InNumBindings)
 		, Bindings(InBindings)
+		, BindingType(InBindingType)
 	{
 
 	}
@@ -3656,12 +3606,12 @@ public:
 		}
 		else
 		{
-			FRayTracingLocalShaderBindings* InlineBindings = nullptr;
-
 			// By default all batch binding data is stored in the command list memory.
 			// However, user may skip this copy if they take responsibility for keeping data alive until this command is executed.
 			if (bCopyDataToInlineStorage)
 			{
+				FRayTracingLocalShaderBindings* InlineBindings = nullptr;
+
 				if (NumBindings)
 				{
 					uint32 Size = sizeof(FRayTracingLocalShaderBindings) * NumBindings;
@@ -3727,37 +3677,34 @@ public:
 		uint32 LooseParameterDataSize, const void* LooseParameterData,
 		uint32 UserData)
 	{
-		if (Bypass())
-		{
-			GetContext().RHISetRayTracingHitGroup(Scene, InstanceIndex, SegmentIndex, ShaderSlot, GetRHIRayTracingPipelineState(Pipeline), HitGroupIndex, 
-				NumUniformBuffers, UniformBuffers,
-				LooseParameterDataSize, LooseParameterData,
-				UserData);
-		}
-		else
-		{
-			FRHIUniformBuffer** InlineUniformBuffers = nullptr;
-			if (NumUniformBuffers)
-			{
-				InlineUniformBuffers = (FRHIUniformBuffer**)Alloc(sizeof(FRHIUniformBuffer*) * NumUniformBuffers, alignof(FRHIUniformBuffer*));
-				for (uint32 Index = 0; Index < NumUniformBuffers; ++Index)
-				{
-					InlineUniformBuffers[Index] = UniformBuffers[Index];
-				}
-			}
+		check(NumUniformBuffers <= UINT16_MAX);
+		check(LooseParameterDataSize <= UINT16_MAX);
 
-			void* InlineLooseParameterData = nullptr;
-			if (LooseParameterDataSize)
+		FRayTracingLocalShaderBindings* InlineBindings = Alloc<FRayTracingLocalShaderBindings>();
+		InlineBindings->InstanceIndex = InstanceIndex;
+		InlineBindings->SegmentIndex = SegmentIndex;
+		InlineBindings->ShaderSlot = ShaderSlot;
+		InlineBindings->ShaderIndexInPipeline = HitGroupIndex;
+		InlineBindings->UserData = UserData;
+		InlineBindings->NumUniformBuffers = (uint16)NumUniformBuffers;
+		InlineBindings->LooseParameterDataSize = (uint16)LooseParameterDataSize;
+
+		if (NumUniformBuffers)
+		{
+			InlineBindings->UniformBuffers = (FRHIUniformBuffer**)Alloc(sizeof(FRHIUniformBuffer*) * NumUniformBuffers, alignof(FRHIUniformBuffer*));
+			for (uint32 Index = 0; Index < NumUniformBuffers; ++Index)
 			{
-				InlineLooseParameterData = Alloc(LooseParameterDataSize, 16);
-				FMemory::Memcpy(InlineLooseParameterData, LooseParameterData, LooseParameterDataSize);
+				InlineBindings->UniformBuffers[Index] = UniformBuffers[Index];
 			}
+		}
 
-			ALLOC_COMMAND(FRHICommandSetRayTracingBindings)(Scene, InstanceIndex, SegmentIndex, ShaderSlot, Pipeline, HitGroupIndex, 
-				NumUniformBuffers, InlineUniformBuffers, 
-				LooseParameterDataSize, InlineLooseParameterData,
-				UserData);
+		if (LooseParameterDataSize)
+		{
+			InlineBindings->LooseParameterData = (uint8*)Alloc(LooseParameterDataSize, 16);
+			FMemory::Memcpy(InlineBindings->LooseParameterData, LooseParameterData, LooseParameterDataSize);
 		}
+
+		SetRayTracingBindings(Scene, Pipeline, 1, InlineBindings, ERayTracingBindingType::HitGroup, /*bCopyDataToInlineStorage*/ false);
 	}
 
 	FORCEINLINE_DEBUGGABLE void SetRayTracingCallableShader(
@@ -3766,24 +3713,22 @@ public:
 		uint32 NumUniformBuffers, FRHIUniformBuffer* const* UniformBuffers,
 		uint32 UserData)
 	{
-		if (Bypass())
-		{
-			GetContext().RHISetRayTracingCallableShader(Scene, ShaderSlotInScene, GetRHIRayTracingPipelineState(Pipeline), ShaderIndexInPipeline, NumUniformBuffers, UniformBuffers, UserData);
-		}
-		else
+		FRayTracingLocalShaderBindings* InlineBindings = Alloc<FRayTracingLocalShaderBindings>();
+		InlineBindings->ShaderSlot = ShaderSlotInScene;
+		InlineBindings->ShaderIndexInPipeline = ShaderIndexInPipeline;
+		InlineBindings->UserData = UserData;
+		InlineBindings->NumUniformBuffers = (uint16)NumUniformBuffers;
+
+		if (NumUniformBuffers)
 		{
-			FRHIUniformBuffer** InlineUniformBuffers = nullptr;
-			if (NumUniformBuffers)
+			InlineBindings->UniformBuffers = (FRHIUniformBuffer**)Alloc(sizeof(FRHIUniformBuffer*) * NumUniformBuffers, alignof(FRHIUniformBuffer*));
+			for (uint32 Index = 0; Index < NumUniformBuffers; ++Index)
 			{
-				InlineUniformBuffers = (FRHIUniformBuffer**)Alloc(sizeof(FRHIUniformBuffer*) * NumUniformBuffers, alignof(FRHIUniformBuffer*));
-				for (uint32 Index = 0; Index < NumUniformBuffers; ++Index)
-				{
-					InlineUniformBuffers[Index] = UniformBuffers[Index];
-				}
+				InlineBindings->UniformBuffers[Index] = UniformBuffers[Index];
 			}
-
-			ALLOC_COMMAND(FRHICommandSetRayTracingBindings)(Scene, ShaderSlotInScene, Pipeline, ShaderIndexInPipeline, NumUniformBuffers, InlineUniformBuffers, UserData, ERayTracingBindingType::CallableShader);
 		}
+
+		SetRayTracingBindings(Scene, Pipeline, 1, InlineBindings, ERayTracingBindingType::CallableShader, /*bCopyDataToInlineStorage*/ false);
 	}
 
 	FORCEINLINE_DEBUGGABLE void SetRayTracingMissShader(
@@ -3792,24 +3737,22 @@ public:
 		uint32 NumUniformBuffers, FRHIUniformBuffer* const* UniformBuffers,
 		uint32 UserData)
 	{
-		if (Bypass())
-		{
-			GetContext().RHISetRayTracingMissShader(Scene, ShaderSlotInScene, GetRHIRayTracingPipelineState(Pipeline), ShaderIndexInPipeline, NumUniformBuffers, UniformBuffers, UserData);
-		}
-		else
+		FRayTracingLocalShaderBindings* InlineBindings = Alloc<FRayTracingLocalShaderBindings>();
+		InlineBindings->ShaderSlot = ShaderSlotInScene;
+		InlineBindings->ShaderIndexInPipeline = ShaderIndexInPipeline;
+		InlineBindings->UserData = UserData;
+		InlineBindings->NumUniformBuffers = (uint16)NumUniformBuffers;
+
+		if (NumUniformBuffers)
 		{
-			FRHIUniformBuffer** InlineUniformBuffers = nullptr;
-			if (NumUniformBuffers)
+			InlineBindings->UniformBuffers = (FRHIUniformBuffer**)Alloc(sizeof(FRHIUniformBuffer*) * NumUniformBuffers, alignof(FRHIUniformBuffer*));
+			for (uint32 Index = 0; Index < NumUniformBuffers; ++Index)
 			{
-				InlineUniformBuffers = (FRHIUniformBuffer**)Alloc(sizeof(FRHIUniformBuffer*) * NumUniformBuffers, alignof(FRHIUniformBuffer*));
-				for (uint32 Index = 0; Index < NumUniformBuffers; ++Index)
-				{
-					InlineUniformBuffers[Index] = UniformBuffers[Index];
-				}
+				InlineBindings->UniformBuffers[Index] = UniformBuffers[Index];
 			}
-
-			ALLOC_COMMAND(FRHICommandSetRayTracingBindings)(Scene, ShaderSlotInScene, Pipeline, ShaderIndexInPipeline, NumUniformBuffers, InlineUniformBuffers, UserData, ERayTracingBindingType::MissShader);
 		}
+
+		SetRayTracingBindings(Scene, Pipeline, 1, InlineBindings, ERayTracingBindingType::MissShader, /*bCopyDataToInlineStorage*/ false);
 	}
 
 #endif // RHI_RAYTRACING
diff --git a/Engine/Source/Runtime/RHI/Public/RHICommandListCommandExecutes.inl b/Engine/Source/Runtime/RHI/Public/RHICommandListCommandExecutes.inl
--- a/Engine/Source/Runtime/RHI/Public/RHICommandListCommandExecutes.inl
+++ b/Engine/Source/Runtime/RHI/Public/RHICommandListCommandExecutes.inl
@@ -517,31 +517,7 @@ void FRHICommandRayTraceDispatch::Execute(FRHICommandListBase& CmdList)
 void FRHICommandSetRayTracingBindings::Execute(FRHICommandListBase& CmdList)
 {
 	RHISTAT(SetRayTracingHitGroup);
-	extern RHI_API FRHIRayTracingPipelineState* GetRHIRayTracingPipelineState(FRayTracingPipelineState*);
-	if (NumBindings >= 0)
-	{
-		check(Bindings != nullptr);
-		INTERNAL_DECORATOR(RHISetRayTracingBindings)(Scene, GetRHIRayTracingPipelineState(Pipeline), NumBindings, Bindings, BindingType);
-	}
-	else if (BindingType == ERayTracingBindingType::HitGroup)
-	{
-		INTERNAL_DECORATOR(RHISetRayTracingHitGroup)(Scene, InstanceIndex, SegmentIndex, ShaderSlot, GetRHIRayTracingPipelineState(Pipeline), ShaderIndex,
-			NumUniformBuffers, UniformBuffers,
-			LooseParameterDataSize, LooseParameterData,
-			UserData);
-	}
-	else if (BindingType == ERayTracingBindingType::CallableShader)
-	{
-		INTERNAL_DECORATOR(RHISetRayTracingCallableShader)(Scene, ShaderSlot, GetRHIRayTracingPipelineState(Pipeline), ShaderIndex, NumUniformBuffers, UniformBuffers, UserData);
-	}
-	else if(BindingType == ERayTracingBindingType::MissShader)
-	{
-		INTERNAL_DECORATOR(RHISetRayTracingMissShader)(Scene, ShaderSlot, GetRHIRayTracingPipelineState(Pipeline), ShaderIndex, NumUniformBuffers, UniformBuffers, UserData);
-	}
-	else
-	{
-		checkNoEntry();
-	}
+	INTERNAL_DECORATOR(RHISetRayTracingBindings)(Scene, GetRHIRayTracingPipelineState(Pipeline), NumBindings, Bindings, BindingType);
 }
 
 #endif // RHI_RAYTRACING
diff --git a/Engine/Source/Runtime/RHI/Public/RHIContext.h b/Engine/Source/Runtime/RHI/Public/RHIContext.h
--- a/Engine/Source/Runtime/RHI/Public/RHIContext.h
+++ b/Engine/Source/Runtime/RHI/Public/RHIContext.h
@@ -801,34 +801,6 @@ public:
 		checkNoEntry();
 	}
 
-	virtual void RHISetRayTracingHitGroup(
-		FRHIRayTracingScene* Scene, uint32 InstanceIndex, uint32 SegmentIndex, uint32 ShaderSlot,
-		FRHIRayTracingPipelineState* Pipeline, uint32 HitGroupIndex,
-		uint32 NumUniformBuffers, FRHIUniformBuffer* const* UniformBuffers,
-		uint32 LooseParameterDataSize, const void* LooseParameterData,
-		uint32 UserData)
-	{
-		checkNoEntry();
-	}
-
-	virtual void RHISetRayTracingCallableShader(
-		FRHIRayTracingScene* Scene, uint32 ShaderSlotInScene,
-		FRHIRayTracingPipelineState* Pipeline, uint32 ShaderIndexInPipeline,
-		uint32 NumUniformBuffers, FRHIUniformBuffer* const* UniformBuffers,
-		uint32 UserData)
-	{
-		checkNoEntry();
-	}
-
-	virtual void RHISetRayTracingMissShader(
-		FRHIRayTracingScene* Scene, uint32 ShaderSlotInScene,
-		FRHIRayTracingPipelineState* Pipeline, uint32 ShaderIndexInPipeline,
-		uint32 NumUniformBuffers, FRHIUniformBuffer* const* UniformBuffers,
-		uint32 UserData)
-	{
-		checkNoEntry();
-	}
-
 protected:
 	FRHIRenderPassInfo RenderPassInfo;
 };
diff --git a/Engine/Source/Runtime/RHI/Public/RHIValidationContext.h b/Engine/Source/Runtime/RHI/Public/RHIValidationContext.h
--- a/Engine/Source/Runtime/RHI/Public/RHIValidationContext.h
+++ b/Engine/Source/Runtime/RHI/Public/RHIValidationContext.h
@@ -1118,34 +1118,6 @@ public:
 		RHIContext->RHISetRayTracingBindings(Scene, Pipeline, NumBindings, Bindings, BindingType);
 	}
 
-	virtual void RHISetRayTracingHitGroup(
-		FRHIRayTracingScene* Scene, uint32 InstanceIndex, uint32 SegmentIndex, uint32 ShaderSlot,
-		FRHIRayTracingPipelineState* Pipeline, uint32 HitGroupIndex,
-		uint32 NumUniformBuffers, FRHIUniformBuffer* const* UniformBuffers,
-		uint32 LooseParameterDataSize, const void* LooseParameterData,
-		uint32 UserData) override final
-	{
-		RHIContext->RHISetRayTracingHitGroup(Scene, InstanceIndex, SegmentIndex, ShaderSlot, Pipeline, HitGroupIndex, NumUniformBuffers, UniformBuffers, LooseParameterDataSize, LooseParameterData, UserData);
-	}
-
-	virtual void RHISetRayTracingCallableShader(
-		FRHIRayTracingScene* Scene, uint32 ShaderSlotInScene,
-		FRHIRayTracingPipelineState* Pipeline, uint32 ShaderIndexInPipeline,
-		uint32 NumUniformBuffers, FRHIUniformBuffer* const* UniformBuffers,
-		uint32 UserData) override final
-	{
-		RHIContext->RHISetRayTracingCallableShader(Scene, ShaderSlotInScene, Pipeline, ShaderIndexInPipeline, NumUniformBuffers, UniformBuffers, UserData);
-	}
-
-	virtual void RHISetRayTracingMissShader(
-		FRHIRayTracingScene* Scene, uint32 ShaderSlotInScene,
-		FRHIRayTracingPipelineState* Pipeline, uint32 ShaderIndexInPipeline,
-		uint32 NumUniformBuffers, FRHIUniformBuffer* const* UniformBuffers,
-		uint32 UserData) override final
-	{
-		RHIContext->RHISetRayTracingMissShader(Scene, ShaderSlotInScene, Pipeline, ShaderIndexInPipeline, NumUniformBuffers, UniformBuffers, UserData);
-	}
-
 	void SetupDrawing()
 	{
 		// nothing to validate right now
diff --git a/Engine/Source/Runtime/VulkanRHI/Private/VulkanContext.h b/Engine/Source/Runtime/VulkanRHI/Private/VulkanContext.h
--- a/Engine/Source/Runtime/VulkanRHI/Private/VulkanContext.h
+++ b/Engine/Source/Runtime/VulkanRHI/Private/VulkanContext.h
@@ -122,22 +122,6 @@ public:
 		FRHIRayTracingScene* Scene,
 		const FRayTracingShaderBindings& GlobalResourceBindings,
 		FRHIBuffer* ArgumentBuffer, uint32 ArgumentOffset) final override;
-	virtual void RHISetRayTracingHitGroup(
-		FRHIRayTracingScene* Scene, uint32 InstanceIndex, uint32 SegmentIndex, uint32 ShaderSlot,
-		FRHIRayTracingPipelineState* Pipeline, uint32 HitGroupIndex,
-		uint32 NumUniformBuffers, FRHIUniformBuffer* const* UniformBuffers,
-		uint32 LooseParameterDataSize, const void* LooseParameterData,
-		uint32 UserData) final override;
-	virtual void RHISetRayTracingCallableShader(
-		FRHIRayTracingScene* Scene, uint32 ShaderSlotInScene,
-		FRHIRayTracingPipelineState* Pipeline, uint32 ShaderIndexInPipeline,
-		uint32 NumUniformBuffers, FRHIUniformBuffer* const* UniformBuffers,
-		uint32 UserData) final override;
-	virtual void RHISetRayTracingMissShader(
-		FRHIRayTracingScene* Scene, uint32 ShaderSlotInScene,
-		FRHIRayTracingPipelineState* Pipeline, uint32 ShaderIndexInPipeline,
-		uint32 NumUniformBuffers, FRHIUniformBuffer* const* UniformBuffers,
-		uint32 UserData) final override;
 	virtual void RHISetRayTracingBindings(
 		FRHIRayTracingScene* Scene, FRHIRayTracingPipelineState* Pipeline,
 		uint32 NumBindings, const FRayTracingLocalShaderBindings* Bindings,
diff --git a/Engine/Source/Runtime/VulkanRHI/Private/VulkanRayTracing.cpp b/Engine/Source/Runtime/VulkanRHI/Private/VulkanRayTracing.cpp
--- a/Engine/Source/Runtime/VulkanRHI/Private/VulkanRayTracing.cpp
+++ b/Engine/Source/Runtime/VulkanRHI/Private/VulkanRayTracing.cpp
@@ -2154,36 +2154,6 @@ static void SetRayTracingHitGroup(
 }
 
 
-void FVulkanCommandListContext::RHISetRayTracingHitGroup(
-	FRHIRayTracingScene* InScene, uint32 InstanceIndex, uint32 SegmentIndex, uint32 ShaderSlot,
-	FRHIRayTracingPipelineState* InPipeline, uint32 HitGroupIndex,
-	uint32 NumUniformBuffers, FRHIUniformBuffer* const* UniformBuffers,
-	uint32 LooseParameterDataSize, const void* LooseParameterData,
-	uint32 UserData)
-{
-	FVulkanRayTracingScene* Scene = ResourceCast(InScene);
-	FVulkanRayTracingPipelineState* Pipeline = ResourceCast(InPipeline);
-	FVulkanRayTracingShaderTable* ShaderTable = Scene->FindOrCreateShaderTable(Pipeline);
-
-	checkf(ShaderSlot < Scene->GetInitializer().ShaderSlotsPerGeometrySegment, TEXT("Shader slot is invalid. Make sure that ShaderSlotsPerGeometrySegment is correct on FRayTracingSceneInitializer."));
-
-	const uint32 WorkerIndex = 0;
-
-	SetRayTracingHitGroup(*this, Device, ShaderTable, Scene, Pipeline,
-		InstanceIndex,
-		SegmentIndex,
-		ShaderSlot,
-		HitGroupIndex,
-		NumUniformBuffers,
-		UniformBuffers,
-		LooseParameterDataSize,
-		LooseParameterData,
-		UserData,
-		WorkerIndex);
-}
-
-
-
 static void SetGenericSystemParameters(
 	FRHIRayTracingScene* InScene, uint32 ShaderSlotInScene,
 	FRHIRayTracingPipelineState* InPipeline, uint32 ShaderIndexInPipeline,
@@ -2205,26 +2175,6 @@ static void SetGenericSystemParameters(
 	ShaderTable->SetSlot(ShaderFrequency, ShaderSlotInScene, ShaderIndexInPipeline, Pipeline->GetShaderHandles(ShaderFrequency));
 }
 
-void FVulkanCommandListContext::RHISetRayTracingCallableShader(
-	FRHIRayTracingScene* InScene, uint32 ShaderSlotInScene,
-	FRHIRayTracingPipelineState* InPipeline, uint32 ShaderIndexInPipeline,
-	uint32 NumUniformBuffers, FRHIUniformBuffer* const* UniformBuffers,
-	uint32 UserData)
-{
-	SetGenericSystemParameters(InScene, ShaderSlotInScene, InPipeline, ShaderIndexInPipeline,
-		NumUniformBuffers, UniformBuffers, UserData, SF_RayCallable);
-}
-
-void FVulkanCommandListContext::RHISetRayTracingMissShader(
-	FRHIRayTracingScene* InScene, uint32 ShaderSlotInScene,
-	FRHIRayTracingPipelineState* InPipeline, uint32 ShaderIndexInPipeline,
-	uint32 NumUniformBuffers, FRHIUniformBuffer* const* UniformBuffers,
-	uint32 UserData)
-{
-	SetGenericSystemParameters(InScene, ShaderSlotInScene, InPipeline, ShaderIndexInPipeline,
-		NumUniformBuffers, UniformBuffers, UserData, SF_RayMiss);
-}
-
 
 void FVulkanCommandListContext::RHISetRayTracingBindings(
 	FRHIRayTracingScene* InScene, FRHIRayTracingPipelineState* InPipeline,
