Use RHI breadcrumbs and custom PSO precache event names to show information in Insight about PSOs being compiled (both during PSO precaching and runtime PSO compiles)

diff --git a/Engine/Plugins/FX/Niagara/Source/Niagara/Private/NiagaraEmitter.cpp b/Engine/Plugins/FX/Niagara/Source/Niagara/Private/NiagaraEmitter.cpp
--- a/Engine/Plugins/FX/Niagara/Source/Niagara/Private/NiagaraEmitter.cpp
+++ b/Engine/Plugins/FX/Niagara/Source/Niagara/Private/NiagaraEmitter.cpp
@@ -3554,7 +3554,7 @@ FGraphEventArray FVersionedNiagaraEmitterData::PrecacheComputePSOs(const UNiagar
 		for (int32 i = 0; i < ShaderScript->GetNumPermutations(); ++i)
 		{
 			FRHIComputeShader* ComputeShader = ShaderScript->GetShaderGameThread(i).GetComputeShader();
-			FPSOPrecacheRequestResult PSOPrecacheRequestResult = PipelineStateCache::PrecacheComputePipelineState(ComputeShader, true);
+			FPSOPrecacheRequestResult PSOPrecacheRequestResult = PipelineStateCache::PrecacheComputePipelineState(ComputeShader, TEXT("NiagaraEmitter"), true);
 			if (PSOPrecacheRequestResult.AsyncCompileEvent.IsValid())
 			{
 				PSOPrecacheEvents.Add(PSOPrecacheRequestResult.AsyncCompileEvent);
@@ -3634,7 +3634,7 @@ FGraphEventArray FVersionedNiagaraEmitterData::PrecacheComputePSOs(const UNiagar
 		for (int32 i=0; i < ShaderScript->GetNumPermutations(); ++i)
 		{
 			FRHIComputeShader* ComputeShader = ShaderScript->GetShaderGameThread(i).GetComputeShader();
-			FPSOPrecacheRequestResult PSOPrecacheRequestResult = PipelineStateCache::PrecacheComputePipelineState(ComputeShader, true);
+			FPSOPrecacheRequestResult PSOPrecacheRequestResult = PipelineStateCache::PrecacheComputePipelineState(ComputeShader, TEXT("NiagaraEmitter"), true);
 			if (PSOPrecacheRequestResult.AsyncCompileEvent != nullptr && bAddToPSOReadyGraphTasks)
 			{
 				PSOReadyGraphTasks.Add(PSOPrecacheRequestResult.AsyncCompileEvent);
diff --git a/Engine/Plugins/FX/Niagara/Source/NiagaraShader/Private/NiagaraShared.cpp b/Engine/Plugins/FX/Niagara/Source/NiagaraShader/Private/NiagaraShared.cpp
--- a/Engine/Plugins/FX/Niagara/Source/NiagaraShader/Private/NiagaraShared.cpp
+++ b/Engine/Plugins/FX/Niagara/Source/NiagaraShader/Private/NiagaraShared.cpp
@@ -535,7 +535,7 @@ void FNiagaraShaderScript::UpdateCachedData_PostCompile(bool bCalledFromSerializ
 					check(NiagaraShader->GetFrequency() == SF_Compute);
 					FRHIShader* RHIShader = GameThreadShaderMap->GetResource()->GetShader(Shader->GetResourceIndex());
 					FRHIComputeShader* RHIComputeShader = static_cast<FRHIComputeShader*>(RHIShader);
-					PipelineStateCache::PrecacheComputePipelineState(RHIComputeShader);
+					PipelineStateCache::PrecacheComputePipelineState(RHIComputeShader, TEXT("NiagaraCompute"));
 				}
 			}
 		}
diff --git a/Engine/Source/Runtime/Engine/Private/ShaderCompiler/ShaderCompiler.cpp b/Engine/Source/Runtime/Engine/Private/ShaderCompiler/ShaderCompiler.cpp
--- a/Engine/Source/Runtime/Engine/Private/ShaderCompiler/ShaderCompiler.cpp
+++ b/Engine/Source/Runtime/Engine/Private/ShaderCompiler/ShaderCompiler.cpp
@@ -9524,8 +9524,9 @@ void PrecacheComputePipelineStatesForGlobalShaders(EShaderPlatform Platform, con
 			{
 				for (TShaderRef<FShader> GlobalShader : ComputeShadersToPrecache)
 				{
+					const TCHAR* TypeName = GlobalShader.GetType()->GetName();
 					FRHIComputeShader* RHIComputeShader = GlobalShader.GetComputeShader();
-					PipelineStateCache::PrecacheComputePipelineState(RHIComputeShader);
+					PipelineStateCache::PrecacheComputePipelineState(RHIComputeShader, TypeName);
 				}
 			});
 	}
diff --git a/Engine/Source/Runtime/RHI/Private/PipelineStateCache.cpp b/Engine/Source/Runtime/RHI/Private/PipelineStateCache.cpp
--- a/Engine/Source/Runtime/RHI/Private/PipelineStateCache.cpp
+++ b/Engine/Source/Runtime/RHI/Private/PipelineStateCache.cpp
@@ -16,6 +16,7 @@ PipelineStateCache.cpp: Pipeline state cache implementation.
 #include "RHICommandList.h"
 #include "RHIFwd.h"
 #include "RHIImmutableSamplerState.h"
+#include "RHIBreadcrumbs.h"
 #include "Stats/StatsTrace.h"
 #include "Templates/TypeHash.h"
 
@@ -180,7 +181,18 @@ static FAutoConsoleCommand DumpPipelineCmd(
 	FConsoleCommandDelegate::CreateStatic(DumpPipelineCacheStats)
 );
 
-static inline void CheckAndUpdateHitchCountStat(FPSOPrecacheRequestID::EType PSOType, bool bIsRuntimePSO, uint64 StartTime)
+static int32 GraphicsPSOCreationHitchCount = 0;
+static int32 ComputePSOCreationHitchCount = 0;
+
+struct FPSOCompilationDebugData
+{
+	FString PSOCompilationEventName;
+#if WITH_RHI_BREADCRUMBS
+	FRHIBreadcrumbNode* BreadcrumbNode = nullptr;
+#endif // WITH_RHI_BREADCRUMBS
+};
+
+static inline void CheckAndUpdateHitchCountStat(FPSOPrecacheRequestID::EType PSOType, bool bIsRuntimePSO, const FPSOCompilationDebugData& PSOCompilationDebugData, uint64 StartTime)
 {
 	if (bIsRuntimePSO)
 	{
@@ -190,11 +202,25 @@ static inline void CheckAndUpdateHitchCountStat(FPSOPrecacheRequestID::EType PSO
 		{
 			if (PSOType == FPSOPrecacheRequestID::EType::Graphics)
 			{
+#if WITH_RHI_BREADCRUMBS
+				UE_LOG(LogRHI, Log, TEXT("Runtime graphics PSO creation hitch (%.2f msec) for %s - Breadcrumbs: %s"), PSOCreationTimeMs, *PSOCompilationDebugData.PSOCompilationEventName, PSOCompilationDebugData.BreadcrumbNode ? *PSOCompilationDebugData.BreadcrumbNode->GetFullPath() : TEXT("Unknown"));
+#else
+				UE_LOG(LogRHI, Log, TEXT("Runtime graphics PSO creation hitch (%.2f msec) for %s"), PSOCreationTimeMs, *PSOCompilationDebugData.PSOCompilationEventName);
+#endif // WITH_RHI_BREADCRUMBS
 				INC_DWORD_STAT(STAT_RuntimeGraphicsPSOHitchCount);
+				CSV_CUSTOM_STAT(PSO, GraphicsPSOHitch, 1, ECsvCustomStatOp::Accumulate);
+				GraphicsPSOCreationHitchCount++;
 			}
 			else if (PSOType == FPSOPrecacheRequestID::EType::Compute)
 			{
-				INC_DWORD_STAT(STAT_RuntimeComputePSOHitchCount)
+#if WITH_RHI_BREADCRUMBS
+				UE_LOG(LogRHI, Log, TEXT("Runtime compute PSO creation hitch (%.2f msec) for %s - Breadcrumbs: %s"), PSOCreationTimeMs, *PSOCompilationDebugData.PSOCompilationEventName, PSOCompilationDebugData.BreadcrumbNode ? *PSOCompilationDebugData.BreadcrumbNode->GetFullPath() : TEXT("Unknown"));
+#else
+				UE_LOG(LogRHI, Log, TEXT("Runtime compute PSO creation hitch (%.2f msec) for %s"), PSOCreationTimeMs, *PSOCompilationDebugData.PSOCompilationEventName);
+#endif // WITH_RHI_BREADCRUMBS
+				INC_DWORD_STAT(STAT_RuntimeComputePSOHitchCount);
+				CSV_CUSTOM_STAT(PSO, ComputePSOHitch, 1, ECsvCustomStatOp::Accumulate);
+				ComputePSOCreationHitchCount++;
 			}
 		}
 	}
@@ -574,8 +600,8 @@ public:
 class FComputePipelineState : public FPipelineState
 {
 public:
-	FComputePipelineState(FRHIComputeShader* InComputeShader)
-		: ComputeShader(InComputeShader)
+	FComputePipelineState(FRHIComputeShader* InComputeShader, const TCHAR* InName)
+		: ComputeShader(InComputeShader), Name(InName)
 	{
 		ComputeShader->AddRef();
 	}
@@ -614,6 +640,7 @@ public:
 	}
 
 	FRHIComputeShader* ComputeShader;
+	const TCHAR* Name;
 	TRefCountPtr<FRHIComputePipelineState> RHIPipeline;
 #if PIPELINESTATECACHE_VERIFYTHREADSAFE
 	FThreadSafeCounter InUseCount;
@@ -1190,7 +1217,7 @@ protected:
 		UpdateHighPriorityCompileCount(true /*Increment*/);
 	}
 
-	FPSOPrecacheRequestResult TryAddNewState(const TPrecachedPSOInitializer& Initializer, bool bDoAsyncCompile)
+	FPSOPrecacheRequestResult TryAddNewState(const TPrecachedPSOInitializer& Initializer, const FString& PSOCompilationEventName, bool bDoAsyncCompile)
 	{
 		FPSOPrecacheRequestResult Result;
 		uint32 InitializerHash = TPrecachePipelineCacheDerived::PipelineStateInitializerHash(Initializer);
@@ -1238,7 +1265,7 @@ protected:
 			}
 		}
 
-		TPrecachePipelineCacheDerived::OnNewPipelineStateCreated(Initializer, NewPipelineState, bDoAsyncCompile);
+		TPrecachePipelineCacheDerived::OnNewPipelineStateCreated(Initializer, NewPipelineState, PSOCompilationEventName, bDoAsyncCompile);
 
 		// A boost request might have been issued while we were kicking the task, need to check it here
 		{
@@ -1519,11 +1546,12 @@ struct FPrecacheComputeInitializer
 	{}
 
 
-	FPrecacheComputeInitializer(const FRHIComputeShader* InRHIComputeShader)
-		: RHIComputeShaderAsU64((uint64)InRHIComputeShader)
+	FPrecacheComputeInitializer(const FRHIComputeShader* InRHIComputeShader, const TCHAR* InName)
+		: RHIComputeShaderAsU64((uint64)InRHIComputeShader), Name(InName)
 	{}
 
 	uint64 RHIComputeShaderAsU64; // Using a U64 here rather than FRHIComputeShader* because we keep FPrecacheComputeInitializer but the FRHIComputeShader might get deleted
+	const TCHAR* Name;
 };
 
 class FPrecacheComputePipelineCache : public TPrecachePipelineCacheBase<FPrecacheComputePipelineCache, FPrecacheComputeInitializer, FComputePipelineState>
@@ -1531,12 +1559,12 @@ class FPrecacheComputePipelineCache : public TPrecachePipelineCacheBase<FPrecach
 public:
 	static FComputePipelineState* CreateNewPSO(const FPrecacheComputeInitializer& ComputeShaderInitializer)
 	{
-		return new FComputePipelineState((FRHIComputeShader*)ComputeShaderInitializer.RHIComputeShaderAsU64);
+		return new FComputePipelineState((FRHIComputeShader*)ComputeShaderInitializer.RHIComputeShaderAsU64, ComputeShaderInitializer.Name);
 	}
 
 	FPrecacheComputePipelineCache() : TPrecachePipelineCacheBase(FPSOPrecacheRequestID::EType::Compute) {}
-	FPSOPrecacheRequestResult PrecacheComputePipelineState(FRHIComputeShader* ComputeShader, bool bForcePrecache);
-	static void OnNewPipelineStateCreated(const FPrecacheComputeInitializer& ComputeInitializer, FComputePipelineState* NewComputePipelineState, bool bDoAsyncCompile);
+	FPSOPrecacheRequestResult PrecacheComputePipelineState(FRHIComputeShader* ComputeShader, const TCHAR* Name, bool bForcePrecache);
+	static void OnNewPipelineStateCreated(const FPrecacheComputeInitializer& ComputeInitializer, FComputePipelineState* NewComputePipelineState, const FString& PSOCompilationEventName, bool bDoAsyncCompile);
  
 	static const FName GetActiveCompileStatName()
 	{
@@ -1591,7 +1619,7 @@ public:
 	FPrecacheGraphicsPipelineCache() : TPrecachePipelineCacheBase(FPSOPrecacheRequestID::EType::Graphics) {}
 	FPSOPrecacheRequestResult PrecacheGraphicsPipelineState(const FGraphicsPipelineStateInitializer& Initializer);
 
-	static void OnNewPipelineStateCreated(const FGraphicsPipelineStateInitializer& Initializer, FGraphicsPipelineState* NewGraphicsPipelineState, bool bDoAsyncCompile);
+	static void OnNewPipelineStateCreated(const FGraphicsPipelineStateInitializer& Initializer, FGraphicsPipelineState* NewGraphicsPipelineState, const FString& PSOCompilationEventName, bool bDoAsyncCompile);
 
 };
 
@@ -1836,13 +1864,20 @@ public:
 	FGraphicsPipelineStateInitializer Initializer;
 	EPSOPrecacheResult PSOPreCacheResult;
 	bool bInImmediateCmdList;
+	FPSOCompilationDebugData PSOCompilationDebugData;
 
 	// InInitializer is only used for non-compute tasks, a default can just be used otherwise
-	FCompilePipelineStateTask(FPipelineState* InPipeline, const FGraphicsPipelineStateInitializer& InInitializer, EPSOPrecacheResult InPSOPreCacheResult, bool InbInImmediateCmdList)
+	FCompilePipelineStateTask(
+		FPipelineState* InPipeline, 
+		const FGraphicsPipelineStateInitializer& InInitializer, 
+		EPSOPrecacheResult InPSOPreCacheResult, 
+		bool InbInImmediateCmdList,
+		const FPSOCompilationDebugData& InPSOCompilationDebugData)
 		: Pipeline(InPipeline)
 		, Initializer(InInitializer)
 		, PSOPreCacheResult(InPSOPreCacheResult)
 		, bInImmediateCmdList(InbInImmediateCmdList)
+		, PSOCompilationDebugData(InPSOCompilationDebugData)
 	{
 		ensure(Pipeline->CompletionEvent != nullptr);
 		if(Initializer.bFromPSOFileCache)
@@ -1924,143 +1959,160 @@ public:
 		LLM_SCOPE(ELLMTag::PSO);
 		FTaskTagScope Scope(ETaskTag::EParallelRenderingThread);
 
-		if (Pipeline->IsCompute())
+#if WITH_RHI_BREADCRUMBS
+		if (PSOCompilationDebugData.BreadcrumbNode)
 		{
-			FComputePipelineState* ComputePipeline = static_cast<FComputePipelineState*>(Pipeline);
-			
-			uint64 StartTime = FPlatformTime::Cycles64();
-			ComputePipeline->RHIPipeline = RHICreateComputePipelineState(ComputePipeline->ComputeShader);
-			CheckAndUpdateHitchCountStat(FPSOPrecacheRequestID::EType::Compute, !IsPrecachedPSO(Initializer), StartTime);
-
-			if (!ComputePipeline->RHIPipeline)
-			{
-				HandlePipelineCreationFailure(ComputePipeline->ComputeShader, Initializer.bFromPSOFileCache || Initializer.bPSOPrecache);
-			}
-
-			if (Initializer.bPSOPrecache)
-			{
-				bool bCSValid = ComputePipeline->RHIPipeline != nullptr && ComputePipeline->RHIPipeline->IsValid();
-				GPrecacheComputePipelineCache.PrecacheFinished(ComputePipeline->ComputeShader, bCSValid);
-			}
+			FRHIBreadcrumbNode::WalkIn(PSOCompilationDebugData.BreadcrumbNode);
 		}
-		else
+#endif // WITH_RHI_BREADCRUMBS
 		{
-			const TCHAR* PSOPrecacheResultString = nullptr;
-			switch (PSOPreCacheResult)
-			{
-			case EPSOPrecacheResult::Unknown:			PSOPrecacheResultString = TEXT("PSOPrecache: Unknown"); break;
-			case EPSOPrecacheResult::Active:			PSOPrecacheResultString = TEXT("PSOPrecache: Precaching"); break;
-			case EPSOPrecacheResult::Complete:			PSOPrecacheResultString = TEXT("PSOPrecache: Precached"); break;
-			case EPSOPrecacheResult::Missed:			PSOPrecacheResultString = TEXT("PSOPrecache: Missed"); break;
-			case EPSOPrecacheResult::TooLate:			PSOPrecacheResultString = TEXT("PSOPrecache: Too Late"); break;
-			case EPSOPrecacheResult::NotSupported:		PSOPrecacheResultString = TEXT("PSOPrecache: Precache Untracked"); break;
-			case EPSOPrecacheResult::Untracked:			PSOPrecacheResultString = TEXT("PSOPrecache: Untracked"); break;
-			}
-			TRACE_CPUPROFILER_EVENT_SCOPE_TEXT(PSOPrecacheResultString);
+			TRACE_CPUPROFILER_EVENT_SCOPE_TEXT_CONDITIONAL(*PSOCompilationDebugData.PSOCompilationEventName, !PSOCompilationDebugData.PSOCompilationEventName.IsEmpty())
 
-			bool bSkipCreation = false;
-			if (GRHISupportsMeshShadersTier0)
+			if (Pipeline->IsCompute())
 			{
-				if (!Initializer.BoundShaderState.VertexShaderRHI && !Initializer.BoundShaderState.GetMeshShader())
+				FComputePipelineState* ComputePipeline = static_cast<FComputePipelineState*>(Pipeline);
+
+				uint64 StartTime = FPlatformTime::Cycles64();
+				ComputePipeline->RHIPipeline = RHICreateComputePipelineState(ComputePipeline->ComputeShader);
+				CheckAndUpdateHitchCountStat(FPSOPrecacheRequestID::EType::Compute, !IsPrecachedPSO(Initializer), PSOCompilationDebugData, StartTime);
+
+				if (!ComputePipeline->RHIPipeline)
+				{
+					HandlePipelineCreationFailure(ComputePipeline->ComputeShader, Initializer.bFromPSOFileCache || Initializer.bPSOPrecache);
+				}
+
+				if (Initializer.bPSOPrecache)
 				{
-					UE_LOG(LogRHI, Error, TEXT("Tried to create a Gfx Pipeline State without Vertex or Mesh Shader"));
-					bSkipCreation = true;
+					bool bCSValid = ComputePipeline->RHIPipeline != nullptr && ComputePipeline->RHIPipeline->IsValid();
+					GPrecacheComputePipelineCache.PrecacheFinished(FPrecacheComputeInitializer(ComputePipeline->ComputeShader, ComputePipeline->Name), bCSValid);
 				}
 			}
 			else
 			{
-				if (Initializer.BoundShaderState.GetMeshShader())
+				const TCHAR* PSOPrecacheResultString = nullptr;
+				switch (PSOPreCacheResult)
 				{
-					UE_LOG(LogRHI, Error, TEXT("Tried to create a Gfx Pipeline State with Mesh Shader on hardware without mesh shader support."));
-					bSkipCreation = true;
+				case EPSOPrecacheResult::Unknown:			PSOPrecacheResultString = TEXT("PSOPrecache: Unknown"); break;
+				case EPSOPrecacheResult::Active:			PSOPrecacheResultString = TEXT("PSOPrecache: Precaching"); break;
+				case EPSOPrecacheResult::Complete:			PSOPrecacheResultString = TEXT("PSOPrecache: Precached"); break;
+				case EPSOPrecacheResult::Missed:			PSOPrecacheResultString = TEXT("PSOPrecache: Missed"); break;
+				case EPSOPrecacheResult::TooLate:			PSOPrecacheResultString = TEXT("PSOPrecache: Too Late"); break;
+				case EPSOPrecacheResult::NotSupported:		PSOPrecacheResultString = TEXT("PSOPrecache: Precache Untracked"); break;
+				case EPSOPrecacheResult::Untracked:			PSOPrecacheResultString = TEXT("PSOPrecache: Untracked"); break;
 				}
+				TRACE_CPUPROFILER_EVENT_SCOPE_TEXT(PSOPrecacheResultString);
 
-				if (!Initializer.BoundShaderState.VertexShaderRHI)
+				bool bSkipCreation = false;
+				if (GRHISupportsMeshShadersTier0)
 				{
-					UE_LOG(LogRHI, Error, TEXT("Tried to create a Gfx Pipeline State without Vertex Shader"));
-					bSkipCreation = true;
+					if (!Initializer.BoundShaderState.VertexShaderRHI && !Initializer.BoundShaderState.GetMeshShader())
+					{
+						UE_LOG(LogRHI, Error, TEXT("Tried to create a Gfx Pipeline State without Vertex or Mesh Shader"));
+						bSkipCreation = true;
+					}
 				}
-			}
+				else
+				{
+					if (Initializer.BoundShaderState.GetMeshShader())
+					{
+						UE_LOG(LogRHI, Error, TEXT("Tried to create a Gfx Pipeline State with Mesh Shader on hardware without mesh shader support."));
+						bSkipCreation = true;
+					}
 
-			FGraphicsPipelineState* GfxPipeline = static_cast<FGraphicsPipelineState*>(Pipeline);
+					if (!Initializer.BoundShaderState.VertexShaderRHI)
+					{
+						UE_LOG(LogRHI, Error, TEXT("Tried to create a Gfx Pipeline State without Vertex Shader"));
+						bSkipCreation = true;
+					}
+				}
 
-			uint64 StartTime = FPlatformTime::Cycles64();
-			GfxPipeline->RHIPipeline = bSkipCreation ? nullptr : RHICreateGraphicsPipelineState(Initializer);
-			CheckAndUpdateHitchCountStat(FPSOPrecacheRequestID::EType::Graphics, !IsPrecachedPSO(Initializer), StartTime);
+				FGraphicsPipelineState* GfxPipeline = static_cast<FGraphicsPipelineState*>(Pipeline);
 
-			if (GfxPipeline->RHIPipeline)
-			{
-				GfxPipeline->SortKey = GfxPipeline->RHIPipeline->GetSortKey();
-			}
-			else
-			{
-				HandlePipelineCreationFailure(Initializer);
-			}
+				uint64 StartTime = FPlatformTime::Cycles64();
+				GfxPipeline->RHIPipeline = bSkipCreation ? nullptr : RHICreateGraphicsPipelineState(Initializer);
+				CheckAndUpdateHitchCountStat(FPSOPrecacheRequestID::EType::Graphics, !IsPrecachedPSO(Initializer), PSOCompilationDebugData, StartTime);
 
-			// Mark as finished when it's a precaching job
-			if (Initializer.bPSOPrecache)
-			{
-				GPrecacheGraphicsPipelineCache.PrecacheFinished(Initializer, GfxPipeline->RHIPipeline != nullptr);
-			}
+				if (GfxPipeline->RHIPipeline)
+				{
+					GfxPipeline->SortKey = GfxPipeline->RHIPipeline->GetSortKey();
+				}
+				else
+				{
+					HandlePipelineCreationFailure(Initializer);
+				}
 
-			if (Initializer.BoundShaderState.GetMeshShader())
-			{
-				Initializer.BoundShaderState.GetMeshShader()->Release();
-			}
-			if (Initializer.BoundShaderState.GetAmplificationShader())
-			{
-				Initializer.BoundShaderState.GetAmplificationShader()->Release();
-			}
-			if (Initializer.BoundShaderState.VertexDeclarationRHI)
-			{
-				Initializer.BoundShaderState.VertexDeclarationRHI->Release();
-			}
-			if (Initializer.BoundShaderState.VertexShaderRHI)
-			{
-				Initializer.BoundShaderState.VertexShaderRHI->Release();
-			}
-			if (Initializer.BoundShaderState.PixelShaderRHI)
-			{
-				Initializer.BoundShaderState.PixelShaderRHI->Release();
-			}
-			if (Initializer.BoundShaderState.GetGeometryShader())
-			{
-				Initializer.BoundShaderState.GetGeometryShader()->Release();
-			}
-			if (Initializer.BlendState)
-			{
-				Initializer.BlendState->Release();
-			}
-			if (Initializer.RasterizerState)
-			{
-				Initializer.RasterizerState->Release();
-			}
-			if (Initializer.DepthStencilState)
-			{
-				Initializer.DepthStencilState->Release();
-			}
+				// Mark as finished when it's a precaching job
+				if (Initializer.bPSOPrecache)
+				{
+					GPrecacheGraphicsPipelineCache.PrecacheFinished(Initializer, GfxPipeline->RHIPipeline != nullptr);
+				}
 
-			if (Initializer.BlendState)
-			{
-				Initializer.BlendState->Release();
-			}
-			if (Initializer.RasterizerState)
-			{
-				Initializer.RasterizerState->Release();
+				if (Initializer.BoundShaderState.GetMeshShader())
+				{
+					Initializer.BoundShaderState.GetMeshShader()->Release();
+				}
+				if (Initializer.BoundShaderState.GetAmplificationShader())
+				{
+					Initializer.BoundShaderState.GetAmplificationShader()->Release();
+				}
+				if (Initializer.BoundShaderState.VertexDeclarationRHI)
+				{
+					Initializer.BoundShaderState.VertexDeclarationRHI->Release();
+				}
+				if (Initializer.BoundShaderState.VertexShaderRHI)
+				{
+					Initializer.BoundShaderState.VertexShaderRHI->Release();
+				}
+				if (Initializer.BoundShaderState.PixelShaderRHI)
+				{
+					Initializer.BoundShaderState.PixelShaderRHI->Release();
+				}
+				if (Initializer.BoundShaderState.GetGeometryShader())
+				{
+					Initializer.BoundShaderState.GetGeometryShader()->Release();
+				}
+				if (Initializer.BlendState)
+				{
+					Initializer.BlendState->Release();
+				}
+				if (Initializer.RasterizerState)
+				{
+					Initializer.RasterizerState->Release();
+				}
+				if (Initializer.DepthStencilState)
+				{
+					Initializer.DepthStencilState->Release();
+				}
+
+				if (Initializer.BlendState)
+				{
+					Initializer.BlendState->Release();
+				}
+				if (Initializer.RasterizerState)
+				{
+					Initializer.RasterizerState->Release();
+				}
+				if (Initializer.DepthStencilState)
+				{
+					Initializer.DepthStencilState->Release();
+				}
 			}
-			if (Initializer.DepthStencilState)
+
+			// We kicked a task: the event really should be there
+			if (ensure(Pipeline->CompletionEvent))
 			{
-				Initializer.DepthStencilState->Release();
+				Pipeline->CompletionEvent->DispatchSubsequents();
+				// At this point, it's not safe to use Pipeline anymore, as it might get picked up by ProcessDelayedCleanup and deleted
+				Pipeline = nullptr;
 			}
 		}
 
-		// We kicked a task: the event really should be there
-		if (ensure(Pipeline->CompletionEvent))
+#if WITH_RHI_BREADCRUMBS
+		if (PSOCompilationDebugData.BreadcrumbNode)
 		{
-			Pipeline->CompletionEvent->DispatchSubsequents();
-			// At this point, it's not safe to use Pipeline anymore, as it might get picked up by ProcessDelayedCleanup and deleted
-			Pipeline = nullptr;
+			FRHIBreadcrumbNode::WalkOut(PSOCompilationDebugData.BreadcrumbNode);
 		}
+#endif // WITH_RHI_BREADCRUMBS
 	}
 
 	FORCEINLINE TStatId GetStatId() const
@@ -2189,7 +2241,7 @@ static bool IsAsyncCompilationAllowed(FRHIComputeCommandList& RHICmdList, bool b
 
 	return GRHISupportsAsyncPipelinePrecompile &&
 		FDataDrivenShaderPlatformInfo::GetSupportsAsyncPipelineCompilation(GMaxRHIShaderPlatform) &&
-		bCVarAllowsAsyncCreate && !RHICmdList.Bypass() && (IsRunningRHIInSeparateThread() && !IsInRHIThread()) && RHICmdList.AsyncPSOCompileAllowed();
+		bCVarAllowsAsyncCreate && !RHICmdList.Bypass() && (IsRunningRHIInSeparateThread() && !IsInRHIThread()) && !RHICmdList.IsRecursive();
 }
 
 uint64 PipelineStateCache::RetrieveGraphicsPipelineStateSortKey(const FGraphicsPipelineState* GraphicsPipelineState)
@@ -2197,7 +2249,7 @@ uint64 PipelineStateCache::RetrieveGraphicsPipelineStateSortKey(const FGraphicsP
 	return GraphicsPipelineState != nullptr ? GraphicsPipelineState->SortKey : 0;
 }
 
-static void InternalCreateComputePipelineState(FRHIComputeShader* ComputeShader, bool bDoAsyncCompile, bool bFromPSOFileCache, FComputePipelineState* CachedState, bool bInImmediateCmdList = true)
+static void InternalCreateComputePipelineState(FRHIComputeShader* ComputeShader, bool bDoAsyncCompile, bool bFromPSOFileCache, FComputePipelineState* CachedState, const FPSOCompilationDebugData& PSOCompilationDebugData, bool bInImmediateCmdList)
 {
 	FGraphEventRef GraphEvent = CachedState->CompletionEvent;
 
@@ -2208,14 +2260,16 @@ static void InternalCreateComputePipelineState(FRHIComputeShader* ComputeShader,
 
 		FGraphicsPipelineStateInitializer GraphicsPipelineStateInitializer;
 		GraphicsPipelineStateInitializer.bFromPSOFileCache = bFromPSOFileCache;
-		TGraphTask<FCompilePipelineStateTask>::CreateTask().ConstructAndDispatchWhenReady(CachedState, GraphicsPipelineStateInitializer, EPSOPrecacheResult::Untracked, bInImmediateCmdList);
+		TGraphTask<FCompilePipelineStateTask>::CreateTask().ConstructAndDispatchWhenReady(CachedState, GraphicsPipelineStateInitializer, EPSOPrecacheResult::Untracked, bInImmediateCmdList, PSOCompilationDebugData);
 	}
 	else
 	{
+		TRACE_CPUPROFILER_EVENT_SCOPE_TEXT_CONDITIONAL(*PSOCompilationDebugData.PSOCompilationEventName, !PSOCompilationDebugData.PSOCompilationEventName.IsEmpty())
+
 		check(GraphEvent == nullptr);
 		uint64 StartTime = FPlatformTime::Cycles64();
 		CachedState->RHIPipeline = RHICreateComputePipelineState(ComputeShader);
-		CheckAndUpdateHitchCountStat(FPSOPrecacheRequestID::EType::Compute, !bFromPSOFileCache, StartTime);
+		CheckAndUpdateHitchCountStat(FPSOPrecacheRequestID::EType::Compute, !bFromPSOFileCache, PSOCompilationDebugData, StartTime);
 
 		if (!CachedState->RHIPipeline)
 		{
@@ -2239,7 +2293,7 @@ FComputePipelineState* PipelineStateCache::GetAndOrCreateComputePipelineState(FR
 		FPipelineFileCacheManager::CacheComputePSO(GetTypeHash(ComputeShader), ComputeShader, bWasPSOPrecached);
 
 		// create new compute state
-		OutCachedState = new FComputePipelineState(ComputeShader);
+		OutCachedState = new FComputePipelineState(ComputeShader, nullptr);
 		OutCachedState->Stats = FPipelineFileCacheManager::RegisterPSOStats(GetTypeHash(ComputeShader));
 		if (DoAsyncCompile)
 		{
@@ -2251,8 +2305,13 @@ FComputePipelineState* PipelineStateCache::GetAndOrCreateComputePipelineState(FR
 			++ComputePipelineCacheMisses;
 		}
 
+		FPSOCompilationDebugData PSOCompilationDebugData;
+#if WITH_RHI_BREADCRUMBS
+		PSOCompilationDebugData.BreadcrumbNode = DoAsyncCompile ? RHICmdList.GetCurrentBreadcrumbRef() : nullptr;
+#endif // WITH_RHI_BREADCRUMBS
+		
 		FGraphEventRef GraphEvent = OutCachedState->CompletionEvent;
-		InternalCreateComputePipelineState(ComputeShader, DoAsyncCompile, bFromFileCache, OutCachedState, RHICmdList.IsImmediate());
+		InternalCreateComputePipelineState(ComputeShader, DoAsyncCompile, bFromFileCache, OutCachedState, PSOCompilationDebugData, RHICmdList.IsImmediate());
 
 		if (GraphEvent.IsValid())
 		{
@@ -2603,7 +2662,7 @@ inline void ValidateGraphicsPipelineStateInitializer(const FGraphicsPipelineStat
 	check(Initializer.DepthStencilState && Initializer.BlendState && Initializer.RasterizerState);
 }
 
-static void InternalCreateGraphicsPipelineState(const FGraphicsPipelineStateInitializer& Initializer, EPSOPrecacheResult PSOPrecacheResult, bool bDoAsyncCompile, bool bPSOPrecache, FGraphicsPipelineState* CachedState, bool bInImmediateCmdList = true)
+static void InternalCreateGraphicsPipelineState(const FGraphicsPipelineStateInitializer& Initializer, EPSOPrecacheResult PSOPrecacheResult, bool bDoAsyncCompile, bool bPSOPrecache, FGraphicsPipelineState* CachedState, const FPSOCompilationDebugData& PSOCompilationDebugData, bool bInImmediateCmdList)
 {
 	FGraphEventRef GraphEvent = CachedState->CompletionEvent;
 
@@ -2614,13 +2673,13 @@ static void InternalCreateGraphicsPipelineState(const FGraphicsPipelineStateInit
 		// Use normal task graph for non-precompile jobs (or when thread pool is not enabled)
 		if (!bPSOPrecache || !FPSOPrecacheThreadPool::UsePool())
 		{
-			TGraphTask<FCompilePipelineStateTask>::CreateTask().ConstructAndDispatchWhenReady(CachedState, Initializer, PSOPrecacheResult, bInImmediateCmdList);
+			TGraphTask<FCompilePipelineStateTask>::CreateTask().ConstructAndDispatchWhenReady(CachedState, Initializer, PSOPrecacheResult, bInImmediateCmdList, PSOCompilationDebugData);
 		}
 		else
 		{
 			// Here, PSO precompiles use a separate thread pool.
 			// Note that we do not add precompile tasks as cmdlist prerequisites.
-			TUniquePtr<FCompilePipelineStateTask> ThreadPoolTask = MakeUnique<FCompilePipelineStateTask>(CachedState, Initializer, PSOPrecacheResult, bInImmediateCmdList);
+			TUniquePtr<FCompilePipelineStateTask> ThreadPoolTask = MakeUnique<FCompilePipelineStateTask>(CachedState, Initializer, PSOPrecacheResult, bInImmediateCmdList, PSOCompilationDebugData);
 			uint64 StartTime = FPlatformTime::Cycles64();
 #if	PSO_TRACK_CACHE_STATS
 			StatsStartPrecompile();
@@ -2642,11 +2701,13 @@ static void InternalCreateGraphicsPipelineState(const FGraphicsPipelineStateInit
 		}
 	}
 	else
-	{
+	{		
+		TRACE_CPUPROFILER_EVENT_SCOPE_TEXT_CONDITIONAL(*PSOCompilationDebugData.PSOCompilationEventName, !PSOCompilationDebugData.PSOCompilationEventName.IsEmpty())
+
 		check(GraphEvent == nullptr);
 		uint64 StartTime = FPlatformTime::Cycles64();
 		CachedState->RHIPipeline = RHICreateGraphicsPipelineState(Initializer);
-		CheckAndUpdateHitchCountStat(FPSOPrecacheRequestID::EType::Graphics, !IsPrecachedPSO(Initializer), StartTime);
+		CheckAndUpdateHitchCountStat(FPSOPrecacheRequestID::EType::Graphics, !IsPrecachedPSO(Initializer), PSOCompilationDebugData, StartTime);
 
 		if (Initializer.bPSOPrecache)
 		{
@@ -2738,9 +2799,14 @@ FGraphicsPipelineState* PipelineStateCache::GetAndOrCreateGraphicsPipelineState(
 			PSOPrecacheResult = EPSOPrecacheResult::TooLate;
 		}
 
+		FPSOCompilationDebugData PSOCompilationDebugData;
+#if WITH_RHI_BREADCRUMBS
+		PSOCompilationDebugData.BreadcrumbNode = DoAsyncCompile ? RHICmdList.GetCurrentBreadcrumbRef() : nullptr;
+#endif // WITH_RHI_BREADCRUMBS
+
 		bool bPSOPrecache = Initializer.bFromPSOFileCache;
 		FGraphEventRef GraphEvent = OutCachedState->CompletionEvent;
-		InternalCreateGraphicsPipelineState(Initializer, PSOPrecacheResult, DoAsyncCompile, bPSOPrecache, OutCachedState, RHICmdList.IsImmediate());
+		InternalCreateGraphicsPipelineState(Initializer, PSOPrecacheResult, DoAsyncCompile, bPSOPrecache, OutCachedState, PSOCompilationDebugData, RHICmdList.IsImmediate());
 
 		// Add dispatch pre requisite for non precaching jobs only
 		//if (GraphEvent.IsValid() && (!bPSOPrecache || !FPSOPrecacheThreadPool::UsePool()))
@@ -2833,7 +2899,7 @@ bool PipelineStateCache::IsPSOPrecachingEnabled()
 #endif // WITH_EDITOR
 }
 
-FPSOPrecacheRequestResult FPrecacheComputePipelineCache::PrecacheComputePipelineState(FRHIComputeShader* ComputeShader, bool bForcePrecache)
+FPSOPrecacheRequestResult FPrecacheComputePipelineCache::PrecacheComputePipelineState(FRHIComputeShader* ComputeShader, const TCHAR* Name, bool bForcePrecache)
 {
 	FPSOPrecacheRequestResult Result;
 	if (!PipelineStateCache::IsPSOPrecachingEnabled() && !bForcePrecache)
@@ -2845,12 +2911,12 @@ FPSOPrecacheRequestResult FPrecacheComputePipelineCache::PrecacheComputePipeline
 		return Result;
 	}
 
-	FPrecacheComputeInitializer PrecacheComputeInitializer(ComputeShader);
+	FPrecacheComputeInitializer PrecacheComputeInitializer(ComputeShader, Name);
 	static bool bDoAsyncCompile = FApp::ShouldUseThreadingForPerformance();
-	return GPrecacheComputePipelineCache.TryAddNewState(PrecacheComputeInitializer, bDoAsyncCompile);
+	return GPrecacheComputePipelineCache.TryAddNewState(PrecacheComputeInitializer, Name, bDoAsyncCompile);
 }
 
-void FPrecacheComputePipelineCache::OnNewPipelineStateCreated(const FPrecacheComputeInitializer& ComputeInitializer, FComputePipelineState* CachedState, bool bDoAsyncCompile)
+void FPrecacheComputePipelineCache::OnNewPipelineStateCreated(const FPrecacheComputeInitializer& ComputeInitializer, FComputePipelineState* CachedState, const FString& PSOCompilationEventName, bool bDoAsyncCompile)
 {
 	// create a compilation task, or just do it now...
 	if (bDoAsyncCompile)
@@ -2858,16 +2924,20 @@ void FPrecacheComputePipelineCache::OnNewPipelineStateCreated(const FPrecacheCom
 		check(CachedState->CompletionEvent != nullptr);
 		FGraphicsPipelineStateInitializer GraphicsPipelineStateInitializer;
 		GraphicsPipelineStateInitializer.bPSOPrecache = true;
+
+		FPSOCompilationDebugData PSOCompilationDebugData;
+		PSOCompilationDebugData.PSOCompilationEventName = PSOCompilationEventName;
+
 		// Thread pool disabled?
 		if (!FPSOPrecacheThreadPool::UsePool())
 		{
-			TGraphTask<FCompilePipelineStateTask>::CreateTask().ConstructAndDispatchWhenReady(CachedState, GraphicsPipelineStateInitializer, EPSOPrecacheResult::Active, false);
+			TGraphTask<FCompilePipelineStateTask>::CreateTask().ConstructAndDispatchWhenReady(CachedState, GraphicsPipelineStateInitializer, EPSOPrecacheResult::Active, false, PSOCompilationDebugData);
 		}
 		else
 		{
 			// Here, PSO precompiles use a separate thread pool.
 			// Note that we do not add precompile tasks as cmdlist prerequisites.
-			TUniquePtr<FCompilePipelineStateTask> ThreadPoolTask = MakeUnique<FCompilePipelineStateTask>(CachedState, GraphicsPipelineStateInitializer, EPSOPrecacheResult::Active, false);
+			TUniquePtr<FCompilePipelineStateTask> ThreadPoolTask = MakeUnique<FCompilePipelineStateTask>(CachedState, GraphicsPipelineStateInitializer, EPSOPrecacheResult::Active, false, PSOCompilationDebugData);
 			uint64 StartTime = FPlatformTime::Cycles64();
 #if	PSO_TRACK_CACHE_STATS
 			StatsStartPrecompile();
@@ -2896,9 +2966,9 @@ void FPrecacheComputePipelineCache::OnNewPipelineStateCreated(const FPrecacheCom
 	}
 }
 
-FPSOPrecacheRequestResult PipelineStateCache::PrecacheComputePipelineState(FRHIComputeShader* ComputeShader, bool bForcePrecache)
+FPSOPrecacheRequestResult PipelineStateCache::PrecacheComputePipelineState(FRHIComputeShader* ComputeShader, const TCHAR* Name, bool bForcePrecache)
 {
-	return GPrecacheComputePipelineCache.PrecacheComputePipelineState(ComputeShader, bForcePrecache);
+	return GPrecacheComputePipelineCache.PrecacheComputePipelineState(ComputeShader, Name, bForcePrecache);
 }
 
 FPSOPrecacheRequestResult FPrecacheGraphicsPipelineCache::PrecacheGraphicsPipelineState(const FGraphicsPipelineStateInitializer& Initializer)
@@ -2914,11 +2984,13 @@ FPSOPrecacheRequestResult FPrecacheGraphicsPipelineCache::PrecacheGraphicsPipeli
 	// Use async compilation if available
 	static bool bDoAsyncCompile = FApp::ShouldUseThreadingForPerformance();
 
+	FString PSOCompilationEventName;
+
 	// try and create new graphics state
-	return GPrecacheGraphicsPipelineCache.TryAddNewState(Initializer, bDoAsyncCompile);
+	return GPrecacheGraphicsPipelineCache.TryAddNewState(Initializer, PSOCompilationEventName, bDoAsyncCompile);
 }
 	
-void FPrecacheGraphicsPipelineCache::OnNewPipelineStateCreated(const FGraphicsPipelineStateInitializer & Initializer, FGraphicsPipelineState * NewGraphicsPipelineState, bool bDoAsyncCompile)
+void FPrecacheGraphicsPipelineCache::OnNewPipelineStateCreated(const FGraphicsPipelineStateInitializer & Initializer, FGraphicsPipelineState * NewGraphicsPipelineState, const FString& PSOCompilationEventName, bool bDoAsyncCompile)
 {
 	ValidateGraphicsPipelineStateInitializer(Initializer);
 	check((NewGraphicsPipelineState->CompletionEvent != nullptr) == bDoAsyncCompile);
@@ -2929,8 +3001,11 @@ void FPrecacheGraphicsPipelineCache::OnNewPipelineStateCreated(const FGraphicsPi
 	FGraphicsPipelineStateInitializer InitializerCopy(Initializer);
 	InitializerCopy.bPSOPrecache = bPSOPrecache;
 
-	// Start the precache task
-	InternalCreateGraphicsPipelineState(InitializerCopy, EPSOPrecacheResult::Active, bDoAsyncCompile, bPSOPrecache, NewGraphicsPipelineState);
+	FPSOCompilationDebugData PSOCompilationDebugData;
+	PSOCompilationDebugData.PSOCompilationEventName = PSOCompilationEventName;
+
+	// Start the precache task	
+	InternalCreateGraphicsPipelineState(InitializerCopy, EPSOPrecacheResult::Active, bDoAsyncCompile, bPSOPrecache, NewGraphicsPipelineState, PSOCompilationDebugData, false);
 }
 
 FPSOPrecacheRequestResult PipelineStateCache::PrecacheGraphicsPipelineState(const FGraphicsPipelineStateInitializer& Initializer)
@@ -2955,7 +3030,7 @@ EPSOPrecacheResult PipelineStateCache::CheckPipelineStateInCache(FRHIComputeShad
 		return EPSOPrecacheResult::Unknown;
 	}
 
-	return GPrecacheComputePipelineCache.GetPrecachingState(ComputeShader);
+	return GPrecacheComputePipelineCache.GetPrecachingState(FPrecacheComputeInitializer(ComputeShader, nullptr));
 }
 
 bool PipelineStateCache::IsPrecaching(const FPSOPrecacheRequestID& PSOPrecacheRequestID)
@@ -2994,7 +3069,7 @@ bool PipelineStateCache::IsPrecaching(FRHIComputeShader* ComputeShader)
 		return false;
 	}
 
-	return GPrecacheComputePipelineCache.GetPrecachingState(ComputeShader) == EPSOPrecacheResult::Active;
+	return GPrecacheComputePipelineCache.GetPrecachingState(FPrecacheComputeInitializer(ComputeShader, nullptr)) == EPSOPrecacheResult::Active;
 }
 
 bool PipelineStateCache::IsPrecaching()
@@ -3032,6 +3107,15 @@ uint32 PipelineStateCache::NumActivePrecacheRequests()
 	return GPrecacheGraphicsPipelineCache.NumActivePrecacheRequests() + GPrecacheComputePipelineCache.NumActivePrecacheRequests();
 }
 
+void PipelineStateCache::ResetPSOHitchTrackingStats()
+{
+	SET_DWORD_STAT(STAT_RuntimeGraphicsPSOHitchCount, 0);
+	GraphicsPSOCreationHitchCount = 0;
+
+	SET_DWORD_STAT(STAT_RuntimeComputePSOHitchCount, 0);
+	ComputePSOCreationHitchCount = 0;
+}
+
 FRHIGraphicsPipelineState* ExecuteSetGraphicsPipelineState(FGraphicsPipelineState* GraphicsPipelineState)
 {
 	FRHIGraphicsPipelineState* RHIPipeline = GraphicsPipelineState->RHIPipeline;
diff --git a/Engine/Source/Runtime/RHI/Private/RHIBreadcrumbs.cpp b/Engine/Source/Runtime/RHI/Private/RHIBreadcrumbs.cpp
--- a/Engine/Source/Runtime/RHI/Private/RHIBreadcrumbs.cpp
+++ b/Engine/Source/Runtime/RHI/Private/RHIBreadcrumbs.cpp
@@ -65,7 +65,7 @@ RHI_API FString FRHIBreadcrumbNode::GetFullPath() const
 
 	auto Recurse = [&Result, &Buffer, this](auto& Recurse, FRHIBreadcrumbNode const* Current) -> void
 	{
-		if (!Current)
+		if (!Current || Current == Sentinel)
 			return;
 
 		Recurse(Recurse, Current->Parent);
diff --git a/Engine/Source/Runtime/RHI/Private/RHICommandList.cpp b/Engine/Source/Runtime/RHI/Private/RHICommandList.cpp
--- a/Engine/Source/Runtime/RHI/Private/RHICommandList.cpp
+++ b/Engine/Source/Runtime/RHI/Private/RHICommandList.cpp
@@ -1288,7 +1288,7 @@ FRHICommandList_RecursiveHazardous::FRHICommandList_RecursiveHazardous(IRHIComma
 
 	Contexts[ERHIPipeline::Graphics] = GraphicsContext;
 
-	PersistentState.bAsyncPSOCompileAllowed = false;
+	PersistentState.bRecursive = true;
 }
 
 FRHICommandList_RecursiveHazardous::~FRHICommandList_RecursiveHazardous()
@@ -1316,7 +1316,7 @@ FRHIComputeCommandList_RecursiveHazardous::FRHIComputeCommandList_RecursiveHazar
 	ComputeContext = &Context->GetHighestLevelContext();
 	Contexts[ActivePipeline] = ComputeContext;
 
-	PersistentState.bAsyncPSOCompileAllowed = false;
+	PersistentState.bRecursive = true;
 }
 
 FRHIComputeCommandList_RecursiveHazardous::~FRHIComputeCommandList_RecursiveHazardous()
diff --git a/Engine/Source/Runtime/RHI/Public/PipelineStateCache.h b/Engine/Source/Runtime/RHI/Public/PipelineStateCache.h
--- a/Engine/Source/Runtime/RHI/Public/PipelineStateCache.h
+++ b/Engine/Source/Runtime/RHI/Public/PipelineStateCache.h
@@ -142,7 +142,7 @@ namespace PipelineStateCache
 	extern RHI_API bool						IsPSOPrecachingEnabled();
 
 	/* Precache the compute shader and return a request ID if precached async */
-	extern RHI_API FPSOPrecacheRequestResult PrecacheComputePipelineState(FRHIComputeShader* ComputeShader, bool bForcePrecache = false);
+	extern RHI_API FPSOPrecacheRequestResult PrecacheComputePipelineState(FRHIComputeShader* ComputeShader, const TCHAR* Name = nullptr, bool bForcePrecache = false);
 
 	/* Precache the graphic PSO and return an optional graph event if precached async */
 	extern RHI_API FPSOPrecacheRequestResult PrecacheGraphicsPipelineState(const FGraphicsPipelineStateInitializer& PipelineStateInitializer);
@@ -170,6 +170,9 @@ namespace PipelineStateCache
 
 	/* Return number of active or pending PSO precache requests */
 	extern RHI_API uint32					NumActivePrecacheRequests();
+
+	/* Reset the PSO hitch tracking counters */
+	extern RHI_API void						ResetPSOHitchTrackingStats();
 }
 
 // Returns the shader index within the ray tracing pipeline or INDEX_NONE if given shader does not exist.
diff --git a/Engine/Source/Runtime/RHI/Public/RHIBreadcrumbs.h b/Engine/Source/Runtime/RHI/Public/RHIBreadcrumbs.h
--- a/Engine/Source/Runtime/RHI/Public/RHIBreadcrumbs.h
+++ b/Engine/Source/Runtime/RHI/Public/RHIBreadcrumbs.h
@@ -485,7 +485,7 @@
 		{
 			auto Recurse = [](FRHIBreadcrumbNode* Current, auto& Recurse) -> void
 			{
-				if (!Current)
+				if (!Current || Current == Sentinel)
 					return;
 
 				Recurse(Current->GetParent(), Recurse);
@@ -501,7 +501,7 @@
 #if RHI_BREADCRUMBS_EMIT_CPU
 		if (TRACE_CPUPROFILER_EVENT_MANUAL_IS_ENABLED())
 		{
-			while (Node)
+			while (Node && Node != Sentinel)
 			{
 				Node->EndCPU();
 				Node = Node->GetParent();
diff --git a/Engine/Source/Runtime/RHI/Public/RHICommandList.h b/Engine/Source/Runtime/RHI/Public/RHICommandList.h
--- a/Engine/Source/Runtime/RHI/Public/RHICommandList.h
+++ b/Engine/Source/Runtime/RHI/Public/RHICommandList.h
@@ -587,7 +587,7 @@ public:
 
 	FORCEINLINE FRHIGPUMask GetGPUMask() const { return PersistentState.CurrentGPUMask; }
 
-	bool AsyncPSOCompileAllowed() const { return PersistentState.bAsyncPSOCompileAllowed; }
+	bool IsRecursive		   () const { return PersistentState.bRecursive; }
 	bool IsOutsideRenderPass   () const { return !PersistentState.bInsideRenderPass; }
 	bool IsInsideRenderPass    () const { return PersistentState.bInsideRenderPass;  }
 	bool IsInsideComputePass   () const { return PersistentState.bInsideComputePass; }
@@ -978,7 +978,7 @@ protected:
 
 		PersistentState.CachedNumSimultanousRenderTargets = RTInfo.NumColorRenderTargets;
 		PersistentState.CachedDepthStencilTarget = RTInfo.DepthStencilRenderTarget;
-		PersistentState.HasFragmentDensityAttachment = RTInfo.ShadingRateTexture != nullptr;
+		PersistentState.bHasFragmentDensityAttachment = RTInfo.ShadingRateTexture != nullptr;
 		PersistentState.MultiViewCount = RTInfo.MultiViewCount;
 	}
 
@@ -1082,13 +1082,14 @@ protected:
 		ESubpassHint SubpassHint = ESubpassHint::None;
 		uint8 SubpassIndex = 0;
 		uint8 MultiViewCount = 0;
-		bool HasFragmentDensityAttachment = false;
 
-		bool bInsideRenderPass = false;
-		bool bInsideComputePass = false;
-		bool bInsideOcclusionQueryBatch = false;
-		bool bAsyncPSOCompileAllowed = true;
-		bool bImmediate = false;
+		uint8 bHasFragmentDensityAttachment		: 1 = false;
+
+		uint8 bInsideRenderPass					: 1 = false;
+		uint8 bInsideComputePass				: 1 = false;
+		uint8 bInsideOcclusionQueryBatch		: 1 = false;
+		uint8 bRecursive						: 1 = false;
+		uint8 bImmediate						: 1 = false;
 
 		FRHIGPUMask CurrentGPUMask;
 		FRHIGPUMask InitialGPUMask;
@@ -3275,7 +3276,7 @@ public:
 		GraphicsPSOInit.SubpassHint = PersistentState.SubpassHint;
 		GraphicsPSOInit.SubpassIndex = PersistentState.SubpassIndex;
 		GraphicsPSOInit.MultiViewCount = PersistentState.MultiViewCount;
-		GraphicsPSOInit.bHasFragmentDensityAttachment = PersistentState.HasFragmentDensityAttachment;
+		GraphicsPSOInit.bHasFragmentDensityAttachment = PersistentState.bHasFragmentDensityAttachment;
 	}
 
 	FORCEINLINE_DEBUGGABLE void SetGraphicsPipelineState(class FGraphicsPipelineState* GraphicsPipelineState, const FBoundShaderStateInput& ShaderInput, uint32 StencilRef, bool bApplyAdditionalState)
