Remove default member initializers for bit-fields in public headers

diff --git a/Engine/Source/Runtime/RHI/Public/RHICommandList.h b/Engine/Source/Runtime/RHI/Public/RHICommandList.h
--- a/Engine/Source/Runtime/RHI/Public/RHICommandList.h
+++ b/Engine/Source/Runtime/RHI/Public/RHICommandList.h
@@ -1083,13 +1083,13 @@ protected:
 		uint8 SubpassIndex = 0;
 		uint8 MultiViewCount = 0;
 
-		uint8 bHasFragmentDensityAttachment		: 1 = false;
+		uint8 bHasFragmentDensityAttachment		: 1;
 
-		uint8 bInsideRenderPass					: 1 = false;
-		uint8 bInsideComputePass				: 1 = false;
-		uint8 bInsideOcclusionQueryBatch		: 1 = false;
-		uint8 bRecursive						: 1 = false;
-		uint8 bImmediate						: 1 = false;
+		uint8 bInsideRenderPass					: 1;
+		uint8 bInsideComputePass				: 1;
+		uint8 bInsideOcclusionQueryBatch		: 1;
+		uint8 bRecursive						: 1;
+		uint8 bImmediate						: 1;
 
 		FRHIGPUMask CurrentGPUMask;
 		FRHIGPUMask InitialGPUMask;
@@ -1106,7 +1106,12 @@ protected:
 #endif
 
 		FPersistentState(FRHIGPUMask InInitialGPUMask, bool bInImmediate = false)
-			: bImmediate(bInImmediate)
+			: bHasFragmentDensityAttachment(0)
+			, bInsideRenderPass(0)
+			, bInsideComputePass(0)
+			, bInsideOcclusionQueryBatch(0)
+			, bRecursive(0)
+			, bImmediate(bInImmediate)
 			, CurrentGPUMask(InInitialGPUMask)
 			, InitialGPUMask(InInitialGPUMask)
 		{}
