Fixed incorrect hashing of FGraphicsPipelineStateInitializer, which was causing TMap to fail to find elements in the map, causing a crash in the pipeline state cache.

diff --git a/Engine/Source/Runtime/RHI/Private/PipelineStateCache.cpp b/Engine/Source/Runtime/RHI/Private/PipelineStateCache.cpp
--- a/Engine/Source/Runtime/RHI/Private/PipelineStateCache.cpp
+++ b/Engine/Source/Runtime/RHI/Private/PipelineStateCache.cpp
@@ -74,13 +74,15 @@ static inline uint32 GetTypeHash(const FGraphicsPipelineStateInitializer& Initia
 	Hash = HashCombineFast(Hash, GetTypeHash(Initializer.RasterizerState));
 	Hash = HashCombineFast(Hash, GetTypeHash(Initializer.DepthStencilState));
 	Hash = HashCombineFast(Hash, GetTypeHash(Initializer.ImmutableSamplerState));
-
 	Hash = HashCombineFast(Hash, GetTypeHash(Initializer.PrimitiveType));
 	Hash = HashCombineFast(Hash, GetTypeHash(Initializer.RenderTargetsEnabled));
 	Hash = HashCombineFast(Hash, GetTypeHash(Initializer.RenderTargetFormats));
-	Hash = HashCombineFast(Hash, GetTypeHash(Initializer.RenderTargetFlags));
+	for (int32 Index = 0; Index < Initializer.RenderTargetFlags.Num(); ++Index)
+	{
+		Hash = HashCombineFast(Hash, GetTypeHash(Initializer.RenderTargetFlags[Index] & FGraphicsPipelineStateInitializer::RelevantRenderTargetFlagMask));
+	}
 	Hash = HashCombineFast(Hash, GetTypeHash(Initializer.DepthStencilTargetFormat));
-	Hash = HashCombineFast(Hash, GetTypeHash(Initializer.DepthStencilTargetFlag));
+	Hash = HashCombineFast(Hash, GetTypeHash(Initializer.DepthStencilTargetFlag & FGraphicsPipelineStateInitializer::RelevantDepthStencilFlagMask));
 	Hash = HashCombineFast(Hash, GetTypeHash(Initializer.DepthTargetLoadAction));
 	Hash = HashCombineFast(Hash, GetTypeHash(Initializer.DepthTargetStoreAction));
 	Hash = HashCombineFast(Hash, GetTypeHash(Initializer.StencilTargetLoadAction));
@@ -93,6 +95,7 @@ static inline uint32 GetTypeHash(const FGraphicsPipelineStateInitializer& Initia
 	Hash = HashCombineFast(Hash, GetTypeHash(Initializer.bDepthBounds));
 	Hash = HashCombineFast(Hash, GetTypeHash(Initializer.MultiViewCount));
 	Hash = HashCombineFast(Hash, GetTypeHash(Initializer.bHasFragmentDensityAttachment));
+	Hash = HashCombineFast(Hash, GetTypeHash(Initializer.bAllowVariableRateShading));
 	Hash = HashCombineFast(Hash, GetTypeHash(Initializer.ShadingRate));
 	return Hash;
 }
