Deprecate GetRHIPipelines() and EnumerateRHIPipelines()
 - Replaced with range iteration of set bits, using MakeFlagsRange().
 - Results in better code gen.

diff --git a/Engine/Source/Runtime/D3D12RHI/Private/D3D12BindlessDescriptors.cpp b/Engine/Source/Runtime/D3D12RHI/Private/D3D12BindlessDescriptors.cpp
--- a/Engine/Source/Runtime/D3D12RHI/Private/D3D12BindlessDescriptors.cpp
+++ b/Engine/Source/Runtime/D3D12RHI/Private/D3D12BindlessDescriptors.cpp
@@ -141,9 +141,9 @@ void FD3D12BindlessResourceManager::UpdateDescriptor(FRHICommandListBase& RHICmd
 {
 	if (DstHandle.IsValid())
 	{
-		for (ERHIPipeline PipelineIndex : GetRHIPipelines())
+		for (ERHIPipeline Pipeline : MakeFlagsRange(ERHIPipeline::All))
 		{
-			FRHICommandListScopedPipeline Scope(RHICmdList, PipelineIndex);
+			FRHICommandListScopedPipeline Scope(RHICmdList, Pipeline);
 			RHICmdList.EnqueueLambda([this, View, DstHandle](FRHICommandListBase& ExecutingCmdList)
 			{
 				FD3D12CommandContext& Context =
diff --git a/Engine/Source/Runtime/D3D12RHI/Private/D3D12CommandContext.cpp b/Engine/Source/Runtime/D3D12RHI/Private/D3D12CommandContext.cpp
--- a/Engine/Source/Runtime/D3D12RHI/Private/D3D12CommandContext.cpp
+++ b/Engine/Source/Runtime/D3D12RHI/Private/D3D12CommandContext.cpp
@@ -844,10 +844,10 @@ void FD3D12DynamicRHI::RHICreateTransition(FRHITransition* Transition, const FRH
 		for (uint32 Index : FRHIGPUMask::All())
 		{
 			TRHIPipelineArray<FD3D12SyncPointRef>& DeviceSyncPoints = Data->SyncPoints.Emplace_GetRef();
-			EnumerateRHIPipelines(CreateInfo.SrcPipelines, [&](ERHIPipeline Pipeline)
+			for (ERHIPipeline Pipeline : MakeFlagsRange(CreateInfo.SrcPipelines))
 			{
 				DeviceSyncPoints[Pipeline] = FD3D12SyncPoint::Create(ED3D12SyncPointType::GPUOnly);
-			});
+			}
 		}
 	}
 
diff --git a/Engine/Source/Runtime/D3D12RHI/Private/D3D12Commands.cpp b/Engine/Source/Runtime/D3D12RHI/Private/D3D12Commands.cpp
--- a/Engine/Source/Runtime/D3D12RHI/Private/D3D12Commands.cpp
+++ b/Engine/Source/Runtime/D3D12RHI/Private/D3D12Commands.cpp
@@ -568,13 +568,13 @@ void FD3D12CommandContext::RHIEndTransitions(TArrayView<const FRHITransition*> T
 			{
 				const TRHIPipelineArray<FD3D12SyncPointRef>& DeviceSyncPoints = Data->SyncPoints[GetGPUIndex()];
 
-				EnumerateRHIPipelines(Data->SrcPipelines, [&](ERHIPipeline SrcPipeline)
+				for (ERHIPipeline SrcPipeline : MakeFlagsRange(Data->SrcPipelines))
 				{
 					if (SrcPipeline != DstPipeline && DeviceSyncPoints[SrcPipeline])
 					{
 						WaitSyncPoint(DeviceSyncPoints[SrcPipeline]);
 					}
-				});
+				}
 			}
 		}
 	}
diff --git a/Engine/Source/Runtime/D3D12RHI/Private/D3D12Descriptors.cpp b/Engine/Source/Runtime/D3D12RHI/Private/D3D12Descriptors.cpp
--- a/Engine/Source/Runtime/D3D12RHI/Private/D3D12Descriptors.cpp
+++ b/Engine/Source/Runtime/D3D12RHI/Private/D3D12Descriptors.cpp
@@ -187,7 +187,7 @@ void FD3D12OnlineDescriptorManager::Init(uint32 InTotalSize, uint32 InBlockSize,
 	{
 #if D3D12RHI_USE_CONSTANT_BUFFER_VIEWS
 		Heaps = GetParentDevice()->GetBindlessDescriptorManager().AllocateResourceHeapsForAllPipelines(InTotalSize);
-		for (ERHIPipeline Pipeline : GetRHIPipelines())
+		for (ERHIPipeline Pipeline : MakeFlagsRange(ERHIPipeline::All))
 		{
 			if (Heaps[Pipeline])
 			{
@@ -212,7 +212,7 @@ void FD3D12OnlineDescriptorManager::Init(uint32 InTotalSize, uint32 InBlockSize,
 		INC_DWORD_STAT(STAT_NumViewOnlineDescriptorHeaps);
 		INC_MEMORY_STAT_BY(STAT_ViewOnlineDescriptorHeapMemory, Heap->GetMemorySize());
 
-		for (ERHIPipeline Pipeline : GetRHIPipelines())
+		for (ERHIPipeline Pipeline : MakeFlagsRange(ERHIPipeline::All))
 		{
 			Heaps[Pipeline] = Heap;
 		}
@@ -237,7 +237,7 @@ void FD3D12OnlineDescriptorManager::Init(uint32 InTotalSize, uint32 InBlockSize,
 
 void FD3D12OnlineDescriptorManager::CleanupResources()
 {
-	for (ERHIPipeline Pipeline : GetRHIPipelines())
+	for (ERHIPipeline Pipeline : MakeFlagsRange(ERHIPipeline::All))
 	{
 		Heaps[Pipeline] = nullptr;
 	}
diff --git a/Engine/Source/Runtime/RHI/Private/RHICommandList.cpp b/Engine/Source/Runtime/RHI/Private/RHICommandList.cpp
--- a/Engine/Source/Runtime/RHI/Private/RHICommandList.cpp
+++ b/Engine/Source/Runtime/RHI/Private/RHICommandList.cpp
@@ -698,7 +698,7 @@ void FRHICommandListExecutor::FSubmitState::Dispatch(FRHICommandListBase* CmdLis
 			}
 
 			TRHIPipelineArray<FRHIBreadcrumbNode*> GPUFirst;
-			for (ERHIPipeline Pipeline : GetRHIPipelines())
+			for (ERHIPipeline Pipeline : MakeFlagsRange(ERHIPipeline::All))
 			{
 				FRHIBreadcrumbNode* Node = GRHICommandList.Breadcrumbs.GPU[Pipeline].Current;
 				check(Node != FRHIBreadcrumbNode::Sentinel);
@@ -726,7 +726,7 @@ void FRHICommandListExecutor::FSubmitState::Dispatch(FRHICommandListBase* CmdLis
 				GRHICommandList.Breadcrumbs.GPU[Command->Pipeline].Current = GRHICommandList.Breadcrumbs.CPU.Current;
 			}
 
-			for (ERHIPipeline Pipeline : GetRHIPipelines())
+			for (ERHIPipeline Pipeline : MakeFlagsRange(ERHIPipeline::All))
 			{
 				if (CmdList->GPUBreadcrumbState[Pipeline].Latest != FRHIBreadcrumbNode::Sentinel)
 				{
@@ -799,7 +799,7 @@ void FRHICommandListExecutor::FSubmitState::Dispatch(FRHICommandListBase* CmdLis
 void FRHICommandListExecutor::FTranslateState::Translate(FRHICommandListBase* CmdList)
 {
 	// Apply the current translate job's contexts to the command list
-	for (ERHIPipeline Pipeline : GetRHIPipelines())
+	for (ERHIPipeline Pipeline : MakeFlagsRange(ERHIPipeline::All))
 	{
 		if (!CmdList->Contexts[Pipeline])
 		{
@@ -824,7 +824,7 @@ void FRHICommandListExecutor::FTranslateState::Translate(FRHICommandListBase* Cm
 #endif
 
 	// Extract the contexts from the command list, so we can reuse them for future command lists.
-	for (ERHIPipeline Pipeline : GetRHIPipelines())
+	for (ERHIPipeline Pipeline : MakeFlagsRange(ERHIPipeline::All))
 	{
 		auto& TranslateContext = PipelineStates[Pipeline].Context;
 		check(!TranslateContext || TranslateContext == CmdList->Contexts[Pipeline]);
@@ -857,9 +857,8 @@ FGraphEventRef FRHICommandListExecutor::FTranslateState::Finalize()
 		{
 			SCOPED_NAMED_EVENT(RHI_Finalize, FColor::White);
 
-			for (ERHIPipeline Pipeline : GetRHIPipelines())
+			for (auto& State : PipelineStates)
 			{
-				auto& State = PipelineStates[Pipeline];
 				if (State.Context)
 				{
 					State.FinalizedCmdList = GDynamicRHI->RHIFinalizeContext({ State.Context });
@@ -888,7 +887,7 @@ void FRHICommandListExecutor::FSubmitState::Submit()
 	TArray<IRHIPlatformCommandList*> FinalizedCmdLists;
 	for (TUniquePtr<FTranslateState> const& Job : TranslateJobs)
 	{
-		for (ERHIPipeline Pipeline : GetRHIPipelines())
+		for (ERHIPipeline Pipeline : MakeFlagsRange(ERHIPipeline::All))
 		{
 			auto& TranslateState = Job->PipelineStates[Pipeline];
 			if (TranslateState.FinalizedCmdList)
@@ -1186,7 +1185,7 @@ void FRHICommandListImmediate::InitializeImmediateContexts()
 #if WITH_RHI_BREADCRUMBS
 		CPUBreadcrumbState.Current = GRHICommandList.Breadcrumbs.CPU.Current;
 		CPUBreadcrumbState.bEmitBreadcrumbs = GRHICommandList.bEmitBreadcrumbs;
-		for (ERHIPipeline Pipeline : GetRHIPipelines())
+		for (ERHIPipeline Pipeline : MakeFlagsRange(ERHIPipeline::All))
 		{
 			GPUBreadcrumbState[Pipeline].Current = GRHICommandList.Breadcrumbs.GPU[Pipeline].Current;
 			GPUBreadcrumbState[Pipeline].Latest = GRHICommandList.Breadcrumbs.GPU[Pipeline].Current;
@@ -1468,17 +1467,17 @@ void FRHICommandListImmediate::Transition(TArrayView<const FRHITransitionInfo> I
 
 	const FRHITransition* Transition = RHICreateTransition({ SrcPipelines, DstPipelines, ERHITransitionCreateFlags::None, Infos });
 
-	EnumerateRHIPipelines(SrcPipelines, [&](ERHIPipeline Pipeline)
+	for (ERHIPipeline Pipeline : MakeFlagsRange(SrcPipelines))
 	{
 		FRHICommandListScopedPipeline Scope(*this, Pipeline);
 		BeginTransition(Transition);
-	});
+	}
 
-	EnumerateRHIPipelines(DstPipelines, [&](ERHIPipeline Pipeline)
+	for (ERHIPipeline Pipeline : MakeFlagsRange(DstPipelines))
 	{
 		FRHICommandListScopedPipeline Scope(*this, Pipeline);
 		EndTransition(Transition);
-	});
+	}
 
 	if (EnumHasAnyFlags(SrcPipelines | DstPipelines, ERHIPipeline::Graphics))
 	{
diff --git a/Engine/Source/Runtime/RHI/Private/RHIValidation.cpp b/Engine/Source/Runtime/RHI/Private/RHIValidation.cpp
--- a/Engine/Source/Runtime/RHI/Private/RHIValidation.cpp
+++ b/Engine/Source/Runtime/RHI/Private/RHIValidation.cpp
@@ -524,16 +524,11 @@ void FValidationRHI::RHICreateTransition(FRHITransition* Transition, const FRHIT
 
 	if (SrcPipelines != DstPipelines)
 	{
-		for (ERHIPipeline SrcPipe : GetRHIPipelines())
+		for (ERHIPipeline SrcPipe : MakeFlagsRange(SrcPipelines))
 		{
-			if (!EnumHasAnyFlags(SrcPipelines, SrcPipe))
+			for (ERHIPipeline DstPipe : MakeFlagsRange(DstPipelines))
 			{
-				continue;
-			}
-
-			for (ERHIPipeline DstPipe : GetRHIPipelines())
-			{
-				if (!EnumHasAnyFlags(DstPipelines, DstPipe) || SrcPipe == DstPipe)
+				if (SrcPipe == DstPipe)
 				{
 					continue;
 				}
@@ -1659,7 +1654,7 @@ namespace RHIValidation
 		State.bTransitioning = true;
 
 		// Replicate the state to other pipes that are not part of the begin pipe mask.
-		for (ERHIPipeline OtherPipeline : GetRHIPipelines())
+		for (ERHIPipeline OtherPipeline : MakeFlagsRange(ERHIPipeline::All))
 		{
 			if (!EnumHasAnyFlags(CurrentStateFromRHI.Pipelines, OtherPipeline))
 			{
@@ -1694,7 +1689,7 @@ namespace RHIValidation
 		RHI_VALIDATION_CHECK(TargetState == State.Current, *GetReasonString_MismatchedEndTransition(Resource, SubresourceIndex, State.Current, TargetState));
 
 		// Replicate the state to other pipes that are not part of the end pipe mask.
-		for (ERHIPipeline OtherPipeline : GetRHIPipelines())
+		for (ERHIPipeline OtherPipeline : MakeFlagsRange(ERHIPipeline::All))
 		{
 			if (!EnumHasAnyFlags(TargetState.Pipelines, OtherPipeline))
 			{
@@ -2667,7 +2662,7 @@ void FValidationTransientResourceAllocator::DeallocateMemory(FRHITransientBuffer
 void FValidationTransientResourceAllocator::Flush(FRHICommandListImmediate& RHICmdList, FRHITransientAllocationStats* OutHeapStats)
 {
 	// Insert pending ops into context trackers
-	for (ERHIPipeline Pipeline : GetRHIPipelines())
+	for (ERHIPipeline Pipeline : MakeFlagsRange(ERHIPipeline::All))
 	{
 		if (PendingPipelineOps[Pipeline].Num())
 		{
diff --git a/Engine/Source/Runtime/RHI/Public/RHIPipeline.h b/Engine/Source/Runtime/RHI/Public/RHIPipeline.h
--- a/Engine/Source/Runtime/RHI/Public/RHIPipeline.h
+++ b/Engine/Source/Runtime/RHI/Public/RHIPipeline.h
@@ -4,6 +4,7 @@
 
 #include "HAL/Platform.h"
 #include "Misc/EnumClassFlags.h"
+#include "Misc/EnumRange.h"
 #include "Containers/ArrayView.h"
 #include "Containers/StaticArray.h"
 
@@ -35,6 +36,7 @@ inline constexpr uint32 GetRHIPipelineCount()
 	return uint32(ERHIPipeline::Num);
 }
 
+UE_DEPRECATED(5.5, "GetRHIPipelines is deprecated. Prefer ranged-for iteration over pipelines using 'for (ERHIPipeline Pipeline : MakeFlagsRange(Pipelines))'.")
 inline TArrayView<const ERHIPipeline> GetRHIPipelines()
 {
 	static const ERHIPipeline Pipelines[] = { ERHIPipeline::Graphics, ERHIPipeline::AsyncCompute };
@@ -42,14 +44,12 @@ inline TArrayView<const ERHIPipeline> GetRHIPipelines()
 }
 
 template <typename FunctionType>
+UE_DEPRECATED(5.5, "EnumerateRHIPipelines is deprecated. Prefer ranged-for iteration over pipelines using 'for (ERHIPipeline Pipeline : MakeFlagsRange(Pipelines))'.")
 inline void EnumerateRHIPipelines(ERHIPipeline PipelineMask, FunctionType Function)
 {
-	for (ERHIPipeline Pipeline : GetRHIPipelines())
+	for (ERHIPipeline Pipeline : MakeFlagsRange(PipelineMask))
 	{
-		if (EnumHasAnyFlags(PipelineMask, Pipeline))
-		{
-			Function(Pipeline);
-		}
+		Function(Pipeline);
 	}
 }
 
diff --git a/Engine/Source/Runtime/RHI/Public/RHIValidationCommon.h b/Engine/Source/Runtime/RHI/Public/RHIValidationCommon.h
--- a/Engine/Source/Runtime/RHI/Public/RHIValidationCommon.h
+++ b/Engine/Source/Runtime/RHI/Public/RHIValidationCommon.h
@@ -359,7 +359,7 @@ namespace RHIValidation
 			DebugName = InDebugName;
 			TrackedAccess = ERHIAccess::Discard;
 
-			for (ERHIPipeline Pipeline : GetRHIPipelines())
+			for (ERHIPipeline Pipeline : MakeFlagsRange(ERHIPipeline::All))
 			{
 				auto& State = WholeResourceState.States[Pipeline];
 
@@ -382,7 +382,7 @@ namespace RHIValidation
 			TransientState = FTransientState(InResourceState);
 			TrackedAccess = InResourceState;
 
-			for (ERHIPipeline Pipeline : GetRHIPipelines())
+			for (ERHIPipeline Pipeline : MakeFlagsRange(ERHIPipeline::All))
 			{
 				auto& State = WholeResourceState.States[Pipeline];
 
@@ -614,12 +614,9 @@ namespace RHIValidation
 
 		static inline FOperation BeginTransitionResource(FResourceIdentity Identity, FState PreviousState, FState NextState, EResourceTransitionFlags Flags, void* CreateBacktrace)
 		{
-			for (ERHIPipeline Pipeline : GetRHIPipelines())
+			for (ERHIPipeline Pipeline : MakeFlagsRange(PreviousState.Pipelines))
 			{
-				if (EnumHasAnyFlags(PreviousState.Pipelines, Pipeline))
-				{
-					Identity.Resource->AddOpRef();
-				}
+				Identity.Resource->AddOpRef();
 			}
 
 			FOperation Op;
@@ -634,12 +631,9 @@ namespace RHIValidation
 
 		static inline FOperation EndTransitionResource(FResourceIdentity Identity, FState PreviousState, FState NextState, void* CreateBacktrace)
 		{
-			for (ERHIPipeline Pipeline : GetRHIPipelines())
+			for (ERHIPipeline Pipeline : MakeFlagsRange(NextState.Pipelines))
 			{
-				if (EnumHasAnyFlags(NextState.Pipelines, Pipeline))
-				{
-					Identity.Resource->AddOpRef();
-				}
+				Identity.Resource->AddOpRef();
 			}
 
 			FOperation Op;
diff --git a/Engine/Source/Runtime/RenderCore/Private/RenderGraphBuilder.cpp b/Engine/Source/Runtime/RenderCore/Private/RenderGraphBuilder.cpp
--- a/Engine/Source/Runtime/RenderCore/Private/RenderGraphBuilder.cpp
+++ b/Engine/Source/Runtime/RenderCore/Private/RenderGraphBuilder.cpp
@@ -1117,7 +1117,7 @@ void FRDGBuilder::AddPassDependency(FRDGPass* Producer, FRDGPass* Consumer)
 
 bool FRDGBuilder::AddCullingDependency(FRDGProducerStatesByPipeline& LastProducers, const FRDGProducerState& NextState, ERHIPipeline NextPipeline)
 {
-	for (ERHIPipeline LastPipeline : GetRHIPipelines())
+	for (ERHIPipeline LastPipeline : MakeFlagsRange(ERHIPipeline::All))
 	{
 		FRDGProducerState& LastProducer = LastProducers[LastPipeline];
 
@@ -3994,7 +3994,7 @@ void FRDGBuilder::AddTransition(
 
 		BarriersToBegin->AddTransition(Resource, TransitionInfo);
 
-		for (ERHIPipeline Pipeline : GetRHIPipelines())
+		for (ERHIPipeline Pipeline : MakeFlagsRange(ERHIPipeline::All))
 		{
 			/** If doing a 1-to-N transition and this is the same pipe as the begin, we end it immediately afterwards in the epilogue
 			 *  of the begin pass. This is because we can't guarantee that the other pipeline won't join back before the end. This can
@@ -4031,7 +4031,7 @@ void FRDGBuilder::AddTransition(
 
 		FRDGBarrierBatchBeginId Id;
 		Id.PipelinesAfter = PipelinesAfter;
-		for (ERHIPipeline Pipeline : GetRHIPipelines())
+		for (ERHIPipeline Pipeline : MakeFlagsRange(ERHIPipeline::All))
 		{
 			Id.Passes[Pipeline] = GetEpilogueBarrierPassHandle(PassesBefore[Pipeline]);
 		}
@@ -4055,12 +4055,9 @@ void FRDGBuilder::AddTransition(
 
 		BarriersToBegin->AddTransition(Resource, TransitionInfo);
 
-		for (ERHIPipeline Pipeline : GetRHIPipelines())
+		for (ERHIPipeline Pipeline : MakeFlagsRange(PipelinesAfter))
 		{
-			if (EnumHasAnyFlags(PipelinesAfter, Pipeline))
-			{
-				AddToPrologueBarriersToEnd(PassesAfter[Pipeline], *BarriersToBegin);
-			}
+			AddToPrologueBarriersToEnd(PassesAfter[Pipeline], *BarriersToBegin);
 		}
 	}
 }
diff --git a/Engine/Source/Runtime/RenderCore/Private/RenderGraphPass.cpp b/Engine/Source/Runtime/RenderCore/Private/RenderGraphPass.cpp
--- a/Engine/Source/Runtime/RenderCore/Private/RenderGraphPass.cpp
+++ b/Engine/Source/Runtime/RenderCore/Private/RenderGraphPass.cpp
@@ -195,7 +195,7 @@ void FRDGBarrierBatchEnd::AddDependency(FRDGBarrierBatchBegin* BeginBatch)
 #if RDG_ENABLE_DEBUG
 	check(BeginBatch);
 
-	for (ERHIPipeline Pipeline : GetRHIPipelines())
+	for (ERHIPipeline Pipeline : MakeFlagsRange(ERHIPipeline::All))
 	{
 		const FRDGPass* BeginPass = BeginBatch->DebugPasses[Pipeline];
 		if (BeginPass)
diff --git a/Engine/Source/Runtime/RenderCore/Public/RenderGraphResources.inl b/Engine/Source/Runtime/RenderCore/Public/RenderGraphResources.inl
--- a/Engine/Source/Runtime/RenderCore/Public/RenderGraphResources.inl
+++ b/Engine/Source/Runtime/RenderCore/Public/RenderGraphResources.inl
@@ -15,7 +15,7 @@ inline void FRDGSubresourceState::SetPass(ERHIPipeline Pipeline, FRDGPassHandle
 inline void FRDGSubresourceState::Validate()
 {
 #if RDG_ENABLE_DEBUG
-	for (ERHIPipeline Pipeline : GetRHIPipelines())
+	for (ERHIPipeline Pipeline : MakeFlagsRange(ERHIPipeline::All))
 	{
 		checkf(FirstPass[Pipeline].IsValid() == LastPass[Pipeline].IsValid(), TEXT("Subresource state has unset first or last pass on '%s."), *GetRHIPipelineName(Pipeline));
 	}
