Fixes for Android BootTest

diff --git a/Engine/Source/Runtime/RHI/Private/PipelineStateCache.cpp b/Engine/Source/Runtime/RHI/Private/PipelineStateCache.cpp
--- a/Engine/Source/Runtime/RHI/Private/PipelineStateCache.cpp
+++ b/Engine/Source/Runtime/RHI/Private/PipelineStateCache.cpp
@@ -3205,6 +3205,8 @@ void PipelineStateCache::Shutdown()
 {
 	GComputePipelineCache.WaitTasksComplete();
 	GGraphicsPipelineCache.WaitTasksComplete();
+	GPrecacheGraphicsPipelineCache.WaitTasksComplete();
+	GPrecacheComputePipelineCache.WaitTasksComplete();
 #if RHI_RAYTRACING
 	GRayTracingPipelineCache.Shutdown();
 #endif
