Removed explicit RHI thread fence in RHIEndViewport, added in 3162008 to fix UE-37216, as the original issue has been solved.

diff --git a/Engine/Source/Runtime/RHI/Private/RHICommandList.cpp b/Engine/Source/Runtime/RHI/Private/RHICommandList.cpp
--- a/Engine/Source/Runtime/RHI/Private/RHICommandList.cpp
+++ b/Engine/Source/Runtime/RHI/Private/RHICommandList.cpp
@@ -103,9 +103,6 @@ uint64 GInputLatencyTime = 0;
 
 FRHICommandListExecutor GRHICommandList;
 
-static FGraphEventRef GRHIThreadEndDrawingViewportFences[2];
-static uint32 GRHIThreadEndDrawingViewportFenceIndex = 0;
-
 static FCriticalSection GRHIThreadOnTasksCritical;
 static std::atomic<int32> GRHIThreadStallRequestCount;
 
@@ -1614,11 +1611,6 @@ void FRHICommandListImmediate::EndDrawingViewport(FRHIViewport* Viewport, bool b
 	{
 		ALLOC_COMMAND(FRHICommandEndDrawingViewport)(Viewport, bPresent, bLockToVsync);
 
-		if (IsRunningRHIInSeparateThread())
-		{
-			// Insert a fence to prevent the renderthread getting more than a frame ahead of the RHIThread
-			GRHIThreadEndDrawingViewportFences[GRHIThreadEndDrawingViewportFenceIndex] = static_cast<FRHICommandListImmediate*>(this)->RHIThreadFence();
-		}
 		// if we aren't running an RHIThread, there is no good reason to buffer this frame advance stuff and that complicates state management, so flush everything out now
 		{
 			QUICK_SCOPE_CYCLE_COUNTER(STAT_EndDrawingViewport_Dispatch);
@@ -1626,16 +1618,6 @@ void FRHICommandListImmediate::EndDrawingViewport(FRHIViewport* Viewport, bool b
 		}
 	}
 
-	if (IsRunningRHIInSeparateThread())
-	{
-		// Wait on the previous frame's RHI thread fence (we never want the rendering thread to get more than a frame ahead)
-		uint32 PreviousFrameFenceIndex = 1 - GRHIThreadEndDrawingViewportFenceIndex;
-		FGraphEventRef& LastFrameFence = GRHIThreadEndDrawingViewportFences[PreviousFrameFenceIndex];
-		FRHICommandListExecutor::WaitOnRHIThreadFence(LastFrameFence);
-		GRHIThreadEndDrawingViewportFences[PreviousFrameFenceIndex] = nullptr;
-		GRHIThreadEndDrawingViewportFenceIndex = PreviousFrameFenceIndex;
-	}
-
 	RHIAdvanceFrameForGetViewportBackBuffer(Viewport);
 }
 
@@ -2160,9 +2142,4 @@ void FRHICommandListExecutor::CleanupGraphEvents()
 
 	LastMutate.SafeRelease();
 	LastSubmit.SafeRelease();
-
-	for (FGraphEventRef& GraphEvent : GRHIThreadEndDrawingViewportFences)
-	{
-		GraphEvent.SafeRelease();
-	}
 }
