Build fix for new function placed incorrectly inside RHI_RAYTRACING define.

diff --git a/Engine/Source/Runtime/RHI/Private/PipelineStateCache.cpp b/Engine/Source/Runtime/RHI/Private/PipelineStateCache.cpp
--- a/Engine/Source/Runtime/RHI/Private/PipelineStateCache.cpp
+++ b/Engine/Source/Runtime/RHI/Private/PipelineStateCache.cpp
@@ -748,6 +748,20 @@ public:
 	uint64 SortKey = 0;
 };
 
+FRHIComputePipelineState* GetRHIComputePipelineState(FComputePipelineState* PipelineState)
+{
+	ensure(PipelineState->RHIPipeline);
+	PipelineState->CompletionEvent = nullptr;
+	return PipelineState->RHIPipeline;
+}
+
+FRHIWorkGraphPipelineState* GetRHIWorkGraphPipelineState(FWorkGraphPipelineState* PipelineState)
+{
+	ensure(PipelineState->RHIPipeline);
+	PipelineState->CompletionEvent = nullptr;
+	return PipelineState->RHIPipeline;
+}
+
 #if RHI_RAYTRACING
 /* State for ray tracing */
 class FRayTracingPipelineState : public FPipelineState
@@ -831,22 +845,7 @@ public:
 #endif
 };
 
-RHI_API FRHIComputePipelineState* GetRHIComputePipelineState(FComputePipelineState* PipelineState)
-{
-	ensure(PipelineState->RHIPipeline);
-	PipelineState->CompletionEvent = nullptr;
-	return PipelineState->RHIPipeline;
-}
-
-RHI_API FRHIWorkGraphPipelineState* GetRHIWorkGraphPipelineState(FWorkGraphPipelineState* PipelineState)
-{
-	ensure(PipelineState->RHIPipeline);
-	PipelineState->CompletionEvent = nullptr;
-	return PipelineState->RHIPipeline;
-}
-
-//extern RHI_API FRHIRayTracingPipelineState* GetRHIRayTracingPipelineState(FRayTracingPipelineState* PipelineState);
-RHI_API FRHIRayTracingPipelineState* GetRHIRayTracingPipelineState(FRayTracingPipelineState* PipelineState)
+FRHIRayTracingPipelineState* GetRHIRayTracingPipelineState(FRayTracingPipelineState* PipelineState)
 {
 	ensure(PipelineState->RHIPipeline);
 	PipelineState->CompletionEvent = nullptr;
