Remove UE 5.3 deprecated RHI functions

diff --git a/Engine/Source/Runtime/RHI/Public/RHICommandList.h b/Engine/Source/Runtime/RHI/Public/RHICommandList.h
--- a/Engine/Source/Runtime/RHI/Public/RHICommandList.h
+++ b/Engine/Source/Runtime/RHI/Public/RHICommandList.h
@@ -4798,150 +4798,6 @@ FORCEINLINE_DEBUGGABLE FRHICommandListImmediate& FRHICommandListExecutor::GetImm
 	return GRHICommandList.CommandListImmediate;
 }
 
-UE_DEPRECATED(5.3, "RHICreateBuffer is deprecated. Use FRHICommandListBase::CreateBuffer instead.")
-FORCEINLINE FBufferRHIRef RHICreateBuffer(uint32 Size, EBufferUsageFlags Usage, uint32 Stride, ERHIAccess ResourceState, FRHIResourceCreateInfo& CreateInfo)
-{
-	return FRHICommandListImmediate::Get().CreateBuffer(Size, Usage, Stride, ResourceState, CreateInfo);
-}
-
-UE_DEPRECATED(5.3, "RHICreateIndexBuffer is deprecated. Use FRHICommandListBase::CreateIndexBuffer instead.")
-FORCEINLINE FBufferRHIRef RHICreateIndexBuffer(uint32 Stride, uint32 Size, EBufferUsageFlags Usage, ERHIAccess ResourceState, FRHIResourceCreateInfo& CreateInfo)
-{
-	return FRHICommandListImmediate::Get().CreateIndexBuffer(Stride, Size, Usage, ResourceState, CreateInfo);
-}
-
-UE_DEPRECATED(5.3, "RHICreateIndexBuffer is deprecated. Use FRHICommandListBase::CreateIndexBuffer instead.")
-FORCEINLINE FBufferRHIRef RHICreateIndexBuffer(uint32 Stride, uint32 Size, EBufferUsageFlags Usage, FRHIResourceCreateInfo& CreateInfo)
-{
-	return FRHICommandListImmediate::Get().CreateIndexBuffer(Stride, Size, Usage, CreateInfo);
-}
-
-UE_DEPRECATED(5.3, "RHIUpdateUniformBuffer is deprecated. Use FRHICommandListBase::UpdateUniformBuffer instead.")
-FORCEINLINE void RHIUpdateUniformBuffer(FRHIUniformBuffer* UniformBufferRHI, const void* Contents)
-{
-	return FRHICommandListExecutor::GetImmediateCommandList().UpdateUniformBuffer(UniformBufferRHI, Contents);
-}
-
-UE_DEPRECATED(5.3, "RHICreateVertexBuffer is deprecated. Use FRHICommandListBase::CreateVertexBuffer instead.")
-FORCEINLINE FBufferRHIRef RHICreateVertexBuffer(uint32 Size, EBufferUsageFlags Usage, ERHIAccess ResourceState, FRHIResourceCreateInfo& CreateInfo)
-{
-	return FRHICommandListImmediate::Get().CreateVertexBuffer(Size, Usage, ResourceState, CreateInfo);
-}
-
-UE_DEPRECATED(5.3, "RHICreateVertexBuffer is deprecated. Use FRHICommandListBase::CreateVertexBuffer instead.")
-FORCEINLINE FBufferRHIRef RHICreateVertexBuffer(uint32 Size, EBufferUsageFlags Usage, FRHIResourceCreateInfo& CreateInfo)
-{
-	return FRHICommandListImmediate::Get().CreateVertexBuffer(Size, Usage, CreateInfo);
-}
-
-UE_DEPRECATED(5.3, "RHICreateStructuredBuffer is deprecated. Use FRHICommandListBase::CreateStructuredBuffer instead.")
-FORCEINLINE FBufferRHIRef RHICreateStructuredBuffer(uint32 Stride, uint32 Size, EBufferUsageFlags Usage, ERHIAccess ResourceState, FRHIResourceCreateInfo& CreateInfo)
-{
-	return FRHICommandListImmediate::Get().CreateStructuredBuffer(Stride, Size, Usage, ResourceState, CreateInfo);
-}
-
-UE_DEPRECATED(5.3, "RHICreateStructuredBuffer is deprecated. Use FRHICommandListBase::CreateStructuredBuffer instead.")
-FORCEINLINE FBufferRHIRef RHICreateStructuredBuffer(uint32 Stride, uint32 Size, EBufferUsageFlags Usage, FRHIResourceCreateInfo& CreateInfo)
-{
-	return FRHICommandListImmediate::Get().CreateStructuredBuffer(Stride, Size, Usage, CreateInfo);
-}
-
-UE_DEPRECATED(5.3, "RHILockBuffer is deprecated. Use FRHICommandListBase::LockBuffer instead.")
-FORCEINLINE void* RHILockBuffer(FRHIBuffer* Buffer, uint32 Offset, uint32 SizeRHI, EResourceLockMode LockMode)
-{
-	return FRHICommandListImmediate::Get().LockBuffer(Buffer, Offset, SizeRHI, LockMode);
-}
-
-UE_DEPRECATED(5.3, "RHIUnlockBuffer is deprecated. Use FRHICommandListBase::UnlockBuffer instead.")
-FORCEINLINE void RHIUnlockBuffer(FRHIBuffer* Buffer)
-{
-	FRHICommandListImmediate::Get().UnlockBuffer(Buffer);
-}
-
-UE_DEPRECATED(5.3, "RHICreateShaderResourceView is deprecated. Use FRHICommandListBase::CreateShaderResourceView instead.")
-FORCEINLINE FShaderResourceViewRHIRef RHICreateShaderResourceView(FRHIBuffer* Buffer, FRHIViewDesc::FBufferSRV::FInitializer const& ViewDesc)
-{
-	return FRHICommandListExecutor::GetImmediateCommandList().CreateShaderResourceView(Buffer, ViewDesc);
-}
-
-UE_DEPRECATED(5.3, "RHICreateShaderResourceView is deprecated. Use FRHICommandListBase::CreateShaderResourceView instead.")
-FORCEINLINE FShaderResourceViewRHIRef RHICreateShaderResourceView(FRHITexture* Texture, FRHIViewDesc::FTextureSRV::FInitializer const& ViewDesc)
-{
-	return FRHICommandListExecutor::GetImmediateCommandList().CreateShaderResourceView(Texture, ViewDesc);
-}
-
-UE_DEPRECATED(5.3, "RHICreateShaderResourceView is deprecated. Use FRHICommandListBase::CreateShaderResourceView instead.")
-FORCEINLINE FUnorderedAccessViewRHIRef RHICreateUnorderedAccessView(FRHIBuffer* Buffer, FRHIViewDesc::FBufferUAV::FInitializer const& ViewDesc)
-{
-	return FRHICommandListExecutor::GetImmediateCommandList().CreateUnorderedAccessView(Buffer, ViewDesc);
-}
-
-UE_DEPRECATED(5.3, "RHICreateShaderResourceView is deprecated. Use FRHICommandListBase::CreateShaderResourceView instead.")
-FORCEINLINE FUnorderedAccessViewRHIRef RHICreateUnorderedAccessView(FRHITexture* Texture, FRHIViewDesc::FTextureUAV::FInitializer const& ViewDesc)
-{
-	return FRHICommandListExecutor::GetImmediateCommandList().CreateUnorderedAccessView(Texture, ViewDesc);
-}
-
-UE_DEPRECATED(5.3, "RHICreateUnorderedAccessView is deprecated. Use FRHICommandListBase::CreateUnorderedAccessView instead.")
-FORCEINLINE FUnorderedAccessViewRHIRef RHICreateUnorderedAccessView(FRHIBuffer* Buffer, bool bUseUAVCounter, bool bAppendBuffer)
-{
-	return FRHICommandListExecutor::GetImmediateCommandList().CreateUnorderedAccessView(Buffer, bUseUAVCounter, bAppendBuffer);
-}
-
-UE_DEPRECATED(5.3, "RHICreateUnorderedAccessView is deprecated. Use FRHICommandListBase::CreateUnorderedAccessView instead.")
-FORCEINLINE FUnorderedAccessViewRHIRef RHICreateUnorderedAccessView(FRHITexture* Texture, uint32 MipLevel = 0, uint16 FirstArraySlice = 0, uint16 NumArraySlices = 0)
-{
-	return FRHICommandListExecutor::GetImmediateCommandList().CreateUnorderedAccessView(Texture, MipLevel, FirstArraySlice, NumArraySlices);
-}
-
-UE_DEPRECATED(5.3, "RHICreateUnorderedAccessView is deprecated. Use FRHICommandListBase::CreateUnorderedAccessView instead.")
-FORCEINLINE FUnorderedAccessViewRHIRef RHICreateUnorderedAccessView(FRHITexture* Texture, uint32 MipLevel, uint8 Format, uint16 FirstArraySlice = 0, uint16 NumArraySlices = 0)
-{
-	return FRHICommandListExecutor::GetImmediateCommandList().CreateUnorderedAccessView(Texture, MipLevel, Format, FirstArraySlice, NumArraySlices);
-}
-
-UE_DEPRECATED(5.3, "RHICreateUnorderedAccessView is deprecated. Use FRHICommandListBase::CreateUnorderedAccessView instead.")
-FORCEINLINE FUnorderedAccessViewRHIRef RHICreateUnorderedAccessView(FRHIBuffer* Buffer, uint8 Format)
-{
-	return FRHICommandListExecutor::GetImmediateCommandList().CreateUnorderedAccessView(Buffer, Format);
-}
-
-UE_DEPRECATED(5.3, "RHICreateUnorderedAccessView is deprecated. Use FRHICommandListBase::CreateUnorderedAccessView instead.")
-FORCEINLINE FShaderResourceViewRHIRef RHICreateShaderResourceView(FRHIBuffer* Buffer)
-{
-	return FRHICommandListExecutor::GetImmediateCommandList().CreateShaderResourceView(Buffer);
-}
-
-UE_DEPRECATED(5.3, "RHICreateUnorderedAccessView is deprecated. Use FRHICommandListBase::CreateUnorderedAccessView instead.")
-FORCEINLINE FShaderResourceViewRHIRef RHICreateShaderResourceView(FRHIBuffer* Buffer, uint32 Stride, uint8 Format)
-{
-	return FRHICommandListExecutor::GetImmediateCommandList().CreateShaderResourceView(Buffer, Stride, Format);
-}
-
-UE_DEPRECATED(5.3, "RHICreateUnorderedAccessView is deprecated. Use FRHICommandListBase::CreateUnorderedAccessView instead.")
-FORCEINLINE FShaderResourceViewRHIRef RHICreateShaderResourceView(const FShaderResourceViewInitializer& Initializer)
-{
-	return FRHICommandListExecutor::GetImmediateCommandList().CreateShaderResourceView(Initializer.Buffer, Initializer);
-}
-
-UE_DEPRECATED(5.3, "RHICreateUnorderedAccessView is deprecated. Use FRHICommandListBase::CreateUnorderedAccessView instead.")
-FORCEINLINE FShaderResourceViewRHIRef RHICreateShaderResourceView(FRHITexture* Texture, uint8 MipLevel)
-{
-	return FRHICommandListExecutor::GetImmediateCommandList().CreateShaderResourceView(Texture, MipLevel);
-}
-
-UE_DEPRECATED(5.3, "RHICreateUnorderedAccessView is deprecated. Use FRHICommandListBase::CreateUnorderedAccessView instead.")
-FORCEINLINE FShaderResourceViewRHIRef RHICreateShaderResourceView(FRHITexture* Texture, uint8 MipLevel, uint8 NumMipLevels, EPixelFormat Format)
-{
-	return FRHICommandListExecutor::GetImmediateCommandList().CreateShaderResourceView(Texture, MipLevel, NumMipLevels, Format);
-}
-
-UE_DEPRECATED(5.3, "RHICreateUnorderedAccessView is deprecated. Use FRHICommandListBase::CreateUnorderedAccessView instead.")
-FORCEINLINE FShaderResourceViewRHIRef RHICreateShaderResourceView(FRHITexture* Texture, const FRHITextureSRVCreateInfo& CreateInfo)
-{
-	return FRHICommandListExecutor::GetImmediateCommandList().CreateShaderResourceView(Texture, CreateInfo);
-}
-
 FORCEINLINE FTextureReferenceRHIRef RHICreateTextureReference(FRHITexture* InReferencedTexture = nullptr)
 {
 	return FRHICommandListImmediate::Get().CreateTextureReference(InReferencedTexture);
