The -Bindless command line option shouldn't force shader platforms' support to "AllShaders" when they only support Bindless in RayTracing shaders.

diff --git a/Engine/Source/Runtime/RHI/Private/RHI.cpp b/Engine/Source/Runtime/RHI/Private/RHI.cpp
--- a/Engine/Source/Runtime/RHI/Private/RHI.cpp
+++ b/Engine/Source/Runtime/RHI/Private/RHI.cpp
@@ -1219,7 +1219,8 @@ ERHIBindlessConfiguration RHIParseBindlessConfiguration(EShaderPlatform Platform
 	static const bool bCommandLine = FParse::Param(FCommandLine::Get(), TEXT("Bindless"));
 	if (bCommandLine)
 	{
-		return ERHIBindlessConfiguration::AllShaders;
+		// Only allow what the platform supports
+		return (BindlessSupport == ERHIBindlessSupport::RayTracingOnly) ? ERHIBindlessConfiguration::RayTracingShaders : ERHIBindlessConfiguration::AllShaders;
 	}
 #endif
 
