Fix missing RDG pass names when using 'profilegpu'
 - Force the full breadcrumb tree to be generated while profilegpu is active.
 - Remove GRHIGlobals.GPUProfile_ForceNoParallelTranslate since GTriggerGPUProfile achieves the same thing.

diff --git a/Engine/Source/Runtime/Engine/Private/UnrealEngine.cpp b/Engine/Source/Runtime/Engine/Private/UnrealEngine.cpp
--- a/Engine/Source/Runtime/Engine/Private/UnrealEngine.cpp
+++ b/Engine/Source/Runtime/Engine/Private/UnrealEngine.cpp
@@ -5872,7 +5872,6 @@ bool UEngine::HandleProfileCommand( const TCHAR* Cmd, FOutputDevice& Ar )
 				[](FRHICommandListImmediate& RHICmdList)
 			{
 				RHICmdList.ImmediateFlush(EImmediateFlushType::FlushRHIThread);
-				GRHIGlobals.GPUProfile_ForceNoParallelTranslate = true;
 				GTriggerGPUProfile = true;
 			});
 			Ar.Logf(TEXT("Profiling the next GPU frame"));
@@ -5912,7 +5911,6 @@ bool UEngine::HandleProfileGPUCommand( const TCHAR* Cmd, FOutputDevice& Ar )
 				[](FRHICommandListImmediate& RHICmdList)
 			{
 				RHICmdList.ImmediateFlush(EImmediateFlushType::FlushRHIThread);
-				GRHIGlobals.GPUProfile_ForceNoParallelTranslate = true;
 				GTriggerGPUProfile = true;
 			});
 			Ar.Logf(TEXT("Profiling the next GPU frame"));
diff --git a/Engine/Source/Runtime/RHI/Private/GPUProfiler.cpp b/Engine/Source/Runtime/RHI/Private/GPUProfiler.cpp
--- a/Engine/Source/Runtime/RHI/Private/GPUProfiler.cpp
+++ b/Engine/Source/Runtime/RHI/Private/GPUProfiler.cpp
@@ -608,10 +608,6 @@ void FGPUProfilerEventNodeFrame::DumpEventTree()
 		}
 #endif
 	}
-
-#if WITH_PROFILEGPU
-	GRHIGlobals.GPUProfile_ForceNoParallelTranslate = false;
-#endif
 }
 
 void FGPUProfiler::PushEvent(const TCHAR* Name, FColor Color)
diff --git a/Engine/Source/Runtime/RHI/Private/RHICommandList.cpp b/Engine/Source/Runtime/RHI/Private/RHICommandList.cpp
--- a/Engine/Source/Runtime/RHI/Private/RHICommandList.cpp
+++ b/Engine/Source/Runtime/RHI/Private/RHICommandList.cpp
@@ -1076,7 +1076,7 @@ RHI_API void FRHICommandListExecutor::Submit(TConstArrayView<FRHICommandListBase
 #endif
 
 		// Prevent use of parallel contexts if unsupported by the RHI, while the legacy 'profilegpu' command is active, or while disabled by the cvar.
-		if (GRHISupportsParallelRHIExecute && !GRHIGlobals.GPUProfile_ForceNoParallelTranslate && CVarRHICmdParallelTranslateEnable.GetValueOnRenderThread())
+		if (GRHISupportsParallelRHIExecute && !GTriggerGPUProfile && CVarRHICmdParallelTranslateEnable.GetValueOnRenderThread())
 		{
 			SubmitState->MaxCommandsPerTranslate     = CVarRHICmdParallelTranslateMaxCommandsPerTranslate.GetValueOnRenderThread();
 			SubmitState->bAllowSingleParallelCombine = CVarRHICmdParallelTranslateCombineSingleAndParallel.GetValueOnRenderThread();
diff --git a/Engine/Source/Runtime/RHI/Public/RHIGlobals.h b/Engine/Source/Runtime/RHI/Public/RHIGlobals.h
--- a/Engine/Source/Runtime/RHI/Public/RHIGlobals.h
+++ b/Engine/Source/Runtime/RHI/Public/RHIGlobals.h
@@ -360,13 +360,6 @@ struct FRHIGlobals
 	/** Whether the next frame should profile the GPU. */
 	bool TriggerGPUProfile = false;
 
-#if WITH_PROFILEGPU
-	/** Workaround flag used by 'profilegpu' command to disable paralle translate (legacy GPU profiler is not compatible with threaded command translation). */
-	bool GPUProfile_ForceNoParallelTranslate = false;
-#else
-	static constexpr bool GPUProfile_ForceNoParallelTranslate = false;
-#endif
-
 	/** Whether we are profiling GPU hitches. */
 	bool TriggerGPUHitchProfile = false;
 
diff --git a/Engine/Source/Runtime/RenderCore/Private/RenderGraphPrivate.cpp b/Engine/Source/Runtime/RenderCore/Private/RenderGraphPrivate.cpp
--- a/Engine/Source/Runtime/RenderCore/Private/RenderGraphPrivate.cpp
+++ b/Engine/Source/Runtime/RenderCore/Private/RenderGraphPrivate.cpp
@@ -682,9 +682,9 @@ FRDGScopeState::FState::FState(bool bInImmediate, bool bInParallelExecute)
 		bool bRDGChannelEnabled = false;
 		IF_RDG_ENABLE_TRACE(bRDGChannelEnabled = UE_TRACE_CHANNELEXPR_IS_ENABLED(RDGChannel));
 
-		if (FRDGBuilder::IsDumpingFrame())
+		if (FRDGBuilder::IsDumpingFrame() || GTriggerGPUProfile)
 		{
-			// We want all possible scope and pass names in a DumpGPU trace.
+			// We want all possible scope and pass names in a DumpGPU/profilegpu trace.
 			return ERDGScopeMode::AllEventsAndPassNames;
 		}
 
