Wait for all PSO precache tasks to complete before shutdown.

diff --git a/Engine/Source/Runtime/Launch/Private/LaunchEngineLoop.cpp b/Engine/Source/Runtime/Launch/Private/LaunchEngineLoop.cpp
--- a/Engine/Source/Runtime/Launch/Private/LaunchEngineLoop.cpp
+++ b/Engine/Source/Runtime/Launch/Private/LaunchEngineLoop.cpp
@@ -5041,6 +5041,13 @@ void FEngineLoop::Exit()
 #endif // WITH_EDITOR
 	FModuleManager::Get().UnloadModule("WorldBrowser", true);
 
+	// Clean up all cached pipelines, before the garbage collector shuts down.
+	ENQUEUE_RENDER_COMMAND(ShutdownPipelineStateCache)([](FRHICommandListImmediate& RHICmdList)
+	{
+		PipelineStateCache::Shutdown();
+	});
+	FlushRenderingCommands();
+
 	AppPreExit();
 
 	TermGamePhys();
@@ -5057,7 +5064,7 @@ void FEngineLoop::Exit()
 
 	// Stop the rendering thread.
 	ShutdownRenderingThread();
-	
+
 	// Disable the PSO cache
 	FShaderPipelineCache::Shutdown();
 
diff --git a/Engine/Source/Runtime/RHI/Private/PipelineStateCache.cpp b/Engine/Source/Runtime/RHI/Private/PipelineStateCache.cpp
--- a/Engine/Source/Runtime/RHI/Private/PipelineStateCache.cpp
+++ b/Engine/Source/Runtime/RHI/Private/PipelineStateCache.cpp
@@ -1359,20 +1359,41 @@ public:
 
 	void WaitTasksComplete()
 	{
-		FRWScopeLock WriteLock(PrecachePSOsRWLock, SLT_Write);
-
-		for (auto Iterator = PrecachedPSOInitializerData.CreateIterator(); Iterator; ++Iterator)
+		// We hold the lock to observe task state, releasing it if tasks are still in flight
+		// precache tasks may also attempt to lock PrecachePSOsRWLock (TPrecachePipelineCacheBase::PrecacheFinished).
+		// TODO: Replace all of this spin wait.
+		bool bTasksWaiting = true;
+		while(bTasksWaiting)
 		{
-			FPrecacheTask& PrecacheTask = Iterator->Value;
-			if (PrecacheTask.PipelineState && !PrecacheTask.PipelineState->IsComplete())
+			bTasksWaiting = false;
+			{
+				FRWScopeLock WriteLock(PrecachePSOsRWLock, SLT_Write);
+				for (auto Iterator = PrecachedPSOInitializerData.CreateIterator(); Iterator; ++Iterator)
+				{
+					FPrecacheTask& PrecacheTask = Iterator->Value;
+					if (PrecacheTask.PipelineState && !PrecacheTask.PipelineState->IsComplete())
+					{
+						bTasksWaiting = true;						
+						break; // release PrecachePSOsRWLock so's to avoid any further blocking of in-progress tasks.
+					}
+					else if (PrecacheTask.PipelineState)
+					{
+						check(EnumHasAnyFlags(PrecacheTask.ReadPSOPrecacheState(), (EPSOPrecacheStateMask::Succeeded | EPSOPrecacheStateMask::Failed)));
+						delete PrecacheTask.PipelineState;
+						PrecacheTask.PipelineState = nullptr;
+					}
+				}
+				if (!bTasksWaiting)
+				{
+					PrecachedPSOs.Empty();
+				}
+			}
+			if (bTasksWaiting)
 			{
-				PrecacheTask.PipelineState->WaitCompletion();
-				check(EnumHasAnyFlags(PrecacheTask.ReadPSOPrecacheState(), (EPSOPrecacheStateMask::Succeeded | EPSOPrecacheStateMask::Failed)));
-				delete PrecacheTask.PipelineState;
-				PrecacheTask.PipelineState = nullptr;
+				// Yield while we wait.
+				FPlatformProcess::Sleep(0.01f);
 			}
 		}
-		PrecachedPSOs.Empty();
 	}
 
 	EPSOPrecacheResult GetPrecachingState(const FPSOPrecacheRequestID& RequestID)
