Deprecate RHICmdList.DiscardRenderTargets.
 - It is entirely unused now.

diff --git a/Engine/Source/Runtime/RHI/Public/RHICommandList.h b/Engine/Source/Runtime/RHI/Public/RHICommandList.h
--- a/Engine/Source/Runtime/RHI/Public/RHICommandList.h
+++ b/Engine/Source/Runtime/RHI/Public/RHICommandList.h
@@ -3602,6 +3602,7 @@ public:
 	FORCEINLINE_DEBUGGABLE void RHIInvalidateCachedState()
 	{}
 
+	UE_DEPRECATED(5.5, "DiscardRenderTargets is deprecated and will be removed in future engine versions.")
 	FORCEINLINE void DiscardRenderTargets(bool Depth, bool Stencil, uint32 ColorBitMask)
 	{
 		if (Bypass())
