Fix dangling ptr crash in FRayTracingPipelineCache::Shutdown()

diff --git a/Engine/Source/Runtime/RHI/Private/PipelineStateCache.cpp b/Engine/Source/Runtime/RHI/Private/PipelineStateCache.cpp
--- a/Engine/Source/Runtime/RHI/Private/PipelineStateCache.cpp
+++ b/Engine/Source/Runtime/RHI/Private/PipelineStateCache.cpp
@@ -1843,6 +1843,7 @@ public:
 			{
 				It.Value->WaitCompletion();
 				delete It.Value;
+				It.Value = nullptr;
 			}
 		}
 		for (auto& It : PartialPipelines)
@@ -1851,8 +1852,11 @@ public:
 			{
 				It.Value->WaitCompletion();
 				delete It.Value;
+				It.Value = nullptr;
 			}
 		}
+		FullPipelines.Reset();
+		PartialPipelines.Reset();
 	}
 
 	void Trim(int32 TargetNumEntries)
