Do not attempt to cleanup the pipeline state cache twice.

The cache gets cleaned up in the rendering thread.

#tests tested locally

diff --git a/Engine/Source/Runtime/RHI/Private/DynamicRHI.cpp b/Engine/Source/Runtime/RHI/Private/DynamicRHI.cpp
--- a/Engine/Source/Runtime/RHI/Private/DynamicRHI.cpp
+++ b/Engine/Source/Runtime/RHI/Private/DynamicRHI.cpp
@@ -420,9 +420,6 @@ void RHIExit()
 {
 	if (!GUsingNullRHI && GDynamicRHI != NULL)
 	{
-		// Clean up all cached pipelines
-		PipelineStateCache::Shutdown();
-
 		// Flush any potential commands queued before we shut things down.
 		FRHICommandListExecutor::GetImmediateCommandList().ImmediateFlush(EImmediateFlushType::FlushRHIThread);
 
