Added static feature level support for IOS

diff --git a/Engine/Source/Runtime/Apple/MetalRHI/Private/MetalRHI.cpp b/Engine/Source/Runtime/Apple/MetalRHI/Private/MetalRHI.cpp
--- a/Engine/Source/Runtime/Apple/MetalRHI/Private/MetalRHI.cpp
+++ b/Engine/Source/Runtime/Apple/MetalRHI/Private/MetalRHI.cpp
@@ -297,6 +297,14 @@ FMetalDynamicRHI::FMetalDynamicRHI(ERHIFeatureLevel::Type RequestedFeatureLevel)
 #endif	// PLATFORM_TVOS
         GMaxRHIFeatureLevel = ERHIFeatureLevel::ES3_1;
 	}
+
+	#if USE_STATIC_SHADER_PLATFORM_ENUMS
+		GMaxRHIShaderPlatform = UE_IOS_STATIC_SHADER_PLATFORM;
+	#endif
+	#if	USE_STATIC_FEATURE_LEVEL_ENUMS
+		GMaxRHIFeatureLevel = UE_IOS_STATIC_FEATURE_LEVEL;
+	#endif
+		
 		
 	FPlatformMemoryStats Stats = FPlatformMemory::GetStats();
 		
diff --git a/Engine/Source/Runtime/RHI/Public/IOS/IOSDataDrivenShaderPlatformInfo.inl b/Engine/Source/Runtime/RHI/Public/IOS/IOSDataDrivenShaderPlatformInfo.inl
new file mode 100644
--- /dev/null
+++ b/Engine/Source/Runtime/RHI/Public/IOS/IOSDataDrivenShaderPlatformInfo.inl
@@ -0,0 +1,16 @@
+// Copyright Epic Games, Inc. All Rights Reserved.
+
+#pragma once
+
+class FDataDrivenShaderPlatformInfo final : public FGenericDataDrivenShaderPlatformInfo
+{
+public:
+	IMPLEMENT_DDPSPI_SETTING(GetTargetsTiledGPU,				true);
+	IMPLEMENT_DDPSPI_SETTING(GetSupportsUniformBufferObjects,	false);
+	IMPLEMENT_DDPSPI_SETTING(GetSupportsManualVertexFetch,		false);
+
+#if USE_STATIC_FEATURE_LEVEL_ENUMS
+	IMPLEMENT_DDPSPI_SETTING_WITH_RETURN_TYPE(ERHIFeatureLevel::Type, GetMaxFeatureLevel, FStaticFeatureLevel(UE_IOS_STATIC_FEATURE_LEVEL));
+	IMPLEMENT_DDPSPI_SETTING(GetIsMobile, (UE_IOS_STATIC_FEATURE_LEVEL == ERHIFeatureLevel::ES3_1));
+#endif
+};
diff --git a/Engine/Source/Runtime/RHI/Public/IOS/IOSStaticFeatureLevel.inl b/Engine/Source/Runtime/RHI/Public/IOS/IOSStaticFeatureLevel.inl
new file mode 100644
--- /dev/null
+++ b/Engine/Source/Runtime/RHI/Public/IOS/IOSStaticFeatureLevel.inl
@@ -0,0 +1,57 @@
+// Copyright Epic Games, Inc. All Rights Reserved.
+
+#pragma once
+
+#if defined(UE_IOS_SHADER_PLATFORM_METAL_MRT)
+	#define UE_IOS_STATIC_FEATURE_LEVEL ERHIFeatureLevel::SM5
+#else
+	#define UE_IOS_STATIC_FEATURE_LEVEL ERHIFeatureLevel::ES3_1
+#endif
+
+struct FStaticFeatureLevel
+{
+	inline FStaticFeatureLevel(const ERHIFeatureLevel::Type InFeatureLevel)
+	{
+		checkSlow(InFeatureLevel == UE_IOS_STATIC_FEATURE_LEVEL);
+	}
+
+	inline FStaticFeatureLevel(const TEnumAsByte<ERHIFeatureLevel::Type> InFeatureLevel)
+	{
+		checkSlow(InFeatureLevel.GetValue() == UE_IOS_STATIC_FEATURE_LEVEL);
+	}
+
+	inline operator ERHIFeatureLevel::Type() const
+	{
+		return UE_IOS_STATIC_FEATURE_LEVEL;
+	}
+
+	inline bool operator == (const ERHIFeatureLevel::Type Other) const
+	{
+		return Other == UE_IOS_STATIC_FEATURE_LEVEL;
+	}
+
+	inline bool operator != (const ERHIFeatureLevel::Type Other) const
+	{
+		return Other != UE_IOS_STATIC_FEATURE_LEVEL;
+	}
+
+	inline bool operator <= (const ERHIFeatureLevel::Type Other) const
+	{
+		return UE_IOS_STATIC_FEATURE_LEVEL <= Other;
+	}
+
+	inline bool operator < (const ERHIFeatureLevel::Type Other) const
+	{
+		return UE_IOS_STATIC_FEATURE_LEVEL < Other;
+	}
+
+	inline bool operator >= (const ERHIFeatureLevel::Type Other) const
+	{
+		return UE_IOS_STATIC_FEATURE_LEVEL >= Other;
+	}
+
+	inline bool operator > (const ERHIFeatureLevel::Type Other) const
+	{
+		return UE_IOS_STATIC_FEATURE_LEVEL > Other;
+	}
+};
diff --git a/Engine/Source/Runtime/RHI/Public/IOS/IOSStaticShaderPlatform.inl b/Engine/Source/Runtime/RHI/Public/IOS/IOSStaticShaderPlatform.inl
new file mode 100644
--- /dev/null
+++ b/Engine/Source/Runtime/RHI/Public/IOS/IOSStaticShaderPlatform.inl
@@ -0,0 +1,38 @@
+// Copyright Epic Games, Inc. All Rights Reserved.
+
+#pragma once
+
+#include "Misc/AssertionMacros.h"
+
+#if defined(UE_IOS_SHADER_PLATFORM_METAL)
+	#define UE_IOS_STATIC_SHADER_PLATFORM SP_METAL
+#elif defined(UE_IOS_SHADER_PLATFORM_METAL_SIM)
+	#define UE_IOS_STATIC_SHADER_PLATFORM SP_METAL_SIM
+#elif defined(UE_IOS_SHADER_PLATFORM_METAL_MRT)
+	#define UE_IOS_STATIC_SHADER_PLATFORM SP_METAL_MRT
+#else
+	#error "Unknown IOS static shader platform"
+#endif
+
+struct FStaticShaderPlatform
+{
+	inline FStaticShaderPlatform(const EShaderPlatform InPlatform)
+	{
+		checkSlow(UE_IOS_STATIC_SHADER_PLATFORM == InPlatform);
+	}
+
+	inline operator EShaderPlatform() const
+	{
+		return UE_IOS_STATIC_SHADER_PLATFORM;
+	}
+
+	inline bool operator == (const EShaderPlatform Other) const
+	{
+		return Other == UE_IOS_STATIC_SHADER_PLATFORM;
+	}
+	
+	inline bool operator != (const EShaderPlatform Other) const
+	{
+		return Other != UE_IOS_STATIC_SHADER_PLATFORM;
+	}
+};
