Fix broken RHI breadcrumb scope nesting in Insights
 - RHI wasn't handling an edge case with parallel RHICmdLists that have a Sentinel parent, and contain an ActivatePipeline command using all pipes.
 - The "GPUBreadcrumbState[...].Latest" and "CPUBreadcrumbState.Current" may be == Sentinel for two reasons: 1) The RHICmdList contained no breadcrumb operations, or 2) the RHICmdList popped all the breadcrumbs it pushed, but the root node is unknown since this is a parallel cmdlist that is attached to the root tree via QueueAsyncCommandListSubmit.
 - In the case of 2), the Sentinel value should be replaced with the CPU/GPUFirst value, i.e. the breadcrumb that was active before the cmdlist was dispatched, but we were always treating this as case 1). That was breaking the structure of the tree and led to some breadcrumbs being pushed without ever being popped in Insights.

Also rename "SwitchPipeline" -> "ActivatePipeline" in some comments.

diff --git a/Engine/Source/Runtime/RHI/Private/RHICommandList.cpp b/Engine/Source/Runtime/RHI/Private/RHICommandList.cpp
--- a/Engine/Source/Runtime/RHI/Private/RHICommandList.cpp
+++ b/Engine/Source/Runtime/RHI/Private/RHICommandList.cpp
@@ -225,7 +225,7 @@ void FRHICommandListBase::ActivatePipelines(ERHIPipeline Pipelines)
 
 		for (ERHIPipeline Pipeline : MakeFlagsRange(ActivePipelines))
 		{
-			GPUBreadcrumbState[Pipeline].Latest = FRHIBreadcrumbNode::Sentinel;
+			GPUBreadcrumbState[Pipeline].Latest.Reset();
 		}
 
 		if (IsTopOfPipe())
@@ -732,7 +732,7 @@ void FRHICommandListExecutor::FSubmitState::Dispatch(FRHICommandListBase* CmdLis
 				}
 			}
 
-			// Walk the SwitchPipeline commands, resolve unknown targets, and update per-pipe pointers.
+			// Walk the ActivatePipeline commands, resolve unknown targets, and update per-pipe pointers.
 			for (FRHICommandListBase::FActivatePipelineCommand* Command = CmdList->ActivatePipelineCommands.First; Command; Command = Command->Next)
 			{
 				if (Command->Target == FRHIBreadcrumbNode::Sentinel)
@@ -752,10 +752,18 @@ void FRHICommandListExecutor::FSubmitState::Dispatch(FRHICommandListBase* CmdLis
 
 			for (ERHIPipeline Pipeline : MakeFlagsRange(ERHIPipeline::All))
 			{
-				if (CmdList->GPUBreadcrumbState[Pipeline].Latest != FRHIBreadcrumbNode::Sentinel)
+				if (CmdList->GPUBreadcrumbState[Pipeline].Latest.IsSet())
 				{
-					// A Begin/End happened on this pipeline after the last SwitchPipeline command.
-					GRHICommandList.Breadcrumbs.GPU[Pipeline].Current = CmdList->GPUBreadcrumbState[Pipeline].Latest;
+					// A Begin/End happened on this pipeline after the last ActivatePipeline command.
+					FRHIBreadcrumbNode* Node = CmdList->GPUBreadcrumbState[Pipeline].Latest.GetValue();
+					if (Node != FRHIBreadcrumbNode::Sentinel)
+					{
+						GRHICommandList.Breadcrumbs.GPU[Pipeline].Current = Node;
+					}
+					else
+					{
+						GRHICommandList.Breadcrumbs.GPU[Pipeline].Current = GPUFirst[Pipeline];
+					}
 				}
 
 				// Rewind GPU state
@@ -766,6 +774,10 @@ void FRHICommandListExecutor::FSubmitState::Dispatch(FRHICommandListBase* CmdLis
 			{
 				GRHICommandList.Breadcrumbs.CPU.Current = CmdList->CPUBreadcrumbState.Current;
 			}
+			else
+			{
+				GRHICommandList.Breadcrumbs.CPU.Current = CPUFirst;
+			}
 
 			// Rewind CPU state
 			CmdList->CPUBreadcrumbState.Current = CPUFirst;
@@ -839,7 +851,7 @@ void FRHICommandListExecutor::FTranslateState::Translate(FRHICommandListBase* Cm
 #endif
 
 	// Replay the recorded commands. The Contexts array accumulates any used
-	// contexts depending on the SwitchPipeline commands that were recorded.
+	// contexts depending on the ActivatePipeline commands that were recorded.
 	CmdList->Execute();
 
 #if WITH_RHI_BREADCRUMBS
diff --git a/Engine/Source/Runtime/RHI/Public/RHICommandList.h b/Engine/Source/Runtime/RHI/Public/RHICommandList.h
--- a/Engine/Source/Runtime/RHI/Public/RHICommandList.h
+++ b/Engine/Source/Runtime/RHI/Public/RHICommandList.h
@@ -1104,11 +1104,11 @@ protected:
 	bool bUsesShaderBundles      = false;
 
 	// The currently selected pipelines that RHI commands are directed to, during command list recording.
-	// This is also adjusted during command list execution based on recorded use of SwitchPipeline().
+	// This is also adjusted during command list execution based on recorded use of ActivatePipeline().
 	ERHIPipeline ActivePipelines = ERHIPipeline::None;
 
 #if DO_CHECK
-	// Used to check for valid pipelines passed to SwitchPipeline().
+	// Used to check for valid pipelines passed to ActivatePipeline().
 	ERHIPipeline AllowedPipelines = ERHIPipeline::All;
 #endif
 
@@ -1132,7 +1132,7 @@ protected:
 	struct FBreadcrumbState
 	{
 		FRHIBreadcrumbNode* Current = FRHIBreadcrumbNode::Sentinel;
-		FRHIBreadcrumbNode* Latest = FRHIBreadcrumbNode::Sentinel;
+		TOptional<FRHIBreadcrumbNode*> Latest {};
 		FRHIBreadcrumbNode* Prev = nullptr;
 		FRHIBreadcrumbRange Range {};
 	};
diff --git a/Engine/Source/Runtime/RHI/Public/RHIContext.h b/Engine/Source/Runtime/RHI/Public/RHIContext.h
--- a/Engine/Source/Runtime/RHI/Public/RHIContext.h
+++ b/Engine/Source/Runtime/RHI/Public/RHIContext.h
@@ -510,7 +510,7 @@ public:
 		return *ExecutingCmdList;
 	}
 
-	// Used within FRHICommandListBase::SwitchPipeline to setup a context for command execution.
+	// Used within FRHICommandListBase::ActivatePipeline to setup a context for command execution.
 	virtual void SetExecutingCommandList(FRHICommandListBase* InCmdList)
 	{
 		ExecutingCmdList = InCmdList;
