- Add bSupportsUnrestrictedHalfFloatBuffers to DDPI
- Disable Niagara compute compilation if the shader platform does not support half floats

diff --git a/Engine/Config/Android/DataDrivenPlatformInfo.ini b/Engine/Config/Android/DataDrivenPlatformInfo.ini
--- a/Engine/Config/Android/DataDrivenPlatformInfo.ini
+++ b/Engine/Config/Android/DataDrivenPlatformInfo.ini
@@ -62,6 +62,7 @@ bIsAndroidOpenGLES=true
 bTargetsTiledGPU=true
 bSupportsMobileMultiView=true
 bNeedsOfflineCompiler=true
+bSupportsUnrestrictedHalfFloatBuffers=false
 bSupportsManualVertexFetch=false
 bSupportsClipDistance = false
 bSupportsShaderPipelines = false
diff --git a/Engine/Plugins/FX/Niagara/Source/Niagara/Private/NiagaraScript.cpp b/Engine/Plugins/FX/Niagara/Source/Niagara/Private/NiagaraScript.cpp
--- a/Engine/Plugins/FX/Niagara/Source/Niagara/Private/NiagaraScript.cpp
+++ b/Engine/Plugins/FX/Niagara/Source/Niagara/Private/NiagaraScript.cpp
@@ -3635,6 +3635,21 @@ bool UNiagaraScript::IsEditorOnly() const
 bool UNiagaraScript::ShouldCompile(EShaderPlatform Platform) const
 {
 #if WITH_EDITOR
+
+	// We can not guarantee half float support on all platforms on the GPU, therefore we skip compilation in these cases
+	TOptional<ENiagaraSimTarget> ActualSimTarget = GetSimTarget();
+	if (ActualSimTarget.IsSet() && ActualSimTarget.GetValue() == ENiagaraSimTarget::GPUComputeSim)
+	{
+		UNiagaraSystem* OuterSystem = GetTypedOuter<UNiagaraSystem>();
+		if (OuterSystem && OuterSystem->ShouldCompressAttributes())
+		{
+			if ( !FDataDrivenShaderPlatformInfo::GetSupportsUnrestrictedHalfFloatBuffers(Platform) )
+			{
+				return false;
+			}
+		}
+	}
+
 	// Add all data interfaces
 	TSet<UClass*> DIUniqueClasses;
 	for (const FNiagaraScriptDataInterfaceInfo& DataInterfaceInfo : CachedDefaultDataInterfaces)
diff --git a/Engine/Source/Runtime/RHI/Private/DataDrivenShaderPlatformInfo.cpp b/Engine/Source/Runtime/RHI/Private/DataDrivenShaderPlatformInfo.cpp
--- a/Engine/Source/Runtime/RHI/Private/DataDrivenShaderPlatformInfo.cpp
+++ b/Engine/Source/Runtime/RHI/Private/DataDrivenShaderPlatformInfo.cpp
@@ -139,6 +139,7 @@ void FGenericDataDrivenShaderPlatformInfo::SetDefaultValues()
 {
 	MaxFeatureLevel = ERHIFeatureLevel::Num;
 	bSupportsMSAA = true;
+	bSupportsUnrestrictedHalfFloatBuffers = true;
 	bSupportsDOFHybridScattering = true;
 	bSupportsHZBOcclusion = true;
 	bSupportsWaterIndirectDraw = true;
@@ -228,6 +229,7 @@ void FGenericDataDrivenShaderPlatformInfo::ParseDataDrivenShaderInfo(const FConf
 	GET_SECTION_BOOL_HELPER(bSupportsRayTracingIndirectInstanceData);
 	GET_SECTION_BOOL_HELPER(bSupportsPathTracing);
 	GET_SECTION_BOOL_HELPER(bSupportsHighEndRayTracingEffects);
+	GET_SECTION_BOOL_HELPER(bSupportsUnrestrictedHalfFloatBuffers)
 	GET_SECTION_BOOL_HELPER(bSupportsByteBufferComputeShaders);
 	GET_SECTION_BOOL_HELPER(bSupportsGPUScene);
 	GET_SECTION_BOOL_HELPER(bSupportsPrimitiveShaders);
diff --git a/Engine/Source/Runtime/RHI/Public/DataDrivenShaderPlatformInfo.h b/Engine/Source/Runtime/RHI/Public/DataDrivenShaderPlatformInfo.h
--- a/Engine/Source/Runtime/RHI/Public/DataDrivenShaderPlatformInfo.h
+++ b/Engine/Source/Runtime/RHI/Public/DataDrivenShaderPlatformInfo.h
@@ -53,6 +53,7 @@ class FGenericDataDrivenShaderPlatformInfo
 	uint32 bSupportsHighEndRayTracingEffects : 1; // Whether fully-featured RT effects can be used on the platform (with translucent shadow, etc.)
 	uint32 bSupportsPathTracing : 1; // Whether real-time path tracer is supported on this platform (avoids compiling unnecessary shaders)
 	uint32 bSupportsGPUScene : 1;
+	uint32 bSupportsUnrestrictedHalfFloatBuffers : 1;
 	uint32 bSupportsByteBufferComputeShaders : 1;
 	uint32 bSupportsPrimitiveShaders : 1;
 	uint32 bSupportsUInt64ImageAtomics : 1;
@@ -418,6 +419,12 @@ public:
 		return Infos[Platform].bNeedsOfflineCompiler;
 	}
 
+	static FORCEINLINE_DEBUGGABLE const bool GetSupportsUnrestrictedHalfFloatBuffers(const FStaticShaderPlatform Platform)
+	{
+		check(IsValid(Platform));
+		return Infos[Platform].bSupportsUnrestrictedHalfFloatBuffers;
+	}
+
 	static FORCEINLINE_DEBUGGABLE const bool GetSupportsByteBufferComputeShaders(const FStaticShaderPlatform Platform)
 	{
 		check(IsValid(Platform));
