Only export RHI breadcrumbs to crash payload when they have data.

Otherwise, empty RHI breadcrumb data would override DRED breadcrumbs. In the future, we'll allow exporting multiple breadcrumb trees from different systems.

#tests tested locally

diff --git a/Engine/Source/Runtime/RHI/Private/RHIBreadcrumbs.cpp b/Engine/Source/Runtime/RHI/Private/RHIBreadcrumbs.cpp
--- a/Engine/Source/Runtime/RHI/Private/RHIBreadcrumbs.cpp
+++ b/Engine/Source/Runtime/RHI/Private/RHIBreadcrumbs.cpp
@@ -315,7 +315,10 @@ RHI_API void FRHIBreadcrumbState::DumpActiveBreadcrumbs(TMap<FQueueID, TArray<FR
 		}
 	}
 
-	FGenericCrashContext::SetGPUBreadcrumbs(MoveTemp(CrashData));
+	if (CrashData.Queues.Num())
+	{
+		FGenericCrashContext::SetGPUBreadcrumbs(MoveTemp(CrashData));
+	}
 	UE_LOG(LogRHI, Error, TEXT("Active GPU breadcrumbs:%s\r\n"), *Tree);
 }
 
