Remove "nested" RHI resource deletion to fix race condition between the render and RHI threads.
 - In cases where both the render thread and a uniform buffer own a refcount on an SRV resource, the render and RHI threads can race during resource deletion. This resulted in a deleted SRV still being used by future RHI commands and crashing inside the platform RHI.
 - If the UB is deleted on the RHI thread, before the RT releases its refcount, everything is fine. The SRV will get enqueued for deletion after the commands the RT is currently recording.
 - However, if the RT releases its reference first, the UB will assume the SRV was nested since its refcount is now 1, and delete the SRV on the RHI thread, before the future RHI commands have been executed.

The fix is simply to remove this "nested" resource deletion handling. The "CurrentlyDeleting" thread local now has no runtime effect, aside from check()s, so is now wrapped in #if DO_CHECK.

diff --git a/Engine/Source/Runtime/RHI/Private/RHIResources.cpp b/Engine/Source/Runtime/RHI/Private/RHIResources.cpp
--- a/Engine/Source/Runtime/RHI/Private/RHIResources.cpp
+++ b/Engine/Source/Runtime/RHI/Private/RHIResources.cpp
@@ -13,7 +13,10 @@
 UE::TConsumeAllMpmcQueue<FRHIResource*> PendingDeletes;
 UE::TConsumeAllMpmcQueue<FRHIResource*> PendingDeletesWithLifetimeExtension;
 
+#if DO_CHECK
+// Used to check that we only enter the FRHIResource destructor via the FRHIResource::DeleteResources() function.
 thread_local FRHIResource const* FRHIResource::CurrentlyDeleting = nullptr;
+#endif
 
 FRHIResource::FRHIResource(ERHIResourceType InResourceType)
 	: ResourceType(InResourceType)
@@ -32,7 +35,9 @@ FRHIResource::~FRHIResource()
 {
 	check(IsEngineExitRequested() || CurrentlyDeleting == this);
 	check(AtomicFlags.GetNumRefs(std::memory_order_relaxed) == 0); // this should not have any outstanding refs
+#if DO_CHECK
 	CurrentlyDeleting = nullptr;
+#endif
 
 #if RHI_ENABLE_RESOURCE_INFO
 	EndTrackingResource(this);
@@ -41,33 +46,15 @@ FRHIResource::~FRHIResource()
 
 void FRHIResource::MarkForDelete() const
 {
-	if (CurrentlyDeleting)
+	if (!AtomicFlags.MarkForDelete(std::memory_order_release))
 	{
-		if (!AtomicFlags.IsMarkedForDelete(std::memory_order_release))
+		if (bAllowExtendLifetime)
 		{
-			// We hit refcount zero while destructing another RHI resource, and we're not already marked for delete.
-			// The current resource is nested inside the previous one, so can be destructed immediately.
-			FRHIResource const* Previous = CurrentlyDeleting;
-
-			CurrentlyDeleting = this;
-			delete this;
-
-			check(CurrentlyDeleting == nullptr);
-			CurrentlyDeleting = Previous;
+			PendingDeletesWithLifetimeExtension.ProduceItem(const_cast<FRHIResource*>(this));
 		}
-	}
-	else
-	{
-		if (!AtomicFlags.MarkForDelete(std::memory_order_release))
+		else
 		{
-			if (bAllowExtendLifetime)
-			{
-				PendingDeletesWithLifetimeExtension.ProduceItem(const_cast<FRHIResource*>(this));
-			}
-			else
-			{
-				PendingDeletes.ProduceItem(const_cast<FRHIResource*>(this));
-			}
+			PendingDeletes.ProduceItem(const_cast<FRHIResource*>(this));
 		}
 	}
 }
@@ -78,7 +65,9 @@ void FRHIResource::DeleteResources(TArray<FRHIResource*> const& Resources)
 	{
 		if (Resource->AtomicFlags.Deleting())
 		{
+#if DO_CHECK
 			CurrentlyDeleting = Resource;
+#endif
 			delete Resource;
 
 			check(CurrentlyDeleting == nullptr);
diff --git a/Engine/Source/Runtime/RHI/Public/RHIResources.h b/Engine/Source/Runtime/RHI/Public/RHIResources.h
--- a/Engine/Source/Runtime/RHI/Public/RHIResources.h
+++ b/Engine/Source/Runtime/RHI/Public/RHIResources.h
@@ -206,7 +206,9 @@ private:
 	FName OwnerName;
 #endif
 
+#if DO_CHECK
 	static thread_local FRHIResource const* CurrentlyDeleting;
+#endif
 
 	friend FRHICommandListImmediate;
 };
