Bindless Command Line Utilities
- "-Bindless" on the command line should imply "-sm6"
- If no RHI was specified on the command line, "-Bindless" should also imply "-dx12"

diff --git a/Engine/Source/Runtime/RHI/Private/Windows/WindowsDynamicRHI.cpp b/Engine/Source/Runtime/RHI/Private/Windows/WindowsDynamicRHI.cpp
--- a/Engine/Source/Runtime/RHI/Private/Windows/WindowsDynamicRHI.cpp
+++ b/Engine/Source/Runtime/RHI/Private/Windows/WindowsDynamicRHI.cpp
@@ -726,6 +726,13 @@ static TOptional<EWindowsRHI> ChooseForcedRHI(TOptional<ERHIFeatureLevel::Type>
 		Sum++;
 	}
 
+#if WITH_EDITOR
+	if (!ForcedRHI && GIsEditor && FParse::Param(FCommandLine::Get(), TEXT("Bindless")))
+	{
+		ForcedRHI = EWindowsRHI::D3D12;
+	}
+#endif
+
 	if (Sum > 1)
 	{
 		FMessageDialog::Open(EAppMsgType::Ok, LOCTEXT("RHIOptionsError", "-d3d12/dx12, -d3d11/dx11, -vulkan, and -opengl are mutually exclusive options, but more than one was specified on the command-line."));
@@ -771,6 +778,13 @@ static TOptional<ERHIFeatureLevel::Type> GetForcedFeatureLevel()
 		ForcedFeatureLevel = ERHIFeatureLevel::SM6;
 	}
 
+#if WITH_EDITOR
+	if (GIsEditor && FParse::Param(FCommandLine::Get(), TEXT("Bindless")))
+	{
+		ForcedFeatureLevel = ERHIFeatureLevel::SM6;
+	}
+#endif
+
 	return ForcedFeatureLevel;
 }
 
