Fix circular reference in RHI breadcrumb allocators causing a memory leak
 - RDG uses its own breadcrumb allocator for RDG scopes. When the graph executes, this allocator is attached to the parent one owned by the immediate command list.
 - Later, when a pass makes use of the immediate command list, a new breadcrumb is inserted on that command list's allocator, which gets reattached to the RDG allocator, forming a circular reference.
 - The circular reference causes a memory leak, and a crash on shutdown when the underlying mempage allocators are destroyed.
 - The reference only occurs in bypass mode, and when platform RHIs skip the DispatchToRHIThread inside FRDGBuilder::BeginFlushResourcesRHI().
 - Fix is to swap out the immediate command list's allocator for the RDG one when the graph executes. This is now done in FRHICommandListBase::AttachBreadcrumbSubTree(...).

Also added RHI validation to check for circular references in the RHI breadcrumb allocators.

diff --git a/Engine/Source/Runtime/RHI/Private/RHICommandList.cpp b/Engine/Source/Runtime/RHI/Private/RHICommandList.cpp
--- a/Engine/Source/Runtime/RHI/Private/RHICommandList.cpp
+++ b/Engine/Source/Runtime/RHI/Private/RHICommandList.cpp
@@ -201,6 +201,30 @@ TOptional<FRHIDrawStatsCategory const*> FRHICommandListBase::SetDrawStatsCategor
 }
 #endif
 
+#if WITH_RHI_BREADCRUMBS
+void FRHICommandListBase::AttachBreadcrumbSubTree(FRHIBreadcrumbAllocator& Allocator, FRHIBreadcrumbList& Nodes)
+{
+	for (FRHIBreadcrumbNode* Node : Nodes.IterateAndUnlink())
+	{
+		checkf(Node->Allocator == &Allocator, TEXT("All the nodes in a subtree must come from the same breadcrumb allocator."));
+		if (Node->GetParent() == FRHIBreadcrumbNode::Sentinel)
+		{
+			Node->SetParent(GetCurrentBreadcrumbRef());
+		}
+	}
+
+	// Switch the current breadcrumb allocator out for the subtree one.
+	if (BreadcrumbAllocator.Get() != &Allocator)
+	{
+		if (BreadcrumbAllocator)
+		{
+			BreadcrumbAllocatorRefs.AddUnique(BreadcrumbAllocator.Get());
+		}
+		BreadcrumbAllocator = Allocator.AsShared();
+	}
+}
+#endif
+
 void FRHICommandListBase::ActivatePipelines(ERHIPipeline Pipelines)
 {
 	checkf(IsTopOfPipe() || Bypass(), TEXT("Cannot be called from the bottom of pipe."));
diff --git a/Engine/Source/Runtime/RHI/Private/RHIValidation.cpp b/Engine/Source/Runtime/RHI/Private/RHIValidation.cpp
--- a/Engine/Source/Runtime/RHI/Private/RHIValidation.cpp
+++ b/Engine/Source/Runtime/RHI/Private/RHIValidation.cpp
@@ -1891,7 +1891,6 @@ namespace RHIValidation
 		case EOpType::BeginBreadcrumbGPU:
 			{
 				FRHIBreadcrumbNode* Node = Data_Breadcrumb.Breadcrumb;
-				LogNode(Node, true, Queue.Pipeline);
 
 				check(Node && Node != FRHIBreadcrumbNode::Sentinel);
 				check(Node->GetParent() != FRHIBreadcrumbNode::Sentinel);
@@ -1899,6 +1898,8 @@ namespace RHIValidation
 				check(GRHICommandList.Bypass() || IsInRange(Queue.Breadcrumbs.Range, Node, Queue.Pipeline));
 				check(EnumHasAllFlags(static_cast<ERHIPipeline>(Node->BeginPipes.load()), Queue.Pipeline));
 
+				LogNode(Node, true, Queue.Pipeline);
+
 				Queue.Breadcrumbs.Current = Node;
 			}
 			break;
@@ -1906,7 +1907,6 @@ namespace RHIValidation
 		case EOpType::EndBreadcrumbGPU:
 			{
 				FRHIBreadcrumbNode* Node = Data_Breadcrumb.Breadcrumb;
-				LogNode(Node, false, Queue.Pipeline);
 
 				check(Node && Node != FRHIBreadcrumbNode::Sentinel);
 				check(Node->GetParent() != FRHIBreadcrumbNode::Sentinel);
@@ -1914,21 +1914,46 @@ namespace RHIValidation
 				check(GRHICommandList.Bypass() || IsInRange(Queue.Breadcrumbs.Range, Node, Queue.Pipeline));
 				check(EnumHasAllFlags(static_cast<ERHIPipeline>(Node->EndPipes.load()), Queue.Pipeline));
 
+				LogNode(Node, false, Queue.Pipeline);
+
 				Queue.Breadcrumbs.Current = Node->GetParent();
 			}
 			break;
 
 		case EOpType::SetBreadcrumbRange:
-			Queue.Breadcrumbs.Range = Data_BreadcrumbRange.Range;
-			check(!Queue.Breadcrumbs.Range.First == !Queue.Breadcrumbs.Range.Last);
-
-			for (FRHIBreadcrumbNode* Node : Queue.Breadcrumbs.Range.Enumerate(Queue.Pipeline))
 			{
-				// Check current node and all parents are valid
-				for (FRHIBreadcrumbNode* Other = Node; Other; Other = Other->GetParent())
+				Queue.Breadcrumbs.Range = Data_BreadcrumbRange.Range;
+				check(!Queue.Breadcrumbs.Range.First == !Queue.Breadcrumbs.Range.Last);
+
+				TSet<FRHIBreadcrumbAllocator*> AllAllocators;
+				for (FRHIBreadcrumbNode* Node : Queue.Breadcrumbs.Range.Enumerate(Queue.Pipeline))
 				{
-					check(Other != FRHIBreadcrumbNode::Sentinel);
-					check(Other->GetParent() != FRHIBreadcrumbNode::Sentinel);
+					AllAllocators.Add(Node->Allocator);
+
+					// Check current node and all parents are valid
+					for (FRHIBreadcrumbNode* Other = Node; Other; Other = Other->GetParent())
+					{
+						check(Other != FRHIBreadcrumbNode::Sentinel);
+						check(Other->GetParent() != FRHIBreadcrumbNode::Sentinel);
+					}
+				}
+
+				// Check for circular references in the allocator parent pointers
+				for (FRHIBreadcrumbAllocator* Allocator : AllAllocators)
+				{
+					auto Recurse = [](FRHIBreadcrumbAllocator* Current, auto& Recurse) -> void
+					{
+						checkf(!Current->bVisited, TEXT("Circular reference detected in breadcrumb allocators."));
+						Current->bVisited = true;
+						
+						for (auto const& Parent : Current->GetParents())
+						{
+							Recurse(&Parent.Get(), Recurse);
+						}
+
+						Current->bVisited = false;
+					};
+					Recurse(Allocator, Recurse);
 				}
 			}
 			break;
diff --git a/Engine/Source/Runtime/RHI/Public/RHIBreadcrumbs.h b/Engine/Source/Runtime/RHI/Public/RHIBreadcrumbs.h
--- a/Engine/Source/Runtime/RHI/Public/RHIBreadcrumbs.h
+++ b/Engine/Source/Runtime/RHI/Public/RHIBreadcrumbs.h
@@ -154,6 +154,8 @@
 		}
 
 	public:
+		FRHIBreadcrumbAllocatorArray const& GetParents() const { return Parents; }
+
 		template<size_t N, typename... TArgs>
 		inline FRHIBreadcrumbNode* AllocBreadcrumb(TCHAR const(&FormatString)[N], TArgs&&... Args);
 
@@ -161,6 +163,11 @@
 		{
 			return Inner.Alloc(Size, Align);
 		}
+
+#if ENABLE_RHI_VALIDATION
+		// Used by RHI validation for circular reference detection.
+		bool bVisited = false;
+#endif
 	};
 
 	inline void FRHIBreadcrumbAllocatorArray::AddUnique(FRHIBreadcrumbAllocator* Allocator)
diff --git a/Engine/Source/Runtime/RHI/Public/RHICommandList.h b/Engine/Source/Runtime/RHI/Public/RHICommandList.h
--- a/Engine/Source/Runtime/RHI/Public/RHICommandList.h
+++ b/Engine/Source/Runtime/RHI/Public/RHICommandList.h
@@ -1226,6 +1226,8 @@ public:
 	{
 		return PersistentState.LocalBreadcrumb;
 	}
+
+	RHI_API void AttachBreadcrumbSubTree(FRHIBreadcrumbAllocator& Allocator, FRHIBreadcrumbList& Nodes);
 #endif
 
 #if HAS_GPU_STATS
diff --git a/Engine/Source/Runtime/RenderCore/Private/RenderGraphBuilder.cpp b/Engine/Source/Runtime/RenderCore/Private/RenderGraphBuilder.cpp
--- a/Engine/Source/Runtime/RenderCore/Private/RenderGraphBuilder.cpp
+++ b/Engine/Source/Runtime/RenderCore/Private/RenderGraphBuilder.cpp
@@ -1542,11 +1542,7 @@ void FRDGBuilder::Compile()
 #if WITH_RHI_BREADCRUMBS
 	// Attach the RDG breadcrumb nodes to the current top-of-stack RHI immediate breadcrumb,
 	// Also unlink them from each other.
-	for (FRHIBreadcrumbNode* Node : LocalBreadcrumbList.IterateAndUnlink())
-	{
-		if (Node->GetParent() == FRHIBreadcrumbNode::Sentinel)
-			Node->SetParent(RHICmdList.GetCurrentBreadcrumbRef());
-	}
+	RHICmdList.AttachBreadcrumbSubTree(GetBreadcrumbAllocator(), LocalBreadcrumbList);
 #endif
 }
 
