Fixed RHI validation layer problems introduced by the addition of separate pixel and non-pixel SRV graphics states:
* assert in DecayResourceAccess due to SRVGraphics being two bits now
* incorrect expected state for UAVs when binding resources

#rnx

diff --git a/Engine/Source/Runtime/RHI/Private/RHIValidation.cpp b/Engine/Source/Runtime/RHI/Private/RHIValidation.cpp
--- a/Engine/Source/Runtime/RHI/Private/RHIValidation.cpp
+++ b/Engine/Source/Runtime/RHI/Private/RHIValidation.cpp
@@ -996,7 +996,7 @@ namespace RHIValidation
 	ERHIAccess DecayResourceAccess(ERHIAccess AccessMask, ERHIAccess RequiredAccess, bool bAllowUAVOverlap)
 	{
 		using T = __underlying_type(ERHIAccess);
-		checkf((T(RequiredAccess) & (T(RequiredAccess) - 1)) == 0, TEXT("Only one required access bit may be set at once."));
+		checkf(RequiredAccess == ERHIAccess::SRVGraphics || (T(RequiredAccess) & (T(RequiredAccess) - 1)) == 0, TEXT("Only one required access bit may be set at once."));
 		
 		if (EnumHasAnyFlags(RequiredAccess, ERHIAccess::UAVMask | ERHIAccess::BVHWrite))
 		{
diff --git a/Engine/Source/Runtime/RHICore/Public/RHICoreShader.h b/Engine/Source/Runtime/RHICore/Public/RHICoreShader.h
--- a/Engine/Source/Runtime/RHICore/Public/RHICoreShader.h
+++ b/Engine/Source/Runtime/RHICore/Public/RHICoreShader.h
@@ -272,9 +272,7 @@ void SetResourcesFromTables(TBinder&& Binder, FRHIShader const& Shader, TBitMask
 				{
 					ERHIAccess Access = IsComputeShaderFrequency(Shader.GetFrequency())
 						? ERHIAccess::UAVCompute
-						: Shader.GetFrequency() == SF_Pixel
-							? ERHIAccess::SRVGraphicsPixel
-							: ERHIAccess::SRVGraphicsNonPixel;
+						: ERHIAccess::UAVGraphics;
 
 					Tracker->AssertUAV(UAV, Access, Index);
 				}
