Fixed missing enum in RHI validation
#rnx

diff --git a/Engine/Source/Runtime/RHI/Private/RHIStrings.cpp b/Engine/Source/Runtime/RHI/Private/RHIStrings.cpp
--- a/Engine/Source/Runtime/RHI/Private/RHIStrings.cpp
+++ b/Engine/Source/Runtime/RHI/Private/RHIStrings.cpp
@@ -475,7 +475,8 @@ FString GetRHIAccessName(ERHIAccess Access)
 		case ERHIAccess::IndirectArgs:        return TEXT("IndirectArgs");
 		case ERHIAccess::VertexOrIndexBuffer: return TEXT("VertexOrIndexBuffer");
 		case ERHIAccess::SRVCompute:          return TEXT("SRVCompute");
-		case ERHIAccess::SRVGraphics:         return TEXT("SRVGraphics");
+		case ERHIAccess::SRVGraphicsPixel:    return TEXT("SRVGraphicsPixel");
+		case ERHIAccess::SRVGraphicsNonPixel: return TEXT("SRVGraphicsNonPixel");
 		case ERHIAccess::CopySrc:             return TEXT("CopySrc");
 		case ERHIAccess::ResolveSrc:          return TEXT("ResolveSrc");
 		case ERHIAccess::DSVRead:             return TEXT("DSVRead");
@@ -485,10 +486,10 @@ FString GetRHIAccessName(ERHIAccess Access)
 		case ERHIAccess::CopyDest:            return TEXT("CopyDest");
 		case ERHIAccess::ResolveDst:          return TEXT("ResolveDst");
 		case ERHIAccess::DSVWrite:            return TEXT("DSVWrite");
-		case ERHIAccess::ShadingRateSource:	  return TEXT("ShadingRateSource");
 		case ERHIAccess::BVHRead:             return TEXT("BVHRead");
 		case ERHIAccess::BVHWrite:            return TEXT("BVHWrite");
 		case ERHIAccess::Discard:             return TEXT("Discard");
+		case ERHIAccess::ShadingRateSource:	  return TEXT("ShadingRateSource");
 		}
 	});
 }
