Adding "-BindlessRT" command line option for the Editor to force enable Bindless on RayTracing shaders. This will also override the default RHI to something that enables Bindless.

diff --git a/Engine/Source/Runtime/RHI/Private/RHI.cpp b/Engine/Source/Runtime/RHI/Private/RHI.cpp
--- a/Engine/Source/Runtime/RHI/Private/RHI.cpp
+++ b/Engine/Source/Runtime/RHI/Private/RHI.cpp
@@ -1210,11 +1210,17 @@ ERHIBindlessConfiguration RHIParseBindlessConfiguration(EShaderPlatform Platform
 
 #if WITH_EDITOR
 	// We have to check the -bindless command line option here to make sure the shaders are compiled with bindless enabled too.
-	static const bool bCommandLine = FParse::Param(FCommandLine::Get(), TEXT("Bindless"));
-	if (bCommandLine)
+	static const bool bRegularCommandLine = FParse::Param(FCommandLine::Get(), TEXT("Bindless"));
+	static const bool bRTOnlyCommandLine = FParse::Param(FCommandLine::Get(), TEXT("BindlessRT"));
+	if (bRegularCommandLine || bRTOnlyCommandLine)
 	{
 		// Only allow what the platform supports
-		return (BindlessSupport == ERHIBindlessSupport::RayTracingOnly) ? ERHIBindlessConfiguration::RayTracingShaders : ERHIBindlessConfiguration::AllShaders;
+		if (BindlessSupport == ERHIBindlessSupport::RayTracingOnly || bRTOnlyCommandLine)
+		{
+			return ERHIBindlessConfiguration::RayTracingShaders;
+		}
+
+		return ERHIBindlessConfiguration::AllShaders;
 	}
 #endif
 
diff --git a/Engine/Source/Runtime/RHI/Private/Windows/WindowsDynamicRHI.cpp b/Engine/Source/Runtime/RHI/Private/Windows/WindowsDynamicRHI.cpp
--- a/Engine/Source/Runtime/RHI/Private/Windows/WindowsDynamicRHI.cpp
+++ b/Engine/Source/Runtime/RHI/Private/Windows/WindowsDynamicRHI.cpp
@@ -727,7 +727,7 @@ static TOptional<EWindowsRHI> ChooseForcedRHI(TOptional<ERHIFeatureLevel::Type>
 	}
 
 #if WITH_EDITOR
-	if (!ForcedRHI && GIsEditor && FParse::Param(FCommandLine::Get(), TEXT("Bindless")))
+	if (!ForcedRHI && GIsEditor && (FParse::Param(FCommandLine::Get(), TEXT("Bindless")) || FParse::Param(FCommandLine::Get(), TEXT("BindlessRT"))))
 	{
 		ForcedRHI = EWindowsRHI::D3D12;
 	}
@@ -779,7 +779,7 @@ static TOptional<ERHIFeatureLevel::Type> GetForcedFeatureLevel()
 	}
 
 #if WITH_EDITOR
-	if (GIsEditor && FParse::Param(FCommandLine::Get(), TEXT("Bindless")))
+	if (GIsEditor && (FParse::Param(FCommandLine::Get(), TEXT("Bindless")) || FParse::Param(FCommandLine::Get(), TEXT("BindlessRT"))))
 	{
 		ForcedFeatureLevel = ERHIFeatureLevel::SM6;
 	}
