HWRT: First pass at reducing RHI_RAYTRACING throughout the engine - starting with RHI.

diff --git a/Engine/Source/Runtime/RHI/Private/PipelineStateCache.cpp b/Engine/Source/Runtime/RHI/Private/PipelineStateCache.cpp
--- a/Engine/Source/Runtime/RHI/Private/PipelineStateCache.cpp
+++ b/Engine/Source/Runtime/RHI/Private/PipelineStateCache.cpp
@@ -777,7 +777,6 @@ FRHIWorkGraphPipelineState* GetRHIWorkGraphPipelineState(FWorkGraphPipelineState
 	return PipelineState->RHIPipeline;
 }
 
-#if RHI_RAYTRACING
 /* State for ray tracing */
 class FRayTracingPipelineState : public FPipelineState
 {
@@ -867,8 +866,6 @@ FRHIRayTracingPipelineState* GetRHIRayTracingPipelineState(FRayTracingPipelineSt
 	return PipelineState->RHIPipeline;
 }
 
-#endif // RHI_RAYTRACING
-
 int32 FindRayTracingHitGroupIndex(FRayTracingPipelineState* Pipeline, FRHIRayTracingShader* HitGroupShader, bool bRequired)
 {
 #if RHI_RAYTRACING
diff --git a/Engine/Source/Runtime/RHI/Private/RHICommandList.cpp b/Engine/Source/Runtime/RHI/Private/RHICommandList.cpp
--- a/Engine/Source/Runtime/RHI/Private/RHICommandList.cpp
+++ b/Engine/Source/Runtime/RHI/Private/RHICommandList.cpp
@@ -1775,7 +1775,6 @@ void FRHIComputeCommandList::Transition(TArrayView<const FRHITransitionInfo> Inf
 	}
 }
 
-#if RHI_RAYTRACING
 void FRHIComputeCommandList::BuildAccelerationStructure(FRHIRayTracingGeometry* Geometry)
 {
 	FRayTracingGeometryBuildParams Params;
@@ -1805,7 +1804,6 @@ void FRHIComputeCommandList::BuildAccelerationStructures(const TArrayView<const
 
 	BuildAccelerationStructures(Params, ScratchBufferRange);
 }
-#endif
 
 static FLockTracker GLockTracker;
 
diff --git a/Engine/Source/Runtime/RHI/Public/DynamicRHI.h b/Engine/Source/Runtime/RHI/Public/DynamicRHI.h
--- a/Engine/Source/Runtime/RHI/Public/DynamicRHI.h
+++ b/Engine/Source/Runtime/RHI/Public/DynamicRHI.h
@@ -941,8 +941,6 @@ public:
 
 	virtual uint16 RHIGetPlatformTextureMaxSampleCount() { return 8; };
 
-#if RHI_RAYTRACING
-
 	virtual FRayTracingAccelerationStructureSize RHICalcRayTracingSceneSize(uint32 MaxInstances, ERayTracingAccelerationStructureFlags Flags)
 	{
 		checkNoEntry();
@@ -978,7 +976,6 @@ public:
 		checkNoEntry();
 		return nullptr;
 	}
-#endif // RHI_RAYTRACING
 
 	virtual FShaderBundleRHIRef RHICreateShaderBundle(uint32 NumRecords)
 	{
@@ -1138,12 +1135,10 @@ FORCEINLINE FRenderQueryRHIRef RHICreateRenderQuery(ERenderQueryType QueryType)
 	return GDynamicRHI->RHICreateRenderQuery(QueryType);
 }
 
-#if RHI_RAYTRACING
 FORCEINLINE TRefCountPtr<FRHIRayTracingPipelineState> RHICreateRayTracingPipelineState(const FRayTracingPipelineStateInitializer& Initializer)
 {
 	return GDynamicRHI->RHICreateRayTracingPipelineState(Initializer);
 }
-#endif //RHI_RAYTRACING
 
 FORCEINLINE FUniformBufferLayoutRHIRef RHICreateUniformBufferLayout(const FRHIUniformBufferLayoutInitializer& Initializer)
 {
@@ -1449,8 +1444,6 @@ FORCEINLINE bool RHIMatchPrecachePSOInitializers(const FGraphicsPipelineStateIni
 	return GDynamicRHI->RHIMatchPrecachePSOInitializers(LHS, RHS);
 }
 
-#if RHI_RAYTRACING
-
 FORCEINLINE FRayTracingAccelerationStructureSize RHICalcRayTracingSceneSize(uint32 MaxInstances, ERayTracingAccelerationStructureFlags Flags)
 {
 	return GDynamicRHI->RHICalcRayTracingSceneSize(MaxInstances, Flags);
@@ -1473,8 +1466,6 @@ FORCEINLINE FRayTracingShaderRHIRef RHICreateRayTracingShader(TArrayView<const u
 	return GDynamicRHI->RHICreateRayTracingShader(Code, Hash, ShaderFrequency);
 }
 
-#endif // RHI_RAYTRACING
-
 FORCEINLINE FShaderBundleRHIRef RHICreateShaderBundle(uint32 NumRecords)
 {
 	return GDynamicRHI->RHICreateShaderBundle(NumRecords);
diff --git a/Engine/Source/Runtime/RHI/Public/RHICommandList.h b/Engine/Source/Runtime/RHI/Public/RHICommandList.h
--- a/Engine/Source/Runtime/RHI/Public/RHICommandList.h
+++ b/Engine/Source/Runtime/RHI/Public/RHICommandList.h
@@ -237,7 +237,6 @@ private:
 	FUpdateTexture3DData();
 };
 
-#if RHI_RAYTRACING
 struct FRayTracingShaderBindings
 {
 	FRHITexture* Textures[64] = {};
@@ -266,7 +265,6 @@ enum class ERayTracingBindingType : uint8
 	CallableShader,
 	MissShader,
 };
-#endif // RHI_RAYTRACING
 
 struct FLockTracker
 {
@@ -995,7 +993,6 @@ public:
 		return GDynamicRHI->RHICreateResourceCollection(*this, InMembers);
 	}
 
-#if RHI_RAYTRACING
 	FORCEINLINE FRayTracingGeometryRHIRef CreateRayTracingGeometry(const FRayTracingGeometryInitializer& Initializer)
 	{
 		return GDynamicRHI->RHICreateRayTracingGeometry(*this, Initializer);
@@ -1006,7 +1003,6 @@ public:
 	{
 		return GDynamicRHI->RHICalcRayTracingGeometrySize(Initializer);
 	}
-#endif
 
 	FORCEINLINE void ReplaceResources(TArray<FRHIResourceReplaceInfo>&& ReplaceInfos)
 	{
@@ -2271,8 +2267,6 @@ FRHICOMMAND_MACRO(FRHICommandCopyBufferRegion)
 	RHI_API void Execute(FRHICommandListBase& CmdList);
 };
 
-#if RHI_RAYTRACING
-
 FRHICOMMAND_UNNAMED(FRHICommandBindAccelerationStructureMemory)
 {
 	FRHIRayTracingScene* Scene;
@@ -2393,7 +2387,6 @@ FRHICOMMAND_MACRO(FRHICommandSetRayTracingBindings)
 
 	RHI_API void Execute(FRHICommandListBase& CmdList);
 };
-#endif // RHI_RAYTRACING
 
 template<> RHI_API void FRHICommandSetShaderParameters           <FRHIComputeShader>::Execute(FRHICommandListBase& CmdList);
 template<> RHI_API void FRHICommandSetShaderUnbinds              <FRHIComputeShader>::Execute(FRHICommandListBase& CmdList);
@@ -3088,7 +3081,6 @@ public:
 #endif // WITH_MGPU
 	}
 
-#if RHI_RAYTRACING
 	RHI_API void BuildAccelerationStructure(FRHIRayTracingGeometry* Geometry);
 	RHI_API void BuildAccelerationStructures(TConstArrayView<FRayTracingGeometryBuildParams> Params);
 
@@ -3143,7 +3135,6 @@ public:
 			RHIThreadFence(true);
 		}
 	}
-#endif
 
 	FORCEINLINE_DEBUGGABLE void PostExternalCommandsReset()
 	{
@@ -3703,7 +3694,6 @@ public:
 	UE_DEPRECATED(5.5, "GenerateMips on RHI command lists is deprecated and no longer functions. Use the FGenerateMips helper class from the RenderCore module to generate mips on textures.")
 	RHI_API void GenerateMips(FRHITexture*);
 
-#if RHI_RAYTRACING
 	// Ray tracing API
 	
 	FORCEINLINE_DEBUGGABLE void CommitRayTracingBindings(FRHIRayTracingScene* Scene)
@@ -3926,8 +3916,6 @@ public:
 
 		SetRayTracingBindings(Scene, Pipeline, 1, InlineBindings, ERayTracingBindingType::MissShader, /*bCopyDataToInlineStorage*/ false);
 	}
-
-#endif // RHI_RAYTRACING
 };
 
 namespace EImmediateFlushType
@@ -5030,7 +5018,6 @@ FORCEINLINE void RHIUnlockStagingBuffer(FRHIStagingBuffer* StagingBuffer)
 	 FRHICommandListExecutor::GetImmediateCommandList().UnlockStagingBuffer(StagingBuffer);
 }
 
-#if RHI_RAYTRACING
 FORCEINLINE FRayTracingGeometryRHIRef RHICreateRayTracingGeometry(const FRayTracingGeometryInitializer& Initializer)
 {
 	return FRHICommandListExecutor::GetImmediateCommandList().CreateRayTracingGeometry(Initializer);
@@ -5040,7 +5027,6 @@ FORCEINLINE FRayTracingAccelerationStructureSize RHICalcRayTracingGeometrySize(c
 {
 	return GDynamicRHI->RHICalcRayTracingGeometrySize(Initializer);
 }
-#endif
 
 FORCEINLINE void RHIBindDebugLabelName(FRHITexture* Texture, const TCHAR* Name)
 {
diff --git a/Engine/Source/Runtime/RHI/Public/RHICommandListCommandExecutes.inl b/Engine/Source/Runtime/RHI/Public/RHICommandListCommandExecutes.inl
--- a/Engine/Source/Runtime/RHI/Public/RHICommandListCommandExecutes.inl
+++ b/Engine/Source/Runtime/RHI/Public/RHICommandListCommandExecutes.inl
@@ -475,8 +475,6 @@ void FRHICommandCopyBufferRegion::Execute(FRHICommandListBase& CmdList)
 	INTERNAL_DECORATOR(RHICopyBufferRegion)(DestBuffer, DstOffset, SourceBuffer, SrcOffset, NumBytes);
 }
 
-#if RHI_RAYTRACING
-
 void FRHICommandBindAccelerationStructureMemory::Execute(FRHICommandListBase& CmdList)
 {
 	INTERNAL_DECORATOR_COMPUTE(RHIBindAccelerationStructureMemory)(Scene, Buffer, BufferOffset);
@@ -526,8 +524,6 @@ void FRHICommandSetRayTracingBindings::Execute(FRHICommandListBase& CmdList)
 	INTERNAL_DECORATOR(RHISetRayTracingBindings)(Scene, GetRHIRayTracingPipelineState(Pipeline), NumBindings, Bindings, BindingType);
 }
 
-#endif // RHI_RAYTRACING
-
 void FRHICommandBeginScene::Execute(FRHICommandListBase& CmdList)
 {
 	RHISTAT(BeginScene);
