Removed InitializedType from FRHIRayTracingGeometry.
- in some cases (eg: FSkeletalMeshObjectStatic::InitResources(...)), when we create non-inline LODs ray tracing geometry the LOD is already streamed in, so ERayTracingGeometryInitializerType::Rendering is used during creation. Later if the LOD is streamed out/in again it will have unexpected InitializedType.

diff --git a/Engine/Source/Runtime/Apple/MetalRHI/Private/MetalRayTracing.cpp b/Engine/Source/Runtime/Apple/MetalRHI/Private/MetalRayTracing.cpp
--- a/Engine/Source/Runtime/Apple/MetalRHI/Private/MetalRayTracing.cpp
+++ b/Engine/Source/Runtime/Apple/MetalRHI/Private/MetalRayTracing.cpp
@@ -393,7 +393,6 @@ void FMetalRayTracingGeometry::ReleaseUnderlyingResource()
 
 void FMetalRayTracingGeometry::SetInitializer(FRHICommandListBase& RHICmdList, const FRayTracingGeometryInitializer& InInitializer)
 {
-	checkf(InitializedType == ERayTracingGeometryInitializerType::StreamingDestination, TEXT("Only FMetalRayTracingGeometry that was created as StreamingDestination can update their initializer."));
 	Initializer = InInitializer;
 
 	// HitGroup Parameters Update is handled by the Scene
diff --git a/Engine/Source/Runtime/D3D12RHI/Private/D3D12RayTracing.cpp b/Engine/Source/Runtime/D3D12RHI/Private/D3D12RayTracing.cpp
--- a/Engine/Source/Runtime/D3D12RHI/Private/D3D12RayTracing.cpp
+++ b/Engine/Source/Runtime/D3D12RHI/Private/D3D12RayTracing.cpp
@@ -3105,7 +3105,6 @@ void FD3D12RayTracingGeometry::UpdateResidency(FD3D12CommandContext& CommandCont
 
 void FD3D12RayTracingGeometry::SetInitializer(FRHICommandListBase& RHICmdList, const FRayTracingGeometryInitializer& InInitializer)
 {
-	checkf(InitializedType == ERayTracingGeometryInitializerType::StreamingDestination, TEXT("Only FD3D12RayTracingGeometry that was created as StreamingDestination can update their initializer."));
 	Initializer = InInitializer;
 
 	for (uint32 GPUIndex = 0; GPUIndex < MAX_NUM_GPUS && GPUIndex < GNumExplicitGPUsForRendering; ++GPUIndex)
diff --git a/Engine/Source/Runtime/RHI/Public/RHIResources.h b/Engine/Source/Runtime/RHI/Public/RHIResources.h
--- a/Engine/Source/Runtime/RHI/Public/RHIResources.h
+++ b/Engine/Source/Runtime/RHI/Public/RHIResources.h
@@ -3318,7 +3318,6 @@ public:
 	FRHIRayTracingGeometry() = default;
 	FRHIRayTracingGeometry(const FRayTracingGeometryInitializer& InInitializer)
 		: Initializer(InInitializer)
-		, InitializedType(InInitializer.Type)
 	{}
 
 	virtual FRayTracingAccelerationStructureAddress GetAccelerationStructureAddress(uint64 GPUIndex) const = 0;
@@ -3336,7 +3335,6 @@ public:
 	}
 protected:
 	FRayTracingGeometryInitializer Initializer = {};
-	ERayTracingGeometryInitializerType InitializedType = ERayTracingGeometryInitializerType::Rendering;
 };
 
 /** Top level ray tracing acceleration structure (contains instances of meshes). */
diff --git a/Engine/Source/Runtime/VulkanRHI/Private/VulkanRayTracing.cpp b/Engine/Source/Runtime/VulkanRHI/Private/VulkanRayTracing.cpp
--- a/Engine/Source/Runtime/VulkanRHI/Private/VulkanRayTracing.cpp
+++ b/Engine/Source/Runtime/VulkanRHI/Private/VulkanRayTracing.cpp
@@ -437,7 +437,6 @@ FVulkanRayTracingGeometry::~FVulkanRayTracingGeometry()
 
 void FVulkanRayTracingGeometry::SetInitializer(FRHICommandListBase& RHICmdList, const FRayTracingGeometryInitializer& InInitializer)
 {
-	checkf(InitializedType == ERayTracingGeometryInitializerType::StreamingDestination, TEXT("Only FVulkanRayTracingGeometry that was created as StreamingDestination can update their initializer."));
 	Initializer = InInitializer;
 
 	// TODO: Update HitGroup Parameters
