Added support to fetch graphics pipeline state objects in from worker threads if they are already in the cache. Also added a new SetGraphicsPipelineStateCheckApply so it can be called from outside of the state cache like SetComputePipelineState can.

diff --git a/Engine/Source/Runtime/RHI/Private/PipelineStateCache.cpp b/Engine/Source/Runtime/RHI/Private/PipelineStateCache.cpp
--- a/Engine/Source/Runtime/RHI/Private/PipelineStateCache.cpp
+++ b/Engine/Source/Runtime/RHI/Private/PipelineStateCache.cpp
@@ -918,7 +918,7 @@ FComputePipelineState* FindComputePipelineState(FRHIComputeShader* ComputeShader
 FComputePipelineState* GetComputePipelineState(FRHIComputeCommandList& RHICmdList, FRHIComputeShader* ComputeShader, bool bVerifyUse)
 {
 	FComputePipelineState* PipelineState = PipelineStateCache::GetAndOrCreateComputePipelineState(RHICmdList, ComputeShader, false);
-	if (bVerifyUse)
+	if (PipelineState && bVerifyUse)
 	{
 		PipelineState->Verify_IncUse();
 	}
@@ -931,16 +931,53 @@ void SetComputePipelineState(FRHIComputeCommandList& RHICmdList, FRHIComputeShad
 	RHICmdList.SetComputePipelineState(PipelineState, ComputeShader);
 }
 
+FGraphicsPipelineState* FindGraphicsPipelineState(const FGraphicsPipelineStateInitializer& Initializer, bool bVerifyUse)
+{
+	return PipelineStateCache::FindGraphicsPipelineState(Initializer, bVerifyUse);
+}
+
+FGraphicsPipelineState* GetGraphicsPipelineState(FRHICommandList& RHICmdList, const FGraphicsPipelineStateInitializer& Initializer, EApplyRendertargetOption ApplyFlags, bool bVerifyUse)
+{
+#if PLATFORM_USE_FALLBACK_PSO
+	checkNoEntry();
+	return nullptr;
+#else
+	FGraphicsPipelineState* PipelineState = PipelineStateCache::GetAndOrCreateGraphicsPipelineState(RHICmdList, Initializer, ApplyFlags);
+	if (PipelineState && bVerifyUse && !Initializer.bFromPSOFileCache)
+	{
+		PipelineState->Verify_IncUse();
+	}
+	return PipelineState;
+#endif
+}
+
+FGraphicsPipelineState* GetGraphicsPipelineState(FRHICommandList& RHICmdList, const FGraphicsPipelineStateInitializer& Initializer, bool bVerifyUse)
+{
+	return GetGraphicsPipelineState(RHICmdList, Initializer, EApplyRendertargetOption::CheckApply, bVerifyUse);
+}
+
 void SetGraphicsPipelineState(FRHICommandList& RHICmdList, const FGraphicsPipelineStateInitializer& Initializer, uint32 StencilRef, EApplyRendertargetOption ApplyFlags, bool bApplyAdditionalState)
 {
 #if PLATFORM_USE_FALLBACK_PSO
 	RHICmdList.SetGraphicsPipelineState(Initializer, StencilRef, bApplyAdditionalState);
 #else
-	FGraphicsPipelineState* PipelineState = PipelineStateCache::GetAndOrCreateGraphicsPipelineState(RHICmdList, Initializer, ApplyFlags);
+	FGraphicsPipelineState* PipelineState = GetGraphicsPipelineState(RHICmdList, Initializer, ApplyFlags, true /* bVerifyUse */);
+	if (PipelineState && !Initializer.bFromPSOFileCache)
+	{
+		check(IsInRenderingThread() || IsInParallelRenderingThread());
+		RHICmdList.SetGraphicsPipelineState(PipelineState, Initializer.BoundShaderState, StencilRef, bApplyAdditionalState);
+	}
+#endif
+}
 
+void SetGraphicsPipelineStateCheckApply(FRHICommandList& RHICmdList, const FGraphicsPipelineStateInitializer& Initializer, uint32 StencilRef, bool bApplyAdditionalState)
+{
+#if PLATFORM_USE_FALLBACK_PSO
+	RHICmdList.SetGraphicsPipelineState(Initializer, StencilRef, bApplyAdditionalState);
+#else
+	FGraphicsPipelineState* PipelineState = GetGraphicsPipelineState(RHICmdList, Initializer, true /* bVerifyUse */);
 	if (PipelineState && !Initializer.bFromPSOFileCache)
 	{
-		PipelineState->Verify_IncUse();
 		check(IsInRenderingThread() || IsInParallelRenderingThread());
 		RHICmdList.SetGraphicsPipelineState(PipelineState, Initializer.BoundShaderState, StencilRef, bApplyAdditionalState);
 	}
diff --git a/Engine/Source/Runtime/RHI/Public/PipelineStateCache.h b/Engine/Source/Runtime/RHI/Public/PipelineStateCache.h
--- a/Engine/Source/Runtime/RHI/Public/PipelineStateCache.h
+++ b/Engine/Source/Runtime/RHI/Public/PipelineStateCache.h
@@ -101,6 +101,7 @@ struct FPSOPrecacheRequestResult
 };
 
 extern RHI_API void SetComputePipelineState(FRHIComputeCommandList& RHICmdList, FRHIComputeShader* ComputeShader);
+extern RHI_API void SetGraphicsPipelineStateCheckApply(FRHICommandList& RHICmdList, const FGraphicsPipelineStateInitializer& Initializer, uint32 StencilRef, bool bApplyAdditionalState = true);
 extern RHI_API void SetGraphicsPipelineState(FRHICommandList& RHICmdList, const FGraphicsPipelineStateInitializer& Initializer, uint32 StencilRef, EApplyRendertargetOption ApplyFlags = EApplyRendertargetOption::CheckApply, bool bApplyAdditionalState = true);
 
 namespace PipelineStateCache
diff --git a/Engine/Source/Runtime/RHI/Public/RHICommandList.h b/Engine/Source/Runtime/RHI/Public/RHICommandList.h
--- a/Engine/Source/Runtime/RHI/Public/RHICommandList.h
+++ b/Engine/Source/Runtime/RHI/Public/RHICommandList.h
@@ -2423,6 +2423,8 @@ extern RHI_API FRHIComputePipelineState*	ExecuteSetComputePipelineState(FCompute
 extern RHI_API FRHIGraphicsPipelineState*	ExecuteSetGraphicsPipelineState(class FGraphicsPipelineState* GraphicsPipelineState);
 extern RHI_API FComputePipelineState*		FindComputePipelineState(FRHIComputeShader* ComputeShader, bool bVerifyUse = true);
 extern RHI_API FComputePipelineState*		GetComputePipelineState(FRHIComputeCommandList& RHICmdList, FRHIComputeShader* ComputeShader, bool bVerifyUse = true);
+extern RHI_API FGraphicsPipelineState*		FindGraphicsPipelineState(const FGraphicsPipelineStateInitializer& Initializer, bool bVerifyUse = true);
+extern RHI_API FGraphicsPipelineState*		GetGraphicsPipelineState(FRHICommandList& RHICmdList, const FGraphicsPipelineStateInitializer& Initializer, bool bVerifyUse = true);
 extern RHI_API FRHIComputePipelineState*	GetRHIComputePipelineState(FComputePipelineState*);
 extern RHI_API FRHIWorkGraphPipelineState*	GetRHIWorkGraphPipelineState(FWorkGraphPipelineState*);
 extern RHI_API FRHIRayTracingPipelineState*	GetRHIRayTracingPipelineState(FRayTracingPipelineState*);
diff --git a/Engine/Source/Runtime/Renderer/Private/MeshPassProcessor.cpp b/Engine/Source/Runtime/Renderer/Private/MeshPassProcessor.cpp
--- a/Engine/Source/Runtime/Renderer/Private/MeshPassProcessor.cpp
+++ b/Engine/Source/Runtime/Renderer/Private/MeshPassProcessor.cpp
@@ -1394,8 +1394,7 @@ bool FMeshDrawCommand::SubmitDrawBegin(
 		}
 
 		// We can set the new StencilRef here to avoid the set below
-		bool bApplyAdditionalState = true;
-		SetGraphicsPipelineState(RHICmdList, GraphicsPSOInit, MeshDrawCommand.StencilRef, EApplyRendertargetOption::CheckApply, bApplyAdditionalState);
+		SetGraphicsPipelineStateCheckApply(RHICmdList, GraphicsPSOInit, MeshDrawCommand.StencilRef);
 		StateCache.SetPipelineState(MeshDrawCommand.CachedPipelineId.GetId());
 		StateCache.StencilRef = MeshDrawCommand.StencilRef;
 	}
