Fix RHISupportsGeometryShaders not taking into account preview platform with SPIRV cross compilation.

SPRIV currently does not support geometry shader. This CL adds exclusion of geometry shader for platform preview.

diff --git a/Engine/Source/Runtime/RHI/Public/DataDrivenShaderPlatformInfo.h b/Engine/Source/Runtime/RHI/Public/DataDrivenShaderPlatformInfo.h
--- a/Engine/Source/Runtime/RHI/Public/DataDrivenShaderPlatformInfo.h
+++ b/Engine/Source/Runtime/RHI/Public/DataDrivenShaderPlatformInfo.h
@@ -991,7 +991,8 @@ inline bool RHISupportsGeometryShaders(const FStaticShaderPlatform Platform)
 	return IsFeatureLevelSupported(Platform, ERHIFeatureLevel::SM5)
 		&& !IsMetalPlatform(Platform)
 		&& !IsVulkanMobilePlatform(Platform)
-		&& !IsVulkanMobileSM5Platform(Platform);
+		&& !IsVulkanMobileSM5Platform(Platform)
+		&& !(FDataDrivenShaderPlatformInfo::GetIsPreviewPlatform(Platform) && FDataDrivenShaderPlatformInfo::GetIsSPIRV(Platform));
 }
 
 inline bool RHIHasTiledGPU(const FStaticShaderPlatform Platform)
