Graphics shader bundles fixes, and also implemented a skip nop packet

diff --git a/Engine/Source/Runtime/RHI/Public/RHIResources.h b/Engine/Source/Runtime/RHI/Public/RHIResources.h
--- a/Engine/Source/Runtime/RHI/Public/RHIResources.h
+++ b/Engine/Source/Runtime/RHI/Public/RHIResources.h
@@ -3429,6 +3429,8 @@ public:
 	const uint32 ArgStride	= 0;
 	const ERHIShaderBundleMode Mode = ERHIShaderBundleMode::CS;
 
+	FName ModeName;
+
 public:
 	FRHIShaderBundle(const FShaderBundleCreateInfo& CreateInfo)
 	: FRHIResource(RRT_ShaderBundle)
@@ -3441,16 +3443,19 @@ public:
 		{
 			// Load3
 			check(ArgStride >= 12u);
+			ModeName = TEXT("CS");
 		}
 		else if (Mode == ERHIShaderBundleMode::MSPS)
 		{
 			// Load
 			check(ArgStride >= 4u);
+			ModeName = TEXT("MSPS");
 		}
 		else if (Mode == ERHIShaderBundleMode::VSPS)
 		{
 			// Load4
 			check(ArgStride >= 16u);
+			ModeName = TEXT("VSPS");
 		}
 		else
 		{
