Verify uniform buffer bindings in the RHI validation layer.

#rnx
#tests Ran EngineTest with -rhivalidation on D3D12 and confirmed that no new validation errors are reported. Backed out CL 33173708 locally and checked that the repro steps in UE-211460 trigger a validation layer error before (possibly) crashing in the D3D12 RHI.

diff --git a/Engine/Source/Runtime/RHI/Private/RHIValidation.cpp b/Engine/Source/Runtime/RHI/Private/RHIValidation.cpp
--- a/Engine/Source/Runtime/RHI/Private/RHIValidation.cpp
+++ b/Engine/Source/Runtime/RHI/Private/RHIValidation.cpp
@@ -862,6 +862,59 @@ void FValidationRHI::ReportValidationFailure(const TCHAR* InMessage)
 	}
 }
 
+static void ValidateBoundUniformBuffers(FRHIShader* Shader, const RHIValidation::FStaticUniformBuffers& StaticUniformBuffers, const RHIValidation::FStageBoundUniformBuffers& BoundUniformBuffers)
+{
+	const TCHAR* FreqName = GetShaderFrequencyString(Shader->GetFrequency(), false);
+	const TArray<uint32>& LayoutHashes = Shader->GetShaderResourceTable().ResourceTableLayoutHashes;
+
+	const TArray<FUniformBufferStaticSlot>& StaticSlots = Shader->GetStaticSlots();
+	if (LayoutHashes.Num() != StaticSlots.Num())
+	{
+		RHI_VALIDATION_CHECK(false, *FString::Printf(TEXT("Shader %s(%s): The number of layout hashes (%d) is different from the number of static slots (%d)."), Shader->GetShaderName(), FreqName, LayoutHashes.Num(), StaticSlots.Num()));
+		return;
+	}
+
+	for (int32 BindIndex = 0; BindIndex < LayoutHashes.Num(); ++BindIndex)
+	{
+		uint32 ExpectedLayoutHash = LayoutHashes[BindIndex];
+		if (ExpectedLayoutHash == 0)
+		{
+			continue;
+		}
+
+		FRHIUniformBuffer* BoundBuffer = nullptr;
+		bool bIsStatic = false;
+
+		const FUniformBufferStaticSlot StaticSlot = StaticSlots[BindIndex];
+		if (IsUniformBufferStaticSlotValid(StaticSlot) && StaticSlot < StaticUniformBuffers.Bindings.Num())
+		{
+			BoundBuffer = StaticUniformBuffers.Bindings[StaticSlot];
+			if (BoundBuffer)
+			{
+				bIsStatic = true;
+			}
+		}
+
+		if (BoundBuffer == nullptr && BindIndex < BoundUniformBuffers.Buffers.Num())
+		{
+			BoundBuffer = BoundUniformBuffers.Buffers[BindIndex];
+		}
+
+		if (BoundBuffer != nullptr)
+		{
+			const FRHIUniformBufferLayout& Layout = BoundBuffer->GetLayout();
+			uint32 UniformBufferHash = Layout.GetHash();
+			RHI_VALIDATION_CHECK(UniformBufferHash == ExpectedLayoutHash, *FString::Printf(TEXT("Shader %s(%s): Invalid layout hash %u for uniform buffer \"%s\" at bind index %d (static: %s). Expecting a buffer called \"%s\", hash %u.)"),
+				Shader->GetShaderName(), FreqName, UniformBufferHash, *Layout.GetDebugName(), BindIndex, bIsStatic ? TEXT("yes") : TEXT("no"), *Shader->GetUniformBufferName(BindIndex), ExpectedLayoutHash));
+		}
+		else
+		{
+			RHI_VALIDATION_CHECK(false, *FString::Printf(TEXT("Shader %s(%s): missing uniform buffer \"%s\" at index %d."),
+				Shader->GetShaderName(), FreqName , *Shader->GetUniformBufferName(BindIndex), BindIndex));
+		}
+	}
+}
+
 FValidationComputeContext::FValidationComputeContext(EType InType)
 	: Type(InType)
 {
@@ -869,12 +922,24 @@ FValidationComputeContext::FValidationComputeContext(EType InType)
 	Tracker = &State.TrackerInstance;
 }
 
+void FValidationComputeContext::ValidateDispatch()
+{
+	if (State.BoundShader == nullptr)
+	{
+		RHI_VALIDATION_CHECK(false, TEXT("A compute PSO has to be set before dispatching a compute shader."));
+		return;
+	}
+
+	ValidateBoundUniformBuffers(State.BoundShader, State.StaticUniformBuffers, State.BoundUniformBuffers);
+}
+
 void FValidationComputeContext::FState::Reset()
 {
 	ComputePassName.Reset();
-	bComputePSOSet = false;
+	BoundShader = nullptr;
 	TrackerInstance.ResetAllUAVState();
 	StaticUniformBuffers.Reset();
+	BoundUniformBuffers.Reset();
 }
 
 FValidationContext::FValidationContext(EType InType)
@@ -948,6 +1013,35 @@ void FValidationComputeContext::RHICopyToStagingBuffer(FRHIBuffer* SourceBufferR
 	RHIContext->RHICopyToStagingBuffer(SourceBufferRHI, DestinationStagingBufferRHI, InOffset, InNumBytes);
 }
 
+void FValidationContext::ValidateDispatch()
+{
+	if (State.BoundShaders[SF_Compute] == nullptr)
+	{
+		RHI_VALIDATION_CHECK(false, TEXT("A compute PSO has to be set before dispatching a compute shader."));
+		return;
+	}
+
+	ValidateBoundUniformBuffers(State.BoundShaders[SF_Compute], State.StaticUniformBuffers, State.BoundUniformBuffers.Get(SF_Compute));
+}
+
+void FValidationContext::ValidateDrawing()
+{
+	if (!State.bGfxPSOSet)
+	{
+		RHI_VALIDATION_CHECK(false, TEXT("A graphics PSO has to be set in order to be able to draw!"));
+		return;
+	}
+
+	for (int32 FrequencyIndex = 0; FrequencyIndex < SF_NumFrequencies; ++FrequencyIndex)
+	{
+		EShaderFrequency Frequency = (EShaderFrequency)FrequencyIndex;
+		if (IsValidGraphicsFrequency(Frequency) && State.BoundShaders[Frequency])
+		{
+			ValidateBoundUniformBuffers(State.BoundShaders[Frequency], State.StaticUniformBuffers, State.BoundUniformBuffers.Get(Frequency));
+		}
+	}
+}
+
 void FValidationContext::FState::Reset()
 {
 	bInsideBeginRenderPass = false;
@@ -955,9 +1049,10 @@ void FValidationContext::FState::Reset()
 	RenderPassName.Reset();
 	PreviousRenderPassName.Reset();
 	ComputePassName.Reset();
-	bComputePSOSet = false;
+	FMemory::Memset(BoundShaders, 0);
 	TrackerInstance.ResetAllUAVState();
 	StaticUniformBuffers.Reset();
+	BoundUniformBuffers.Reset();
 }
 
 namespace RHIValidation
@@ -993,6 +1088,34 @@ namespace RHIValidation
 		}
 	}
 
+	FStageBoundUniformBuffers::FStageBoundUniformBuffers()
+	{
+		Buffers.Reserve(32);
+	}
+
+	void FStageBoundUniformBuffers::Reset()
+	{
+		Buffers.SetNum(0);
+	}
+
+	void FStageBoundUniformBuffers::Bind(uint32 Index, FRHIUniformBuffer* UniformBuffer)
+	{
+		if (Index >= (uint32)Buffers.Num())
+		{
+			Buffers.AddZeroed(Index + 1 - Buffers.Num());
+		}
+
+		Buffers[Index] = UniformBuffer;
+	}
+
+	void FBoundUniformBuffers::Reset()
+	{
+		for (FStageBoundUniformBuffers& Stage : StageBindings)
+		{
+			Stage.Reset();
+		}
+	}
+
 	ERHIAccess DecayResourceAccess(ERHIAccess AccessMask, ERHIAccess RequiredAccess, bool bAllowUAVOverlap)
 	{
 		using T = __underlying_type(ERHIAccess);
@@ -2549,8 +2672,30 @@ namespace RHIValidation
 	/** Validates that the Uniform conforms to what the shader expects */
 	void ValidateUniformBuffer(const FRHIShader* RHIShaderBase, uint32 BindIndex, FRHIUniformBuffer* UB)
 	{
+		if (!UB)
+		{
+			return;
+		}
+
+		const FRHIUniformBufferLayout& Layout = UB->GetLayout();
+
+		const TArray<uint32>& LayoutHashes = RHIShaderBase->GetShaderResourceTable().ResourceTableLayoutHashes;
+		if (BindIndex >= (uint32)LayoutHashes.Num())
+		{
+			FString ErrorMessage = FString::Printf(TEXT("Shader %s: Invalid bind index %u for uniform buffer \"%s\" (UB table size: %d)"), RHIShaderBase->GetShaderName(), BindIndex, *Layout.GetDebugName(), LayoutHashes.Num());
+			RHI_VALIDATION_CHECK(false, *ErrorMessage);
+			return;
+		}
+
+		uint32 ShaderTableHash = LayoutHashes[BindIndex];
+		uint32 UniformBufferHash = Layout.GetHash();
+		if (ShaderTableHash != 0 && UniformBufferHash != ShaderTableHash)
+		{
+			FString ErrorMessage = FString::Printf(TEXT("Shader %s: Invalid layout hash %u for uniform buffer \"%s\" at bind index %u, expecting %u"), RHIShaderBase->GetShaderName(), UniformBufferHash, *Layout.GetDebugName(), BindIndex, ShaderTableHash);
+			RHI_VALIDATION_CHECK(false, *ErrorMessage);
+		}
+
 #if RHI_INCLUDE_SHADER_DEBUG_DATA
-		if (UB)
 		{
 			// Validate Type
 			static const auto ShaderCodeValidationUBSizePredicate = [](const FShaderCodeValidationUBSize& lhs, const FShaderCodeValidationUBSize& rhs) -> bool { return lhs.BindPoint < rhs.BindPoint; };
@@ -2563,8 +2708,6 @@ namespace RHIValidation
 
 				if(Size > 0 && Size > UB->GetSize())
 				{
-					const FRHIUniformBufferLayout& Layout = UB->GetLayout();
-
 					FString ErrorMessage = FString::Printf(TEXT("Shader %s: Uniform buffer \"%s\" has unexpected size"), RHIShaderBase->GetShaderName(), *Layout.GetDebugName());
 					ErrorMessage += FString::Printf(TEXT("\nBind point: %d, HLSL size: %d, Actual size: %d"), BindIndex, Size, UB->GetSize());
 					RHI_VALIDATION_CHECK(false, *ErrorMessage);
@@ -2710,4 +2853,93 @@ void FValidationTransientResourceAllocator::Release(FRHICommandListImmediate& RH
 	delete this;
 }
 
+void ValidateShaderParameters(FRHIShader* RHIShader, RHIValidation::FTracker* Tracker, RHIValidation::FStaticUniformBuffers& StaticUniformBuffers, RHIValidation::FStageBoundUniformBuffers& BoundUniformBuffers, TConstArrayView<FRHIShaderParameterResource> InParameters, ERHIAccess InRequiredAccess, RHIValidation::EUAVMode InRequiredUAVMode)
+{
+	for (const FRHIShaderParameterResource& Parameter : InParameters)
+	{
+		switch (Parameter.Type)
+		{
+		case FRHIShaderParameterResource::EType::Texture:
+			if (FRHITexture* Texture = static_cast<FRHITexture*>(Parameter.Resource))
+			{
+				if (GRHIValidationEnabled)
+				{
+					RHIValidation::ValidateShaderResourceView(RHIShader, Parameter.Index, Texture);
+				}
+				Tracker->Assert(Texture->GetWholeResourceIdentitySRV(), InRequiredAccess);
+			}
+			break;
+		case FRHIShaderParameterResource::EType::ResourceView:
+			if (FRHIShaderResourceView* SRV = static_cast<FRHIShaderResourceView*>(Parameter.Resource))
+			{
+				if (GRHIValidationEnabled)
+				{
+					RHIValidation::ValidateShaderResourceView(RHIShader, Parameter.Index, SRV);
+				}
+				Tracker->Assert(SRV->GetViewIdentity(), InRequiredAccess);
+			}
+			break;
+		case FRHIShaderParameterResource::EType::UnorderedAccessView:
+			if (FRHIUnorderedAccessView* UAV = static_cast<FRHIUnorderedAccessView*>(Parameter.Resource))
+			{
+				if (GRHIValidationEnabled)
+				{
+					RHIValidation::ValidateUnorderedAccessView(RHIShader, Parameter.Index, UAV);
+				}
+				Tracker->AssertUAV(static_cast<FRHIUnorderedAccessView*>(Parameter.Resource), InRequiredUAVMode, Parameter.Index);
+			}
+			break;
+		case FRHIShaderParameterResource::EType::Sampler:
+			// No validation
+			break;
+		case FRHIShaderParameterResource::EType::UniformBuffer:
+			if (FRHIUniformBuffer* UniformBuffer = static_cast<FRHIUniformBuffer*>(Parameter.Resource))
+			{
+				if (GRHIValidationEnabled)
+				{
+					RHIValidation::ValidateUniformBuffer(RHIShader, Parameter.Index, UniformBuffer);
+				}
+
+				BoundUniformBuffers.Bind(Parameter.Index, UniformBuffer);
+				StaticUniformBuffers.ValidateSetShaderUniformBuffer(UniformBuffer);
+			}
+			break;
+		case FRHIShaderParameterResource::EType::ResourceCollection:
+			if (const FRHIResourceCollection* ResourceCollection = static_cast<const FRHIResourceCollection*>(Parameter.Resource))
+			{
+				for (const FRHIResourceCollectionMember& Member : ResourceCollection->Members)
+				{
+					switch (Member.Type)
+					{
+					case FRHIResourceCollectionMember::EType::Texture:
+						if (FRHITexture* Texture = static_cast<FRHITexture*>(Member.Resource))
+						{
+							Tracker->Assert(Texture->GetWholeResourceIdentitySRV(), InRequiredAccess);
+						}
+						break;
+					case FRHIResourceCollectionMember::EType::TextureReference:
+						if (FRHITextureReference* Texture = static_cast<FRHITextureReference*>(Member.Resource))
+						{
+							Tracker->Assert(Texture->GetWholeResourceIdentitySRV(), InRequiredAccess);
+						}
+						break;
+					case FRHIResourceCollectionMember::EType::ShaderResourceView:
+						if (FRHIShaderResourceView* SRV = static_cast<FRHIShaderResourceView*>(Parameter.Resource))
+						{
+							Tracker->Assert(SRV->GetViewIdentity(), InRequiredAccess);
+						}
+						break;
+					default:
+						break;
+					}
+				}
+			}
+			break;
+		default:
+			checkf(false, TEXT("Unhandled resource type?"));
+			break;
+		}
+	}
+}
+
 #endif	// ENABLE_RHI_VALIDATION
diff --git a/Engine/Source/Runtime/RHI/Public/RHIValidationCommon.h b/Engine/Source/Runtime/RHI/Public/RHIValidationCommon.h
--- a/Engine/Source/Runtime/RHI/Public/RHIValidationCommon.h
+++ b/Engine/Source/Runtime/RHI/Public/RHIValidationCommon.h
@@ -49,6 +49,23 @@ namespace RHIValidation
 		void ValidateSetShaderUniformBuffer(FRHIUniformBuffer* UniformBuffer);
 	};
 
+	struct FStageBoundUniformBuffers
+	{
+		FStageBoundUniformBuffers();
+		void Reset();
+		void Bind(uint32 Index, FRHIUniformBuffer* UniformBuffer);
+
+		TArray<FRHIUniformBuffer*> Buffers;
+	};
+
+	struct FBoundUniformBuffers
+	{
+		void Reset();
+		FStageBoundUniformBuffers& Get(EShaderFrequency Stage) { return StageBindings[Stage]; }
+
+		FStageBoundUniformBuffers StageBindings[SF_NumFrequencies];
+	};
+
 	class  FTracker;
 	class  FResource;
 	class  FTextureResource;
diff --git a/Engine/Source/Runtime/RHI/Public/RHIValidationContext.h b/Engine/Source/Runtime/RHI/Public/RHIValidationContext.h
--- a/Engine/Source/Runtime/RHI/Public/RHIValidationContext.h
+++ b/Engine/Source/Runtime/RHI/Public/RHIValidationContext.h
@@ -19,92 +19,7 @@
 
 class FValidationRHI;
 
-inline void ValidateShaderParameters(FRHIShader* RHIShader, RHIValidation::FTracker* Tracker, RHIValidation::FStaticUniformBuffers& StaticUniformBuffers, TConstArrayView<FRHIShaderParameterResource> InParameters, ERHIAccess InRequiredAccess, RHIValidation::EUAVMode InRequiredUAVMode)
-{
-	for (const FRHIShaderParameterResource& Parameter : InParameters)
-	{
-		switch (Parameter.Type)
-		{
-		case FRHIShaderParameterResource::EType::Texture:
-			if (FRHITexture* Texture = static_cast<FRHITexture*>(Parameter.Resource))
-			{
-				if (GRHIValidationEnabled)
-				{
-					RHIValidation::ValidateShaderResourceView(RHIShader, Parameter.Index, Texture);
-				}
-				Tracker->Assert(Texture->GetWholeResourceIdentitySRV(), InRequiredAccess);
-			}
-			break;
-		case FRHIShaderParameterResource::EType::ResourceView:
-			if (FRHIShaderResourceView* SRV = static_cast<FRHIShaderResourceView*>(Parameter.Resource))
-			{
-				if (GRHIValidationEnabled)
-				{
-					RHIValidation::ValidateShaderResourceView(RHIShader, Parameter.Index, SRV);
-				}
-				Tracker->Assert(SRV->GetViewIdentity(), InRequiredAccess);
-			}
-			break;
-		case FRHIShaderParameterResource::EType::UnorderedAccessView:
-			if (FRHIUnorderedAccessView* UAV = static_cast<FRHIUnorderedAccessView*>(Parameter.Resource))
-			{
-				if (GRHIValidationEnabled)
-				{
-					RHIValidation::ValidateUnorderedAccessView(RHIShader, Parameter.Index, UAV);
-				}
-				Tracker->AssertUAV(static_cast<FRHIUnorderedAccessView*>(Parameter.Resource), InRequiredUAVMode, Parameter.Index);
-			}
-			break;
-		case FRHIShaderParameterResource::EType::Sampler:
-			// No validation
-			break;
-		case FRHIShaderParameterResource::EType::UniformBuffer:
-			if (FRHIUniformBuffer* UniformBuffer = static_cast<FRHIUniformBuffer*>(Parameter.Resource))
-			{
-				if (GRHIValidationEnabled)
-				{
-					RHIValidation::ValidateUniformBuffer(RHIShader, Parameter.Index, UniformBuffer);
-				}
-				StaticUniformBuffers.ValidateSetShaderUniformBuffer(UniformBuffer);
-			}
-			break;
-		case FRHIShaderParameterResource::EType::ResourceCollection:
-			if (const FRHIResourceCollection* ResourceCollection = static_cast<const FRHIResourceCollection*>(Parameter.Resource))
-			{
-				for (const FRHIResourceCollectionMember& Member : ResourceCollection->Members)
-				{
-					switch (Member.Type)
-					{
-					case FRHIResourceCollectionMember::EType::Texture:
-						if (FRHITexture* Texture = static_cast<FRHITexture*>(Member.Resource))
-						{
-							Tracker->Assert(Texture->GetWholeResourceIdentitySRV(), InRequiredAccess);
-						}
-						break;
-					case FRHIResourceCollectionMember::EType::TextureReference:
-						if (FRHITextureReference* Texture = static_cast<FRHITextureReference*>(Member.Resource))
-						{
-							Tracker->Assert(Texture->GetWholeResourceIdentitySRV(), InRequiredAccess);
-						}
-						break;
-					case FRHIResourceCollectionMember::EType::ShaderResourceView:
-						if (FRHIShaderResourceView* SRV = static_cast<FRHIShaderResourceView*>(Parameter.Resource))
-						{
-							Tracker->Assert(SRV->GetViewIdentity(), InRequiredAccess);
-						}
-						break;
-					default:
-						break;
-					}
-				}
-			}
-			break;
-		default:
-			checkf(false, TEXT("Unhandled resource type?"));
-			break;
-		}
-	}
-}
+void ValidateShaderParameters(FRHIShader* RHIShader, RHIValidation::FTracker* Tracker, RHIValidation::FStaticUniformBuffers& StaticUniformBuffers, RHIValidation::FStageBoundUniformBuffers& BoundUniformBuffers, TConstArrayView<FRHIShaderParameterResource> InParameters, ERHIAccess InRequiredAccess, RHIValidation::EUAVMode InRequiredUAVMode);
 
 class FValidationComputeContext final : public IRHIComputeContext
 {
@@ -117,6 +32,8 @@ public:
 
 	FValidationComputeContext(EType Type);
 
+	void ValidateDispatch();
+
 	virtual ~FValidationComputeContext()
 	{
 	}
@@ -135,7 +52,7 @@ public:
 
 	virtual void RHISetComputePipelineState(FRHIComputePipelineState* ComputePipelineState) override final
 	{
-		State.bComputePSOSet = true;
+		State.BoundShader = ComputePipelineState->GetComputeShader();
 
 		// Reset the compute UAV tracker since the renderer must re-bind all resources after changing a shader.
 		Tracker->ResetUAVState(RHIValidation::EUAVMode::Compute);
@@ -147,7 +64,7 @@ public:
 
 	virtual void RHIDispatchComputeShader(uint32 ThreadGroupCountX, uint32 ThreadGroupCountY, uint32 ThreadGroupCountZ) override final
 	{
-		checkf(State.bComputePSOSet, TEXT("A Compute PSO has to be set to set resources into a shader!"));
+		ValidateDispatch();
 		FValidationRHI::ValidateThreadGroupCount(ThreadGroupCountX, ThreadGroupCountY, ThreadGroupCountZ);
 		RHIContext->RHIDispatchComputeShader(ThreadGroupCountX, ThreadGroupCountY, ThreadGroupCountZ);
 		Tracker->Dispatch();
@@ -155,10 +72,9 @@ public:
 
 	virtual void RHIDispatchIndirectComputeShader(FRHIBuffer* ArgumentBuffer, uint32 ArgumentOffset) override final
 	{
-		checkf(State.bComputePSOSet, TEXT("A Compute PSO has to be set to set resources into a shader!"));
+		ValidateDispatch();
 		FValidationRHI::ValidateDispatchIndirectArgsBuffer(ArgumentBuffer, ArgumentOffset);
 		Tracker->Assert(ArgumentBuffer->GetWholeResourceIdentity(), ERHIAccess::IndirectArgs);
-
 		RHIContext->RHIDispatchIndirectComputeShader(ArgumentBuffer, ArgumentOffset);
 		Tracker->Dispatch();
 	}
@@ -234,13 +150,13 @@ public:
 		checkf(Dispatches.Num() > 0, TEXT("A shader bundle must be dispatched with at least one record."));
 		for (const FRHIShaderBundleComputeDispatch& Dispatch : Dispatches)
 		{
-			State.bComputePSOSet = true;
+			State.BoundShader = Dispatch.Shader;
 
 			// Reset the compute UAV tracker since the renderer must re-bind all resources after changing a shader.
 			Tracker->ResetUAVState(RHIValidation::EUAVMode::Compute);
 
-			ValidateShaderParameters(Dispatch.Shader, Tracker, State.StaticUniformBuffers, Dispatch.Parameters.ResourceParameters, ERHIAccess::SRVCompute, RHIValidation::EUAVMode::Compute);
-			ValidateShaderParameters(Dispatch.Shader, Tracker, State.StaticUniformBuffers, Dispatch.Parameters.BindlessParameters, ERHIAccess::SRVCompute, RHIValidation::EUAVMode::Compute);
+			ValidateShaderParameters(Dispatch.Shader, Tracker, State.StaticUniformBuffers, State.BoundUniformBuffers, Dispatch.Parameters.ResourceParameters, ERHIAccess::SRVCompute, RHIValidation::EUAVMode::Compute);
+			ValidateShaderParameters(Dispatch.Shader, Tracker, State.StaticUniformBuffers, State.BoundUniformBuffers, Dispatch.Parameters.BindlessParameters, ERHIAccess::SRVCompute, RHIValidation::EUAVMode::Compute);
 
 			if (bEmulated)
 			{
@@ -316,17 +232,32 @@ public:
 
 	virtual void RHISetShaderParameters(FRHIComputeShader* Shader, TConstArrayView<uint8> InParametersData, TConstArrayView<FRHIShaderParameter> InParameters, TConstArrayView<FRHIShaderParameterResource> InResourceParameters, TConstArrayView<FRHIShaderParameterResource> InBindlessParameters) final override
 	{
-		checkf(State.bComputePSOSet, TEXT("A Compute PSO has to be set to set resources into a shader!"));
+		if (State.BoundShader == nullptr)
+		{
+			RHI_VALIDATION_CHECK(false, TEXT("A compute PSO has to be set to set resources into a shader!"));
+			return;
+		}
+
+		if (Shader != State.BoundShader)
+		{
+			RHI_VALIDATION_CHECK(false, *FString::Printf(TEXT("Invalid attempt to set parameters for compute shader '%s' while the currently bound shader is '%s'"), Shader->GetShaderName(), State.BoundShader->GetShaderName()));
+			return;
+		}
 
-		ValidateShaderParameters(Shader, Tracker, State.StaticUniformBuffers, InResourceParameters, ERHIAccess::SRVCompute, RHIValidation::EUAVMode::Compute);
-		ValidateShaderParameters(Shader, Tracker, State.StaticUniformBuffers, InBindlessParameters, ERHIAccess::SRVCompute, RHIValidation::EUAVMode::Compute);
+		ValidateShaderParameters(Shader, Tracker, State.StaticUniformBuffers, State.BoundUniformBuffers, InResourceParameters, ERHIAccess::SRVCompute, RHIValidation::EUAVMode::Compute);
+		ValidateShaderParameters(Shader, Tracker, State.StaticUniformBuffers, State.BoundUniformBuffers, InBindlessParameters, ERHIAccess::SRVCompute, RHIValidation::EUAVMode::Compute);
 
 		RHIContext->RHISetShaderParameters(Shader, InParametersData, InParameters, InResourceParameters, InBindlessParameters);
 	}
 
 	virtual void RHISetShaderUnbinds(FRHIComputeShader* Shader, TConstArrayView<FRHIShaderParameterUnbind> InUnbinds) final override
 	{
-		checkf(State.bComputePSOSet, TEXT("A Compute PSO has to be set to set resources into a shader!"));
+		if (State.BoundShader == nullptr)
+		{
+			RHI_VALIDATION_CHECK(false, TEXT("A compute PSO has to be set to set resources into a shader!"));
+			return;
+		}
+
 		RHIContext->RHISetShaderUnbinds(Shader, InUnbinds);
 	}
 
@@ -433,9 +364,10 @@ protected:
 	{
 		RHIValidation::FTracker TrackerInstance{ ERHIPipeline::AsyncCompute };
 		RHIValidation::FStaticUniformBuffers StaticUniformBuffers;
+		RHIValidation::FStageBoundUniformBuffers BoundUniformBuffers;
 
 		FString ComputePassName;
-		bool bComputePSOSet{};
+		FRHIComputeShader* BoundShader = nullptr;
 
 		void Reset();
 	} State;
@@ -468,7 +400,10 @@ public:
 
 	virtual void RHISetComputePipelineState(FRHIComputePipelineState* ComputePipelineState) override final
 	{
-		State.bComputePSOSet = true;
+		State.bGfxPSOSet = false;
+
+		FMemory::Memset(State.BoundShaders, 0);
+		State.BoundShaders[SF_Compute] = ComputePipelineState->GetComputeShader();
 
 		// Reset the compute UAV tracker since the renderer must re-bind all resources after changing a shader.
 		Tracker->ResetUAVState(RHIValidation::EUAVMode::Compute);
@@ -480,7 +415,7 @@ public:
 
 	virtual void RHIDispatchComputeShader(uint32 ThreadGroupCountX, uint32 ThreadGroupCountY, uint32 ThreadGroupCountZ) override final
 	{
-		checkf(State.bComputePSOSet, TEXT("A Compute PSO has to be set to set resources into a shader!"));
+		ValidateDispatch();
 		FValidationRHI::ValidateThreadGroupCount(ThreadGroupCountX, ThreadGroupCountY, ThreadGroupCountZ);
 		RHIContext->RHIDispatchComputeShader(ThreadGroupCountX, ThreadGroupCountY, ThreadGroupCountZ);
 		Tracker->Dispatch();
@@ -488,10 +423,9 @@ public:
 
 	virtual void RHIDispatchIndirectComputeShader(FRHIBuffer* ArgumentBuffer, uint32 ArgumentOffset) override final
 	{
-		checkf(State.bComputePSOSet, TEXT("A Compute PSO has to be set to set resources into a shader!"));
+		ValidateDispatch();
 		FValidationRHI::ValidateDispatchIndirectArgsBuffer(ArgumentBuffer, ArgumentOffset);
 		Tracker->Assert(ArgumentBuffer->GetWholeResourceIdentity(), ERHIAccess::IndirectArgs);
-
 		RHIContext->RHIDispatchIndirectComputeShader(ArgumentBuffer, ArgumentOffset);
 		Tracker->Dispatch();
 	}
@@ -535,15 +469,17 @@ public:
 		bool bEmulated) final override
 	{
 		checkf(Dispatches.Num() > 0, TEXT("A shader bundle must be dispatched with at least one record."));
+		RHIValidation::FStageBoundUniformBuffers& BoundUniformBuffers = State.BoundUniformBuffers.Get(SF_Compute);
+
 		for (const FRHIShaderBundleComputeDispatch& Dispatch : Dispatches)
 		{
-			State.bComputePSOSet = true;
-
+			State.BoundShaders[SF_Compute] = Dispatch.Shader;
+			
 			// Reset the compute UAV tracker since the renderer must re-bind all resources after changing a shader.
 			Tracker->ResetUAVState(RHIValidation::EUAVMode::Compute);
 
-			ValidateShaderParameters(Dispatch.Shader, Tracker, State.StaticUniformBuffers, Dispatch.Parameters.ResourceParameters, ERHIAccess::SRVCompute, RHIValidation::EUAVMode::Compute);
-			ValidateShaderParameters(Dispatch.Shader, Tracker, State.StaticUniformBuffers, Dispatch.Parameters.BindlessParameters, ERHIAccess::SRVCompute, RHIValidation::EUAVMode::Compute);
+			ValidateShaderParameters(Dispatch.Shader, Tracker, State.StaticUniformBuffers, BoundUniformBuffers, Dispatch.Parameters.ResourceParameters, ERHIAccess::SRVCompute, RHIValidation::EUAVMode::Compute);
+			ValidateShaderParameters(Dispatch.Shader, Tracker, State.StaticUniformBuffers, BoundUniformBuffers, Dispatch.Parameters.BindlessParameters, ERHIAccess::SRVCompute, RHIValidation::EUAVMode::Compute);
 
 			if (bEmulated)
 			{
@@ -790,7 +726,12 @@ public:
 	{
 		checkf(State.bInsideBeginRenderPass, TEXT("Graphics PSOs can only be set inside a RenderPass!"));
 		State.bGfxPSOSet = true;
-		State.bComputePSOSet = false;
+
+		for (int32 FrequencyIndex = 0; FrequencyIndex < SF_NumFrequencies; ++FrequencyIndex)
+		{
+			EShaderFrequency Frequency = (EShaderFrequency)FrequencyIndex;
+			State.BoundShaders[FrequencyIndex] = IsValidGraphicsFrequency(Frequency) ? GraphicsState->GetShader(Frequency) : nullptr;
+		}
 
 		ValidateDepthStencilForSetGraphicsPipelineState(GraphicsState->DSMode);
 
@@ -807,7 +748,13 @@ public:
 	{
 		checkf(State.bInsideBeginRenderPass, TEXT("Graphics PSOs can only be set inside a RenderPass!"));
 		State.bGfxPSOSet = true;
-		State.bComputePSOSet = false;
+
+		FMemory::Memset(State.BoundShaders, 0);
+		State.BoundShaders[SF_Vertex] = PsoInit.BoundShaderState.GetVertexShader();
+		State.BoundShaders[SF_Pixel] = PsoInit.BoundShaderState.GetPixelShader();
+		State.BoundShaders[SF_Geometry] = PsoInit.BoundShaderState.GetGeometryShader();
+		State.BoundShaders[SF_Amplification] = PsoInit.BoundShaderState.GetAmplificationShader();
+		State.BoundShaders[SF_Mesh] = PsoInit.BoundShaderState.GetMeshShader();
 
 		ValidateDepthStencilForSetGraphicsPipelineState(PsoInit.DepthStencilState->ActualDSMode);
 
@@ -822,33 +769,61 @@ public:
 
 	virtual void RHISetShaderParameters(FRHIGraphicsShader* Shader, TConstArrayView<uint8> InParametersData, TConstArrayView<FRHIShaderParameter> InParameters, TConstArrayView<FRHIShaderParameterResource> InResourceParameters, TConstArrayView<FRHIShaderParameterResource> InBindlessParameters) final override
 	{
-		checkf(State.bGfxPSOSet, TEXT("A Graphics PSO has to be set to set resources into a shader!"));
+		if (!State.bGfxPSOSet)
+		{
+			RHI_VALIDATION_CHECK(false, TEXT("A graphics PSO has to be set to set resources into a shader!"));
+			return;
+		}
 
-		ValidateShaderParameters(Shader, Tracker, State.StaticUniformBuffers, InResourceParameters, ERHIAccess::SRVGraphics, RHIValidation::EUAVMode::Graphics);
-		ValidateShaderParameters(Shader, Tracker, State.StaticUniformBuffers, InBindlessParameters, ERHIAccess::SRVGraphics, RHIValidation::EUAVMode::Graphics);
+		RHIValidation::FStageBoundUniformBuffers& BoundUniformBuffers = State.BoundUniformBuffers.Get(Shader->GetFrequency());
+
+		ValidateShaderParameters(Shader, Tracker, State.StaticUniformBuffers, BoundUniformBuffers, InResourceParameters, ERHIAccess::SRVGraphics, RHIValidation::EUAVMode::Graphics);
+		ValidateShaderParameters(Shader, Tracker, State.StaticUniformBuffers, BoundUniformBuffers, InBindlessParameters, ERHIAccess::SRVGraphics, RHIValidation::EUAVMode::Graphics);
 
 		RHIContext->RHISetShaderParameters(Shader, InParametersData, InParameters, InResourceParameters, InBindlessParameters);
 	}
 
 	virtual void RHISetShaderParameters(FRHIComputeShader* Shader, TConstArrayView<uint8> InParametersData, TConstArrayView<FRHIShaderParameter> InParameters, TConstArrayView<FRHIShaderParameterResource> InResourceParameters, TConstArrayView<FRHIShaderParameterResource> InBindlessParameters) final override
 	{
-		checkf(State.bComputePSOSet, TEXT("A Compute PSO has to be set to set resources into a shader!"));
+		if (State.BoundShaders[SF_Compute] == nullptr)
+		{
+			RHI_VALIDATION_CHECK(false, TEXT("A compute PSO has to be set to set resources into a shader!"));
+			return;
+		}
+
+		if (Shader != State.BoundShaders[SF_Compute])
+		{
+			RHI_VALIDATION_CHECK(false, *FString::Printf(TEXT("Invalid attempt to set parameters for compute shader '%s' while the currently bound shader is '%s'"), Shader->GetShaderName(), State.BoundShaders[SF_Compute]->GetShaderName()));
+			return;
+		}
 
-		ValidateShaderParameters(Shader, Tracker, State.StaticUniformBuffers, InResourceParameters, ERHIAccess::SRVCompute, RHIValidation::EUAVMode::Compute);
-		ValidateShaderParameters(Shader, Tracker, State.StaticUniformBuffers, InBindlessParameters, ERHIAccess::SRVCompute, RHIValidation::EUAVMode::Compute);
+		RHIValidation::FStageBoundUniformBuffers& BoundUniformBuffers = State.BoundUniformBuffers.Get(SF_Compute);
+
+		ValidateShaderParameters(Shader, Tracker, State.StaticUniformBuffers, BoundUniformBuffers, InResourceParameters, ERHIAccess::SRVCompute, RHIValidation::EUAVMode::Compute);
+		ValidateShaderParameters(Shader, Tracker, State.StaticUniformBuffers, BoundUniformBuffers, InBindlessParameters, ERHIAccess::SRVCompute, RHIValidation::EUAVMode::Compute);
 
 		RHIContext->RHISetShaderParameters(Shader, InParametersData, InParameters, InResourceParameters, InBindlessParameters);
 	}
 
 	virtual void RHISetShaderUnbinds(FRHIGraphicsShader* Shader, TConstArrayView<FRHIShaderParameterUnbind> InUnbinds) override final
 	{
-		checkf(State.bGfxPSOSet, TEXT("A Graphics PSO has to be set to set resources into a shader!"));
+		if (!State.bGfxPSOSet)
+		{
+			RHI_VALIDATION_CHECK(false, TEXT("A graphics PSO has to be set to set resources into a shader!"));
+			return;
+		}
+
 		RHIContext->RHISetShaderUnbinds(Shader, InUnbinds);
 	}
 
 	virtual void RHISetShaderUnbinds(FRHIComputeShader* Shader, TConstArrayView<FRHIShaderParameterUnbind> InUnbinds) override final
 	{
-		checkf(State.bComputePSOSet, TEXT("A Compute PSO has to be set to set resources into a shader!"));
+		if (State.BoundShaders[SF_Compute] == nullptr)
+		{
+			RHI_VALIDATION_CHECK(false, TEXT("A compute PSO has to be set to set resources into a shader!"));
+			return;
+		}
+
 		RHIContext->RHISetShaderUnbinds(Shader, InUnbinds);
 	}
 
@@ -872,17 +847,15 @@ public:
 
 	virtual void RHIDrawPrimitive(uint32 BaseVertexIndex, uint32 NumPrimitives, uint32 NumInstances) override final
 	{
-		checkf(State.bGfxPSOSet, TEXT("A Graphics PSO has to be set to draw!"));
-		SetupDrawing();
+		ValidateDrawing();
 		RHIContext->RHIDrawPrimitive(BaseVertexIndex, NumPrimitives, NumInstances);
 		Tracker->Draw();
 	}
 
 	virtual void RHIDrawPrimitiveIndirect(FRHIBuffer* ArgumentBuffer, uint32 ArgumentOffset) override final
 	{
-		checkf(State.bGfxPSOSet, TEXT("A Graphics PSO has to be set to draw!"));
+		ValidateDrawing();
 		FValidationRHI::ValidateIndirectArgsBuffer(ArgumentBuffer, ArgumentOffset, sizeof(FRHIDrawIndirectParameters), 0);
-		SetupDrawing();
 		Tracker->Assert(ArgumentBuffer->GetWholeResourceIdentity(), ERHIAccess::IndirectArgs);
 		RHIContext->RHIDrawPrimitiveIndirect(ArgumentBuffer, ArgumentOffset);
 		Tracker->Draw();
@@ -890,9 +863,8 @@ public:
 
 	virtual void RHIDrawIndexedIndirect(FRHIBuffer* IndexBufferRHI, FRHIBuffer* ArgumentsBufferRHI, int32 DrawArgumentsIndex, uint32 NumInstances) override final
 	{
-		checkf(State.bGfxPSOSet, TEXT("A Graphics PSO has to be set to draw!"));
+		ValidateDrawing();
 		FValidationRHI::ValidateIndirectArgsBuffer(ArgumentsBufferRHI, DrawArgumentsIndex * ArgumentsBufferRHI->GetStride(), sizeof(FRHIDrawIndexedIndirectParameters), 0);
-		SetupDrawing();
 		Tracker->Assert(ArgumentsBufferRHI->GetWholeResourceIdentity(), ERHIAccess::IndirectArgs);
 		Tracker->Assert(IndexBufferRHI->GetWholeResourceIdentity(), ERHIAccess::VertexOrIndexBuffer);
 		RHIContext->RHIDrawIndexedIndirect(IndexBufferRHI, ArgumentsBufferRHI, DrawArgumentsIndex, NumInstances);
@@ -902,9 +874,8 @@ public:
 	// @param NumPrimitives need to be >0 
 	virtual void RHIDrawIndexedPrimitive(FRHIBuffer* IndexBuffer, int32 BaseVertexIndex, uint32 FirstInstance, uint32 NumVertices, uint32 StartIndex, uint32 NumPrimitives, uint32 NumInstances) override final
 	{
-		checkf(State.bGfxPSOSet, TEXT("A Graphics PSO has to be set to draw!"));
+		ValidateDrawing();
 		checkf(EnumHasAnyFlags(IndexBuffer->GetUsage(), EBufferUsageFlags::IndexBuffer), TEXT("The buffer '%s' is used as an index buffer, but was not created with the IndexBuffer flag."), *IndexBuffer->GetName().ToString());
-		SetupDrawing();
 		Tracker->Assert(IndexBuffer->GetWholeResourceIdentity(), ERHIAccess::VertexOrIndexBuffer);
 		RHIContext->RHIDrawIndexedPrimitive(IndexBuffer, BaseVertexIndex, FirstInstance, NumVertices, StartIndex, NumPrimitives, NumInstances);
 		Tracker->Draw();
@@ -912,10 +883,9 @@ public:
 
 	virtual void RHIDrawIndexedPrimitiveIndirect(FRHIBuffer* IndexBuffer, FRHIBuffer* ArgumentBuffer, uint32 ArgumentOffset) override final
 	{
-		checkf(State.bGfxPSOSet, TEXT("A Graphics PSO has to be set to draw!"));
+		ValidateDrawing();
 		checkf(EnumHasAnyFlags(IndexBuffer->GetUsage(), EBufferUsageFlags::IndexBuffer), TEXT("The buffer '%s' is used as an index buffer, but was not created with the IndexBuffer flag."), *IndexBuffer->GetName().ToString());
 		FValidationRHI::ValidateIndirectArgsBuffer(ArgumentBuffer, ArgumentOffset, sizeof(FRHIDrawIndexedIndirectParameters), 0);
-		SetupDrawing();
 		Tracker->Assert(ArgumentBuffer->GetWholeResourceIdentity(), ERHIAccess::IndirectArgs);
 		Tracker->Assert(IndexBuffer->GetWholeResourceIdentity(), ERHIAccess::VertexOrIndexBuffer);
 		RHIContext->RHIDrawIndexedPrimitiveIndirect(IndexBuffer, ArgumentBuffer, ArgumentOffset);
@@ -924,10 +894,9 @@ public:
 
 	virtual void RHIMultiDrawIndexedPrimitiveIndirect(FRHIBuffer* IndexBuffer, FRHIBuffer* ArgumentBuffer, uint32 ArgumentOffset, FRHIBuffer* CountBuffer, uint32 CountBufferOffset, uint32 MaxDrawArguments) override final
 	{
-		checkf(State.bGfxPSOSet, TEXT("A Graphics PSO has to be set to draw!"));
+		ValidateDrawing();
 		checkf(EnumHasAnyFlags(IndexBuffer->GetUsage(), EBufferUsageFlags::IndexBuffer), TEXT("The buffer '%s' is used as an index buffer, but was not created with the IndexBuffer flag."), *IndexBuffer->GetName().ToString());
 		FValidationRHI::ValidateIndirectArgsBuffer(ArgumentBuffer, ArgumentOffset, sizeof(FRHIDrawIndexedIndirectParameters), 0);
-		SetupDrawing();
 		Tracker->Assert(ArgumentBuffer->GetWholeResourceIdentity(), ERHIAccess::IndirectArgs);
 		if (CountBuffer)
 		{
@@ -940,18 +909,16 @@ public:
 
 	virtual void RHIDispatchMeshShader(uint32 ThreadGroupCountX, uint32 ThreadGroupCountY, uint32 ThreadGroupCountZ) final override
 	{
-		checkf(State.bGfxPSOSet, TEXT("A Graphics PSO has to be set to draw!"));
+		ValidateDrawing();
 		FValidationRHI::ValidateThreadGroupCount(ThreadGroupCountX, ThreadGroupCountY, ThreadGroupCountZ);
-		SetupDrawing();
 		RHIContext->RHIDispatchMeshShader(ThreadGroupCountX, ThreadGroupCountY, ThreadGroupCountZ);
 		Tracker->Draw();
 	}
 
 	virtual void RHIDispatchIndirectMeshShader(FRHIBuffer* ArgumentBuffer, uint32 ArgumentOffset) final override
 	{
-		checkf(State.bGfxPSOSet, TEXT("A Graphics PSO has to be set to draw!"));
+		ValidateDrawing();
 		FValidationRHI::ValidateDispatchIndirectArgsBuffer(ArgumentBuffer, ArgumentOffset);
-		SetupDrawing();
 		Tracker->Assert(ArgumentBuffer->GetWholeResourceIdentity(), ERHIAccess::IndirectArgs);
 		RHIContext->RHIDispatchIndirectMeshShader(ArgumentBuffer, ArgumentOffset);
 		Tracker->Draw();
@@ -1221,10 +1188,8 @@ public:
 		RHIContext->RHISetBindingsOnShaderBindingTable(SBT, Scene, Pipeline, NumBindings, Bindings, BindingType);
 	}
 
-	void SetupDrawing()
-	{
-		// nothing to validate right now
-	}
+	void ValidateDispatch();
+	void ValidateDrawing();
 
 	IRHICommandContext* RHIContext = nullptr;
 
@@ -1240,6 +1205,7 @@ protected:
 	{
 		RHIValidation::FTracker TrackerInstance{ ERHIPipeline::Graphics };
 		RHIValidation::FStaticUniformBuffers StaticUniformBuffers;
+		RHIValidation::FBoundUniformBuffers BoundUniformBuffers;
 
 		void* PreviousBeginFrame = nullptr;
 		void* PreviousEndFrame = nullptr;
@@ -1249,8 +1215,10 @@ protected:
 		FString RenderPassName;
 		FString PreviousRenderPassName;
 		FString ComputePassName;
+
+		FRHIShader* BoundShaders[SF_NumFrequencies] = {};
+
 		bool bGfxPSOSet{};
-		bool bComputePSOSet{};
 		bool bInsideBeginRenderPass{};
 
 		void Reset();
