[RHI]
 - Add missing GraphEventRef cleanup to avoid potential crash at exit

#rnx

diff --git a/Engine/Source/Runtime/RHI/Private/RHICommandList.cpp b/Engine/Source/Runtime/RHI/Private/RHICommandList.cpp
--- a/Engine/Source/Runtime/RHI/Private/RHICommandList.cpp
+++ b/Engine/Source/Runtime/RHI/Private/RHICommandList.cpp
@@ -2160,6 +2160,7 @@ void FRHICommandListExecutor::CleanupGraphEvents()
 
 	LastMutate.SafeRelease();
 	LastSubmit.SafeRelease();
+	CompletionEvent.SafeRelease();
 }
 
 UE_DEPRECATED(5.5, "GenerateMips on RHI command lists is deprecated and no longer functions. Use the FGenerateMips helper class from the RenderCore module to generate mips on textures.")
