[RHI]
 - Fix another crash at shutdown because of unreleased graph events

#rnx

diff --git a/Engine/Source/Runtime/RHI/Private/RHICommandList.cpp b/Engine/Source/Runtime/RHI/Private/RHICommandList.cpp
--- a/Engine/Source/Runtime/RHI/Private/RHICommandList.cpp
+++ b/Engine/Source/Runtime/RHI/Private/RHICommandList.cpp
@@ -2161,6 +2161,9 @@ void FRHICommandListExecutor::CleanupGraphEvents()
 	LastMutate.SafeRelease();
 	LastSubmit.SafeRelease();
 	CompletionEvent.SafeRelease();
+
+	DispatchPipe.CleanupGraphEvents();
+	RHIThreadPipe.CleanupGraphEvents();
 }
 
 UE_DEPRECATED(5.5, "GenerateMips on RHI command lists is deprecated and no longer functions. Use the FGenerateMips helper class from the RenderCore module to generate mips on textures.")
diff --git a/Engine/Source/Runtime/RHI/Public/RHICommandList.h b/Engine/Source/Runtime/RHI/Public/RHICommandList.h
--- a/Engine/Source/Runtime/RHI/Public/RHICommandList.h
+++ b/Engine/Source/Runtime/RHI/Public/RHICommandList.h
@@ -5096,6 +5096,11 @@ private:
 
 		// Returns a graph event that will be signalled once all work submitted prior to calling Close() has completed.
 		FGraphEventRef Close();
+
+		void CleanupGraphEvents()
+		{
+			LastEvent = nullptr;
+		}
 	};
 
 	FTaskPipe DispatchPipe;
