Allow nested RHIBeginFrame calls.

Also fix crash when running in RHI bypass mode due to EndUAVOverlap trying to use the graphics context on async compute command lists.

#rnx

diff --git a/Engine/Source/Runtime/RHI/Private/RHICommandList.cpp b/Engine/Source/Runtime/RHI/Private/RHICommandList.cpp
--- a/Engine/Source/Runtime/RHI/Private/RHICommandList.cpp
+++ b/Engine/Source/Runtime/RHI/Private/RHICommandList.cpp
@@ -1691,6 +1691,13 @@ void FRHICommandListImmediate::BeginFrame()
 {
 	check(IsImmediate() && IsInRenderingThread());
 
+	++GRHICommandList.BeginFrameCount;
+	if (GRHICommandList.BeginFrameCount > 1)
+	{
+		// Allow nested begin calls, we don't need to do anything if we've already started the frame.
+		return;
+	}
+
 	{
 		QUICK_SCOPE_CYCLE_COUNTER(BeginFrame_Flush);
 		CSV_SCOPED_TIMING_STAT(RHITFlushes, BeginFrame);
@@ -1713,6 +1720,18 @@ void FRHICommandListImmediate::EndFrame()
 {
 	check(IsImmediate() && IsInRenderingThread());
 
+	if (GRHICommandList.BeginFrameCount <= 0)
+	{
+		checkf(false, TEXT("RHIEndFrame called without a matching RHIBeginFrame!"));
+		return;
+	}
+
+	--GRHICommandList.BeginFrameCount;
+	if (GRHICommandList.BeginFrameCount > 0)
+	{
+		return;
+	}
+
 	if (Bypass())
 	{
 		GetContext().RHIEndFrame();
diff --git a/Engine/Source/Runtime/RHI/Private/RHIValidation.cpp b/Engine/Source/Runtime/RHI/Private/RHIValidation.cpp
--- a/Engine/Source/Runtime/RHI/Private/RHIValidation.cpp
+++ b/Engine/Source/Runtime/RHI/Private/RHIValidation.cpp
@@ -951,14 +951,6 @@ FValidationContext::FValidationContext(EType InType)
 
 void FValidationContext::RHIBeginFrame()
 {
-	++State.BeginEndFrameCounter;
-	if (State.BeginEndFrameCounter != 1)
-	{
-		ensureMsgf(0, TEXT("RHIBeginFrame called twice in a row! Previous callstack: (void**)0x%p,32"), State.PreviousBeginFrame);
-	}
-	delete [] (uint64*)State.PreviousBeginFrame;
-	State.PreviousBeginFrame = RHIValidation::CaptureBacktrace();
-
 	State.Reset();
 	RHIContext->RHIBeginFrame();
 }
@@ -971,15 +963,6 @@ void FValidationContext::RHIEndFrame()
 	// The RenderThread FrameID is update in RHIAdvanceFrameFence which is called on the RenderThread
 	FValidationRHI* ValidateRHI = (FValidationRHI*)GDynamicRHI;
 	ValidateRHI->RHIThreadFrameID++;
-
-	--State.BeginEndFrameCounter;
-	if (State.BeginEndFrameCounter != 0)
-	{
-		ensureMsgf(0, TEXT("RHIEndFrame called twice in a row! Previous callstack: (void**)0x%p,32"), State.PreviousEndFrame);
-		State.BeginEndFrameCounter = 0;
-	}
-	delete [] (uint64*)State.PreviousEndFrame;
-	State.PreviousEndFrame = RHIValidation::CaptureBacktrace();
 }
 
 namespace RHIValidation
diff --git a/Engine/Source/Runtime/RHI/Public/RHICommandList.h b/Engine/Source/Runtime/RHI/Public/RHICommandList.h
--- a/Engine/Source/Runtime/RHI/Public/RHICommandList.h
+++ b/Engine/Source/Runtime/RHI/Public/RHICommandList.h
@@ -2922,7 +2922,7 @@ public:
 	{
 		if (Bypass())
 		{
-			GetContext().RHIEndUAVOverlap();
+			GetComputeContext().RHIEndUAVOverlap();
 			return;
 		}
 		ALLOC_COMMAND(FRHICommandEndUAVOverlap)();
@@ -5090,6 +5090,8 @@ private:
 	bool bEmitBreadcrumbs = false;
 #endif
 
+	int BeginFrameCount = 0;
+
 	friend class FRHICommandListBase;
 	friend class FRHICommandListImmediate;
 	FRHICommandListImmediate CommandListImmediate;
diff --git a/Engine/Source/Runtime/RHI/Public/RHIValidationContext.h b/Engine/Source/Runtime/RHI/Public/RHIValidationContext.h
--- a/Engine/Source/Runtime/RHI/Public/RHIValidationContext.h
+++ b/Engine/Source/Runtime/RHI/Public/RHIValidationContext.h
@@ -1209,10 +1209,6 @@ protected:
 		RHIValidation::FStaticUniformBuffers StaticUniformBuffers;
 		RHIValidation::FBoundUniformBuffers BoundUniformBuffers;
 
-		void* PreviousBeginFrame = nullptr;
-		void* PreviousEndFrame = nullptr;
-		int32 BeginEndFrameCounter = 0;
-
 		FRHIRenderPassInfo RenderPassInfo;
 		FString RenderPassName;
 		FString PreviousRenderPassName;
