Change check to verify for call to WaitCompletion during destructor of FPipelineState to make sure the PSO compilation is done before removing the possible final refs to the compute shader still being compiled

diff --git a/Engine/Source/Runtime/RHI/Private/PipelineStateCache.cpp b/Engine/Source/Runtime/RHI/Private/PipelineStateCache.cpp
--- a/Engine/Source/Runtime/RHI/Private/PipelineStateCache.cpp
+++ b/Engine/Source/Runtime/RHI/Private/PipelineStateCache.cpp
@@ -542,7 +542,7 @@ public:
 	virtual ~FPipelineState() 
 	{
 		check(IsComplete());
-		check(!WaitCompletion());
+		verify(!WaitCompletion());
 		check(!PrecompileTask.IsValid());
 	}
 
