Add something actionable to the logs when running into an out of date driver.

diff --git a/Engine/Source/Runtime/RHI/Private/DynamicRHI.cpp b/Engine/Source/Runtime/RHI/Private/DynamicRHI.cpp
--- a/Engine/Source/Runtime/RHI/Private/DynamicRHI.cpp
+++ b/Engine/Source/Runtime/RHI/Private/DynamicRHI.cpp
@@ -119,6 +119,11 @@ void RHIDetectAndWarnOfBadDrivers(bool bHasEditorToken)
 			return;
 		}
 
+		TOptional<FSuggestedDriverEntry> SuggestedDriver = DetectedGPUHardware.FindSuggestedDriverVersion();
+		const FString SuggestedDriverString = SuggestedDriver ? SuggestedDriver->SuggestedDriverVersion : FString(TEXT("Unknown"));
+
+		UE_LOG(LogRHI, Warning, TEXT("Out of date driver found. Using: '%s' Suggested: '%s'"), *DriverInfo.UserDriverVersion, (SuggestedDriver ? *SuggestedDriver->SuggestedDriverVersion : TEXT("Unknown")));
+
 		TArray<FString> DeviceCanUpdateDriverList;
 		GConfig->GetArray(TEXT("Devices"), TEXT("DeviceCanUpdateDriverList"), DeviceCanUpdateDriverList, GHardwareIni);
 
@@ -179,7 +184,6 @@ void RHIDetectAndWarnOfBadDrivers(bool bHasEditorToken)
 			Args.Add(TEXT("InstalledVer"), FText::FromString(DriverInfo.UserDriverVersion));
 
 			// Find the best driver version to recommend.
-			TOptional<FSuggestedDriverEntry> SuggestedDriver = DetectedGPUHardware.FindSuggestedDriverVersion();
 			if (SuggestedDriver)
 			{
 				// Suggest the latest too, if not denylisted.
