Remove validation for nested RHIBeginFrame calls, since nesting them is harmless in practice, and this mechanism is about to be deprecated anyway for the next release.

Also fix crash when running in RHI bypass mode due to EndUAVOverlap trying to use the graphics context on async compute command lists.

#rnx

diff --git a/Engine/Source/Runtime/RHI/Private/RHIValidation.cpp b/Engine/Source/Runtime/RHI/Private/RHIValidation.cpp
--- a/Engine/Source/Runtime/RHI/Private/RHIValidation.cpp
+++ b/Engine/Source/Runtime/RHI/Private/RHIValidation.cpp
@@ -951,14 +951,6 @@ FValidationContext::FValidationContext(EType InType)
 
 void FValidationContext::RHIBeginFrame()
 {
-	++State.BeginEndFrameCounter;
-	if (State.BeginEndFrameCounter != 1)
-	{
-		ensureMsgf(0, TEXT("RHIBeginFrame called twice in a row! Previous callstack: (void**)0x%p,32"), State.PreviousBeginFrame);
-	}
-	delete [] (uint64*)State.PreviousBeginFrame;
-	State.PreviousBeginFrame = RHIValidation::CaptureBacktrace();
-
 	State.Reset();
 	RHIContext->RHIBeginFrame();
 }
@@ -971,15 +963,6 @@ void FValidationContext::RHIEndFrame()
 	// The RenderThread FrameID is update in RHIAdvanceFrameFence which is called on the RenderThread
 	FValidationRHI* ValidateRHI = (FValidationRHI*)GDynamicRHI;
 	ValidateRHI->RHIThreadFrameID++;
-
-	--State.BeginEndFrameCounter;
-	if (State.BeginEndFrameCounter != 0)
-	{
-		ensureMsgf(0, TEXT("RHIEndFrame called twice in a row! Previous callstack: (void**)0x%p,32"), State.PreviousEndFrame);
-		State.BeginEndFrameCounter = 0;
-	}
-	delete [] (uint64*)State.PreviousEndFrame;
-	State.PreviousEndFrame = RHIValidation::CaptureBacktrace();
 }
 
 namespace RHIValidation
diff --git a/Engine/Source/Runtime/RHI/Public/RHICommandList.h b/Engine/Source/Runtime/RHI/Public/RHICommandList.h
--- a/Engine/Source/Runtime/RHI/Public/RHICommandList.h
+++ b/Engine/Source/Runtime/RHI/Public/RHICommandList.h
@@ -2993,7 +2993,7 @@ public:
 	{
 		if (Bypass())
 		{
-			GetContext().RHIEndUAVOverlap();
+			GetComputeContext().RHIEndUAVOverlap();
 			return;
 		}
 		ALLOC_COMMAND(FRHICommandEndUAVOverlap)();
diff --git a/Engine/Source/Runtime/RHI/Public/RHIValidationContext.h b/Engine/Source/Runtime/RHI/Public/RHIValidationContext.h
--- a/Engine/Source/Runtime/RHI/Public/RHIValidationContext.h
+++ b/Engine/Source/Runtime/RHI/Public/RHIValidationContext.h
@@ -1229,10 +1229,6 @@ protected:
 		RHIValidation::FStaticUniformBuffers StaticUniformBuffers;
 		RHIValidation::FBoundUniformBuffers BoundUniformBuffers;
 
-		void* PreviousBeginFrame = nullptr;
-		void* PreviousEndFrame = nullptr;
-		int32 BeginEndFrameCounter = 0;
-
 		FRHIRenderPassInfo RenderPassInfo;
 		FString RenderPassName;
 		FString PreviousRenderPassName;
