Fixed edge case in transient acquire logic where All -> Async needs to use the Graphics join pass. Fixes validation error due to acquiring too early.

diff --git a/Engine/Source/Runtime/RHI/Public/RHITransientResourceAllocator.h b/Engine/Source/Runtime/RHI/Public/RHITransientResourceAllocator.h
--- a/Engine/Source/Runtime/RHI/Public/RHITransientResourceAllocator.h
+++ b/Engine/Source/Runtime/RHI/Public/RHITransientResourceAllocator.h
@@ -437,12 +437,14 @@ public:
 			return Discard.Graphics;
 		}
 
-		// All | Async ->
-		return Acquire.AsyncCompute != Invalid
-			// Async Compute
-			? Discard.AsyncCompute
-			// Graphics
-			: Discard.GraphicsForkJoin.Max;
+		// AsyncCompute -> AsyncCompute
+		if (Acquire.AsyncCompute != Invalid && Discard.Graphics == Invalid)
+		{
+			return Discard.AsyncCompute;
+		}
+
+		// All | AsyncCompute -> Graphics or All -> AsyncCompute
+		return Discard.GraphicsForkJoin.Max;
 	}
 
 	// Returns whether two regions described by the discard and acquire fences contain each other. If they do, that means the memory would being used by both pipes simultaneously and cannot be aliased.
