Returns FSHAHash shader hash by const reference instead of value

#rnx

diff --git a/Engine/Source/Runtime/RHI/Public/RHIResources.h b/Engine/Source/Runtime/RHI/Public/RHIResources.h
--- a/Engine/Source/Runtime/RHI/Public/RHIResources.h
+++ b/Engine/Source/Runtime/RHI/Public/RHIResources.h
@@ -864,8 +864,8 @@ protected:
 class FRHIShader : public FRHIResource, public FRHIShaderData
 {
 public:
-	void SetHash(FSHAHash InHash) { Hash = InHash; }
-	FSHAHash GetHash() const { return Hash; }
+	void SetHash(const FSHAHash& InHash) { Hash = InHash; }
+	const FSHAHash& GetHash() const { return Hash; }
 
 #if RHI_INCLUDE_SHADER_DEBUG_DATA
 
