Fixed incorrect RHI transition acquire logic to handle All -> AsyncCompute case properly.

diff --git a/Engine/Source/Runtime/RHI/Public/RHITransientResourceAllocator.h b/Engine/Source/Runtime/RHI/Public/RHITransientResourceAllocator.h
--- a/Engine/Source/Runtime/RHI/Public/RHITransientResourceAllocator.h
+++ b/Engine/Source/Runtime/RHI/Public/RHITransientResourceAllocator.h
@@ -431,19 +431,29 @@ public:
 	{
 		check(Acquire.IsSinglePipeline());
 
-		// Graphics -> Graphics | AsyncCompute
-		if (Discard.Graphics != Invalid && Discard.AsyncCompute == Invalid)
+		if (Discard.Graphics != Invalid)
 		{
-			return Discard.Graphics;
+			// Graphics -> Graphics | AsyncCompute
+			if (Discard.AsyncCompute == Invalid)
+			{
+				return Discard.Graphics;
+			}
+
+			// All -> AsyncCompute
+			if (Acquire.AsyncCompute != Invalid)
+			{
+				// All -> AsyncCompute - The acquire graphics fork pass is used because a fence from Graphics -> AsyncCompute after the discard's graphics pass must be present.
+				return Acquire.GraphicsForkJoin.Min;
+			}
 		}
 
 		// AsyncCompute -> AsyncCompute
-		if (Acquire.AsyncCompute != Invalid && Discard.Graphics == Invalid)
+		if (Acquire.AsyncCompute != Invalid)
 		{
 			return Discard.AsyncCompute;
 		}
 
-		// All | AsyncCompute -> Graphics or All -> AsyncCompute
+		// All | AsyncCompute -> Graphics - The discard graphics fork pass is used because because a fence from AsyncCompute -> Graphics after the discard's async compute pass must be present.
 		return Discard.GraphicsForkJoin.Max;
 	}
 
