Bring resolved conflicts in to fix compile errors
#rnx

diff --git a/Engine/Source/Runtime/RHI/Private/RHICommandList.cpp b/Engine/Source/Runtime/RHI/Private/RHICommandList.cpp
--- a/Engine/Source/Runtime/RHI/Private/RHICommandList.cpp
+++ b/Engine/Source/Runtime/RHI/Private/RHICommandList.cpp
@@ -2229,7 +2229,6 @@ FRayTracingShaderBindings UE::RHI::ConvertRayTracingShaderBindings(const FRHIBat
 void FRHICommandList::RayTraceDispatch(
 	FRayTracingPipelineState* Pipeline,
 	FRHIRayTracingShader* RayGenShader,
-	FRHIRayTracingScene* Scene,
 	FRHIShaderBindingTable* SBT,
 	const FRHIBatchedShaderParameters& GlobalResourceBindings,
 	uint32 Width, uint32 Height)
@@ -2237,7 +2236,7 @@ void FRHICommandList::RayTraceDispatch(
 #if RHI_RAYTRACING
 
 	FRayTracingShaderBindings LegacyBindings = UE::RHI::ConvertRayTracingShaderBindings(GlobalResourceBindings);
-	RayTraceDispatch(Pipeline, RayGenShader, Scene, SBT, LegacyBindings, Width, Height);
+	RayTraceDispatch(Pipeline, RayGenShader, SBT, LegacyBindings, Width, Height);
 
 #else // RHI_RAYTRACING
 
@@ -2249,7 +2248,6 @@ void FRHICommandList::RayTraceDispatch(
 void FRHICommandList::RayTraceDispatchIndirect(
 	FRayTracingPipelineState* Pipeline,
 	FRHIRayTracingShader* RayGenShader,
-	FRHIRayTracingScene* Scene,
 	FRHIShaderBindingTable* SBT,
 	const FRHIBatchedShaderParameters& GlobalResourceBindings,
 	FRHIBuffer* ArgumentBuffer, uint32 ArgumentOffset)
@@ -2257,7 +2255,7 @@ void FRHICommandList::RayTraceDispatchIndirect(
 #if RHI_RAYTRACING
 
 	FRayTracingShaderBindings LegacyBindings = UE::RHI::ConvertRayTracingShaderBindings(GlobalResourceBindings);
-	RayTraceDispatchIndirect(Pipeline, RayGenShader, Scene, SBT, LegacyBindings, ArgumentBuffer, ArgumentOffset);
+	RayTraceDispatchIndirect(Pipeline, RayGenShader, SBT, LegacyBindings, ArgumentBuffer, ArgumentOffset);
 
 #else // RHI_RAYTRACING
 
diff --git a/Engine/Source/Runtime/RHI/Public/RHICommandList.h b/Engine/Source/Runtime/RHI/Public/RHICommandList.h
--- a/Engine/Source/Runtime/RHI/Public/RHICommandList.h
+++ b/Engine/Source/Runtime/RHI/Public/RHICommandList.h
@@ -4027,7 +4027,7 @@ public:
 	* Compatibility adaptor that operates on the new FRHIBatchedShaderParameters instead of legacy FRayTracingShaderBindings (planned for deprecation).
 	* This will become the default native code path in a future UE version.
 	*/
-	RHI_API void RayTraceDispatch(FRayTracingPipelineState* Pipeline, FRHIRayTracingShader* RayGenShader, FRHIRayTracingScene* Scene, FRHIShaderBindingTable* SBT, const FRHIBatchedShaderParameters& GlobalResourceBindings, uint32 Width, uint32 Height);
+	RHI_API void RayTraceDispatch(FRayTracingPipelineState* Pipeline, FRHIRayTracingShader* RayGenShader, FRHIShaderBindingTable* SBT, const FRHIBatchedShaderParameters& GlobalResourceBindings, uint32 Width, uint32 Height);
 
 	/**
 	 * Trace rays using dimensions from a GPU buffer containing uint[3], interpreted as number of rays in X, Y and Z dimensions.
@@ -4063,7 +4063,7 @@ public:
 	* Compatibility adaptor that operates on the new FRHIBatchedShaderParameters instead of legacy FRayTracingShaderBindings (planned for deprecation).
 	* This will become the default native code path in a future UE version.
 	*/
-	RHI_API void RayTraceDispatchIndirect(FRayTracingPipelineState* Pipeline, FRHIRayTracingShader* RayGenShader, FRHIRayTracingScene* Scene, FRHIShaderBindingTable* SBT, const FRHIBatchedShaderParameters& GlobalResourceBindings, FRHIBuffer* ArgumentBuffer, uint32 ArgumentOffset);
+	RHI_API void RayTraceDispatchIndirect(FRayTracingPipelineState* Pipeline, FRHIRayTracingShader* RayGenShader, FRHIShaderBindingTable* SBT, const FRHIBatchedShaderParameters& GlobalResourceBindings, FRHIBuffer* ArgumentBuffer, uint32 ArgumentOffset);
 
 	UE_DEPRECATED(5.5, "Use FRHIShaderBindingTable instead.")
 	FORCEINLINE_DEBUGGABLE void SetRayTracingBindings(
