Fixed various RHI validation errors and inefficiencies due to async compute.
 - Fixed shader validation to handle buffer types for texture metadata.
 - RDG external access mode will now leave the resource in its external state / pipelines.
 - Modified the RHI validation layer to allow transitioning All -> One pipe without a fence if the NoFence flag is specified. Added fence validation to emit an error of an external fence is not used that might introduce a race condition.
 - Fixed VT / GPU Skin cache buffers to use read states on all pipes.
 - Batched multiple VT external access calls together to reduce passes in RDG.
 - Fixed IES texture manager validation error.
 - Fixed sky cube texture map validation error.
 - Fixed virtual shadow debug shader that was attempting to read and writing using the same texture.
 - Fixed render thread CVar access on game thread.
 - Fixed up -ForceRHIBypass to not crash.

diff --git a/Engine/Source/Runtime/Engine/Private/GPUSkinCache.cpp b/Engine/Source/Runtime/Engine/Private/GPUSkinCache.cpp
--- a/Engine/Source/Runtime/Engine/Private/GPUSkinCache.cpp
+++ b/Engine/Source/Runtime/Engine/Private/GPUSkinCache.cpp
@@ -856,7 +856,7 @@ void FGPUSkinCache::TransitionAllToReadable(FRHICommandList& RHICmdList, const T
 				UAVs.Add(Buffer->UpdateAccessState(ToState));
 			}
 		}
-		RHICmdList.Transition(UAVs);
+		RHICmdList.Transition(UAVs, ERHIPipeline::Graphics, IsGPUSkinCacheRayTracingSupported() ? ERHIPipeline::All : ERHIPipeline::Graphics);
 	}
 }
 
@@ -1173,7 +1173,7 @@ void FGPUSkinCache::DispatchUpdateSkinTangents(FRHICommandList& RHICmdList, FGPU
 				RHICmdList.Transition({
 					DispatchData.GetActiveTangentRWBuffer()->UpdateAccessState(ERHIAccess::SRVCompute),
 					StagingBuffer->UpdateAccessState(ERHIAccess::UAVCompute)
-				});
+				}, ERHITransitionCreateFlags::NoFence);
 			}
 
 			INC_DWORD_STAT_BY(STAT_GPUSkinCache_NumTrianglesForRecomputeTangents, NumTriangles);
@@ -1221,7 +1221,7 @@ void FGPUSkinCache::DispatchUpdateSkinTangents(FRHICommandList& RHICmdList, FGPU
 			RHICmdList.Transition({
 				DispatchData.GetTangentRWBuffer()->UpdateAccessState(ERHIAccess::UAVCompute),
 				StagingBuffer->UpdateAccessState(ERHIAccess::UAVCompute)
-				});
+			}, ERHITransitionCreateFlags::NoFence);
 		}
 
 		SetComputePipelineState(RHICmdList, ComputeShader.GetComputeShader());
@@ -1282,7 +1282,7 @@ void FGPUSkinCache::MakeBufferTransitions(FRHICommandList& RHICmdList, TArray<FS
 				UAVs.Add(Buffer->UpdateAccessState(ToState));
 			}
 		}
-		RHICmdList.Transition(MakeArrayView(UAVs.GetData(), UAVs.Num()));
+		RHICmdList.Transition(MakeArrayView(UAVs.GetData(), UAVs.Num()), ERHITransitionCreateFlags::NoFence);
 	}
 }
 
diff --git a/Engine/Source/Runtime/Engine/Private/VT/VirtualTextureUploadCache.cpp b/Engine/Source/Runtime/Engine/Private/VT/VirtualTextureUploadCache.cpp
--- a/Engine/Source/Runtime/Engine/Private/VT/VirtualTextureUploadCache.cpp
+++ b/Engine/Source/Runtime/Engine/Private/VT/VirtualTextureUploadCache.cpp
@@ -361,7 +361,7 @@ void FVirtualTextureUploadCache::Finalize(FRDGBuilder& GraphBuilder)
 
 			if (!UpdatedTextures.Contains(Entry.RHISubmitTexture))
 			{
-				RHICmdList.Transition(FRHITransitionInfo(Entry.RHISubmitTexture, ERHIAccess::Unknown, ERHIAccess::CopyDest));
+				RHICmdList.Transition(FRHITransitionInfo(Entry.RHISubmitTexture, ERHIAccess::Unknown, ERHIAccess::CopyDest), ERHITransitionCreateFlags::NoFence);
 				UpdatedTextures.Add(Entry.RHISubmitTexture);
 			}
 
@@ -384,7 +384,7 @@ void FVirtualTextureUploadCache::Finalize(FRDGBuilder& GraphBuilder)
 	{
 		SRVTransitions.Add(FRHITransitionInfo(UpdatedTextures[Index], ERHIAccess::CopyDest, ERHIAccess::SRVMask));
 	}
-	RHICmdList.Transition(SRVTransitions);
+	RHICmdList.Transition(SRVTransitions, ERHIPipeline::Graphics, ERHIPipeline::All);
 	UpdatedTextures.Reset();
 }
 
@@ -436,7 +436,7 @@ void FVirtualTextureUploadCache::SubmitTile(FRHICommandList& RHICmdList, const F
 	{
 		if (!UpdatedTextures.Contains(InDestTexture))
 		{
-			RHICmdList.Transition(FRHITransitionInfo(InDestTexture, ERHIAccess::Unknown, ERHIAccess::CopyDest));
+			RHICmdList.Transition(FRHITransitionInfo(InDestTexture, ERHIAccess::Unknown, ERHIAccess::CopyDest), ERHITransitionCreateFlags::NoFence);
 			UpdatedTextures.Add(InDestTexture);
 		}
 
@@ -452,7 +452,7 @@ void FVirtualTextureUploadCache::SubmitTile(FRHICommandList& RHICmdList, const F
 	{
 		if (!UpdatedTextures.Contains(InDestTexture))
 		{
-			RHICmdList.Transition(FRHITransitionInfo(InDestTexture, ERHIAccess::Unknown, ERHIAccess::CopyDest));
+			RHICmdList.Transition(FRHITransitionInfo(InDestTexture, ERHIAccess::Unknown, ERHIAccess::CopyDest), ERHITransitionCreateFlags::NoFence);
 			UpdatedTextures.Add(InDestTexture);
 		}
 
diff --git a/Engine/Source/Runtime/RHI/Private/RHICommandList.cpp b/Engine/Source/Runtime/RHI/Private/RHICommandList.cpp
--- a/Engine/Source/Runtime/RHI/Private/RHICommandList.cpp
+++ b/Engine/Source/Runtime/RHI/Private/RHICommandList.cpp
@@ -1508,47 +1508,6 @@ void FRHICommandListExecutor::WaitForTasks()
 	}
 }
 
-void FRHICommandListImmediate::Transition(TArrayView<const FRHITransitionInfo> Infos, ERHIPipeline SrcPipelines, ERHIPipeline DstPipelines)
-{
-	check(IsInRenderingThread());
-
-#if DO_CHECK
-	for (const FRHITransitionInfo& Info : Infos)
-	{
-		checkf(Info.IsWholeResource(), TEXT("Only whole resource transitions are allowed in FRHICommandListImmediate::Transition."));
-	}
-#endif
-
-	if (!GSupportsEfficientAsyncCompute)
-	{
-		checkf(SrcPipelines != ERHIPipeline::AsyncCompute, TEXT("Async compute is disabled. Cannot transition from it."));
-		checkf(DstPipelines != ERHIPipeline::AsyncCompute, TEXT("Async compute is disabled. Cannot transition to it."));
-
-		EnumRemoveFlags(SrcPipelines, ERHIPipeline::AsyncCompute);
-		EnumRemoveFlags(DstPipelines, ERHIPipeline::AsyncCompute);
-	}
-
-	const FRHITransition* Transition = RHICreateTransition({ SrcPipelines, DstPipelines, ERHITransitionCreateFlags::None, Infos });
-
-	for (ERHIPipeline Pipeline : MakeFlagsRange(SrcPipelines))
-	{
-		FRHICommandListScopedPipeline Scope(*this, Pipeline);
-		BeginTransition(Transition);
-	}
-
-	for (ERHIPipeline Pipeline : MakeFlagsRange(DstPipelines))
-	{
-		FRHICommandListScopedPipeline Scope(*this, Pipeline);
-		EndTransition(Transition);
-	}
-
-	if (EnumHasAnyFlags(SrcPipelines | DstPipelines, ERHIPipeline::Graphics))
-	{
-		FRHICommandListScopedPipeline Scope(*this, ERHIPipeline::Graphics);
-		SetTrackedAccess(Infos);
-	}
-}
-
 bool FRHICommandListImmediate::IsStalled()
 {
 	return GRHIThreadStallRequestCount.load() > 0;
@@ -1736,7 +1695,7 @@ void FRHICommandListImmediate::EndFrame()
 	}
 }
 
-void FRHIComputeCommandList::Transition(TArrayView<const FRHITransitionInfo> Infos)
+void FRHIComputeCommandList::Transition(TArrayView<const FRHITransitionInfo> Infos, ERHITransitionCreateFlags CreateFlags)
 {
 	const ERHIPipeline Pipeline = GetPipeline();
 
@@ -1746,7 +1705,7 @@ void FRHIComputeCommandList::Transition(TArrayView<const FRHITransitionInfo> Inf
 		FMemStack& MemStack = FMemStack::Get();
 		FMemMark Mark(MemStack);
 		FRHITransition* Transition = new (MemStack.Alloc(FRHITransition::GetTotalAllocationSize(), FRHITransition::GetAlignment())) FRHITransition(Pipeline, Pipeline);
-		GDynamicRHI->RHICreateTransition(Transition, FRHITransitionCreateInfo(Pipeline, Pipeline, ERHITransitionCreateFlags::NoSplit, Infos));
+		GDynamicRHI->RHICreateTransition(Transition, FRHITransitionCreateInfo(Pipeline, Pipeline, CreateFlags | ERHITransitionCreateFlags::NoSplit, Infos));
 
 		GetComputeContext().RHIBeginTransitions(MakeArrayView((const FRHITransition**)&Transition, 1));
 		GetComputeContext().RHIEndTransitions(MakeArrayView((const FRHITransition**)&Transition, 1));
@@ -1759,7 +1718,7 @@ void FRHIComputeCommandList::Transition(TArrayView<const FRHITransitionInfo> Inf
 	{
 		// Allocate the transition in the command list
 		FRHITransition* Transition = new (Alloc(FRHITransition::GetTotalAllocationSize(), FRHITransition::GetAlignment())) FRHITransition(Pipeline, Pipeline);
-		GDynamicRHI->RHICreateTransition(Transition, FRHITransitionCreateInfo(Pipeline, Pipeline, ERHITransitionCreateFlags::NoSplit, Infos));
+		GDynamicRHI->RHICreateTransition(Transition, FRHITransitionCreateInfo(Pipeline, Pipeline, CreateFlags | ERHITransitionCreateFlags::NoSplit, Infos));
 
 		ALLOC_COMMAND(FRHICommandResourceTransition)(Transition);
 	}
@@ -1775,6 +1734,42 @@ void FRHIComputeCommandList::Transition(TArrayView<const FRHITransitionInfo> Inf
 	}
 }
 
+void FRHIComputeCommandList::Transition(TArrayView<const FRHITransitionInfo> Infos, ERHIPipeline SrcPipelines, ERHIPipeline DstPipelines, ERHITransitionCreateFlags TransitionCreateFlags)
+{
+#if DO_CHECK
+	for (const FRHITransitionInfo& Info : Infos)
+	{
+		checkf(Info.IsWholeResource(), TEXT("Only whole resource transitions are allowed in FRHIComputeCommandList::Transition."));
+	}
+#endif
+
+	if (!GSupportsEfficientAsyncCompute)
+	{
+		SrcPipelines = ERHIPipeline::Graphics;
+		DstPipelines = ERHIPipeline::Graphics;
+	}
+
+	const FRHITransition* Transition = RHICreateTransition({ SrcPipelines, DstPipelines, TransitionCreateFlags, Infos });
+
+	for (ERHIPipeline Pipeline : MakeFlagsRange(SrcPipelines))
+	{
+		FRHICommandListScopedPipeline Scope(*this, Pipeline);
+		BeginTransition(Transition);
+	}
+
+	for (ERHIPipeline Pipeline : MakeFlagsRange(DstPipelines))
+	{
+		FRHICommandListScopedPipeline Scope(*this, Pipeline);
+		EndTransition(Transition);
+	}
+
+	{
+		// Set the tracked access on only one of the destination pipes.
+		FRHICommandListScopedPipeline Scope(*this, DstPipelines == ERHIPipeline::AsyncCompute ? ERHIPipeline::AsyncCompute : ERHIPipeline::Graphics);
+		SetTrackedAccess(Infos);
+	}
+}
+
 void FRHIComputeCommandList::BuildAccelerationStructure(FRHIRayTracingGeometry* Geometry)
 {
 	FRayTracingGeometryBuildParams Params;
diff --git a/Engine/Source/Runtime/RHI/Private/RHIValidation.cpp b/Engine/Source/Runtime/RHI/Private/RHIValidation.cpp
--- a/Engine/Source/Runtime/RHI/Private/RHIValidation.cpp
+++ b/Engine/Source/Runtime/RHI/Private/RHIValidation.cpp
@@ -513,14 +513,7 @@ void FValidationRHI::RHICreateTransition(FRHITransition* Transition, const FRHIT
 	const ERHIPipeline SrcPipelines = CreateInfo.SrcPipelines;
 	const ERHIPipeline DstPipelines = CreateInfo.DstPipelines;
 
-	struct FFenceEdge
-	{
-		FFence* Fence = nullptr;
-		ERHIPipeline SrcPipe = ERHIPipeline::None;
-		ERHIPipeline DstPipe = ERHIPipeline::None;
-	};
-
-	TArray<FFenceEdge> Fences;
+	TArray<FFence*> Fences;
 
 	if (SrcPipelines != DstPipelines)
 	{
@@ -533,11 +526,10 @@ void FValidationRHI::RHICreateTransition(FRHITransition* Transition, const FRHIT
 					continue;
 				}
 
-				FFenceEdge FenceEdge;
-				FenceEdge.Fence = new FFence;
-				FenceEdge.SrcPipe = SrcPipe;
-				FenceEdge.DstPipe = DstPipe;
-				Fences.Add(FenceEdge);
+				FFence* Fence = new FFence;
+				Fence->SrcPipe = SrcPipe;
+				Fence->DstPipe = DstPipe;
+				Fences.Add(Fence);
 			}
 		}
 	}
@@ -550,9 +542,9 @@ void FValidationRHI::RHICreateTransition(FRHITransition* Transition, const FRHIT
 	BeginOps          .Reserve(CreateInfo.TransitionInfos.Num());
 	EndOps            .Reserve(CreateInfo.TransitionInfos.Num());
 
-	for (const FFenceEdge& FenceEdge : Fences)
+	for (FFence* Fence : Fences)
 	{
-		WaitOps[FenceEdge.DstPipe].Emplace(FOperation::Wait(FenceEdge.Fence, FenceEdge.DstPipe));
+		WaitOps[Fence->DstPipe].Emplace(FOperation::Wait(Fence));
 	}
 
 	// Take a backtrace of this transition creation if any of the resources it contains have logging enabled.
@@ -615,9 +607,6 @@ void FValidationRHI::RHICreateTransition(FRHITransition* Transition, const FRHIT
 
 		if (const FRHICommitResourceInfo* CommitInfo = Info.CommitInfo.GetPtrOrNull())
 		{
-			RHI_VALIDATION_CHECK((SrcPipelines == ERHIPipeline::Graphics && DstPipelines == ERHIPipeline::Graphics),
-				TEXT("Reserved resource commit operations are only supported on the graphics pipeline and must not cross pipeline boundary."));
-
 			if (Info.Type == FRHITransitionInfo::EType::Buffer)
 			{
 				const FRHIBuffer* Buffer = Info.Buffer;
@@ -659,7 +648,7 @@ void FValidationRHI::RHICreateTransition(FRHITransition* Transition, const FRHIT
 		FState PreviousState = FState(Info.AccessBefore, SrcPipelines);
 		FState NextState = FState(Info.AccessAfter, DstPipelines);
 
-		BeginOps.Emplace(FOperation::BeginTransitionResource(Identity, PreviousState, NextState, Info.Flags, nullptr));
+		BeginOps.Emplace(FOperation::BeginTransitionResource(Identity, PreviousState, NextState, Info.Flags, CreateInfo.Flags, nullptr));
 		EndOps  .Emplace(FOperation::EndTransitionResource(Identity, PreviousState, NextState, nullptr));
 	}
 
@@ -682,9 +671,9 @@ void FValidationRHI::RHICreateTransition(FRHITransition* Transition, const FRHIT
 		for (FOperation& Op : EndOps            ) { Op.Data_EndTransition  .CreateBacktrace = Backtrace; }
 	}
 
-	for (const FFenceEdge& FenceEdge : Fences)
+	for (FFence* Fence : Fences)
 	{
-		SignalOps[FenceEdge.SrcPipe].Emplace(FOperation::Signal(FenceEdge.Fence, FenceEdge.SrcPipe));
+		SignalOps[Fence->SrcPipe].Emplace(FOperation::Signal(Fence));
 	}
 
 	Transition->PendingSignals          = MoveTemp(SignalOps);
@@ -1449,6 +1438,28 @@ namespace RHIValidation
 			*GetRHIAccessName(ActualCurrentState.Access),
 			*GetRHIAccessName(CurrentStateFromRHI.Access));
 	}
+	
+	static inline FString GetReasonString_IncorrectFencing(
+		FResource* Resource, FSubresourceIndex const& SubresourceIndex,
+		ERHIPipeline SrcPipelineSkipped,
+		ERHIPipeline DstPipeline)
+	{
+		FString DebugName = GetResourceDebugName(Resource, SubresourceIndex);
+		FString SrcPipelineName = *GetRHIPipelineName(SrcPipelineSkipped);
+		FString DstPipelineName = *GetRHIPipelineName(DstPipeline);
+		return FString::Printf(
+			BARRIER_TRACKER_LOG_PREFIX_RESNAME
+			TEXT("Attemped to begin a resource transition for resource %s on the %s pipeline but skipping the transition on the %s pipeline (which is allowed with the NoFence flag), however no external\n")
+			TEXT("fence was issued between these two pipelines between this begin transition and the last end transition call on the %s pipeline. You must insert a manual fence from '%s' to '%s'.\n")
+			BARRIER_TRACKER_LOG_SUFFIX,
+			*DebugName,
+			*DebugName,
+			*DstPipelineName,
+			*SrcPipelineName,
+			*SrcPipelineName,
+			*SrcPipelineName,
+			*DstPipelineName);
+	}
 
 	static inline FString GetReasonString_IncorrectPreviousExplicitState(
 		FResource* Resource, FSubresourceIndex const& SubresourceIndex,
@@ -1695,7 +1706,7 @@ namespace RHIValidation
 		}
 	}
 
-	void FSubresourceState::BeginTransition(FResource* Resource, FSubresourceIndex const& SubresourceIndex, const FState& CurrentStateFromRHI, const FState& TargetState, EResourceTransitionFlags NewFlags, ERHIPipeline ExecutingPipeline, void* CreateTrace)
+	void FSubresourceState::BeginTransition(FResource* Resource, FSubresourceIndex const& SubresourceIndex, const FState& CurrentStateFromRHI, const FState& TargetState, EResourceTransitionFlags NewFlags, ERHITransitionCreateFlags CreateFlags, ERHIPipeline ExecutingPipeline, const TRHIPipelineArray<uint64>& PipelineMaxAwaitedFenceValues, void* CreateTrace)
 	{
 		FPipelineState& State = States[ExecutingPipeline];
 
@@ -1724,6 +1735,16 @@ namespace RHIValidation
 			}
 		}
 
+		// If we are collapsing multiple pipes to one pipe (only allowed when not fencing), check that the other pipes were fenced prior to this call.
+		if (EnumHasAnyFlags(CreateFlags, ERHITransitionCreateFlags::NoFence))
+		{
+			for (ERHIPipeline AlreadyFencedPipeline : MakeFlagsRange(State.Previous.Pipelines & ~CurrentStateFromRHI.Pipelines))
+			{
+				// The max awaited fence value should be higher than the last transitioned fence value, otherwise a fence was not issued.
+				RHI_VALIDATION_CHECK(LastTransitionFences[AlreadyFencedPipeline] < PipelineMaxAwaitedFenceValues[AlreadyFencedPipeline], *GetReasonString_IncorrectFencing(Resource, SubresourceIndex, AlreadyFencedPipeline, ExecutingPipeline));
+			}
+		}
+
 		// Check we're not already transitioning
 		RHI_VALIDATION_CHECK(!State.bTransitioning, *GetReasonString_DuplicateBeginTransition(Resource, SubresourceIndex, State.Current, TargetState, State.CreateTransitionBacktrace, BeginTrace));
 
@@ -1732,15 +1753,25 @@ namespace RHIValidation
 			// Check for the correct pipeline
 			RHI_VALIDATION_CHECK(EnumHasAllFlags(CurrentStateFromRHI.Pipelines, ExecutingPipeline), *GetReasonString_WrongPipeline(Resource, SubresourceIndex, State.Current, TargetState));
 
+			const auto HasMatchingPipelines = [CreateFlags] (ERHIPipeline Previous, ERHIPipeline Next)
+			{
+				// If no fence is being issued we only need to validate that the transition is happening from one of the previous pipes.
+				if (EnumHasAnyFlags(CreateFlags, ERHITransitionCreateFlags::NoFence))
+				{
+					return EnumHasAllFlags(Previous, Next);
+				}
+				return Previous == Next;
+			};
+
 			if (CurrentStateFromRHI.Access == ERHIAccess::Unknown)
 			{
-				RHI_VALIDATION_CHECK(Resource->TrackedAccess == State.Previous.Access && CurrentStateFromRHI.Pipelines == State.Previous.Pipelines,
+				RHI_VALIDATION_CHECK(Resource->TrackedAccess == State.Previous.Access && HasMatchingPipelines(State.Previous.Pipelines, CurrentStateFromRHI.Pipelines),
 					*GetReasonString_IncorrectPreviousTrackedState(Resource, SubresourceIndex, State.Previous, CurrentStateFromRHI.Pipelines));
 			}
 			else
 			{
 				// Check the current RHI state passed in matches the tracked state for the resource.
-				RHI_VALIDATION_CHECK(CurrentStateFromRHI.Access == State.Previous.Access && CurrentStateFromRHI.Pipelines == State.Previous.Pipelines,
+				RHI_VALIDATION_CHECK(CurrentStateFromRHI.Access == State.Previous.Access && HasMatchingPipelines(State.Previous.Pipelines, CurrentStateFromRHI.Pipelines),
 					*GetReasonString_IncorrectPreviousExplicitState(Resource, SubresourceIndex, State.Previous, CurrentStateFromRHI));
 			}
 		}
@@ -1760,16 +1791,13 @@ namespace RHIValidation
 		State.bTransitioning = true;
 
 		// Replicate the state to other pipes that are not part of the begin pipe mask.
-		for (ERHIPipeline OtherPipeline : MakeFlagsRange(ERHIPipeline::All))
+		for (ERHIPipeline OtherPipeline : MakeFlagsRange(ERHIPipeline::All & ~CurrentStateFromRHI.Pipelines))
 		{
-			if (!EnumHasAnyFlags(CurrentStateFromRHI.Pipelines, OtherPipeline))
-			{
-				States[OtherPipeline] = State;
-			}
+			States[OtherPipeline] = State;
 		}
 	}
 
-	void FSubresourceState::EndTransition(FResource* Resource, FSubresourceIndex const& SubresourceIndex, const FState& CurrentStateFromRHI, const FState& TargetState, ERHIPipeline ExecutingPipeline, void* CreateTrace)
+	void FSubresourceState::EndTransition(FResource* Resource, FSubresourceIndex const& SubresourceIndex, const FState& CurrentStateFromRHI, const FState& TargetState, ERHIPipeline ExecutingPipeline, uint64 ExecutingPipelineFenceValue, void* CreateTrace)
 	{
 		if (Resource->LoggingMode != ELoggingMode::None
 #if LOG_UNNAMED_RESOURCES
@@ -1802,6 +1830,8 @@ namespace RHIValidation
 				States[OtherPipeline] = State;
 			}
 		}
+
+		LastTransitionFences[ExecutingPipeline] = ExecutingPipelineFenceValue;
 	}
 
 	void FSubresourceState::Assert(FResource* Resource, FSubresourceIndex const& SubresourceIndex, const FState& RequiredState, bool bAllowAllUAVsOverlap)
@@ -2080,7 +2110,9 @@ namespace RHIValidation
 					Data_BeginTransition.PreviousState,
 					Data_BeginTransition.NextState,
 					Data_BeginTransition.Flags,
+					Data_BeginTransition.CreateFlags,
 					Queue.Pipeline,
+					Queue.MaxAwaitedFenceValues,
 					Data_BeginTransition.CreateBacktrace);
 
 			}, true);
@@ -2096,6 +2128,7 @@ namespace RHIValidation
 					Data_EndTransition.PreviousState,
 					Data_EndTransition.NextState,
 					Queue.Pipeline,
+					Queue.FenceValue,
 					Data_EndTransition.CreateBacktrace);
 			});
 			Data_EndTransition.Identity.Resource->ReleaseOpRef();
@@ -2143,17 +2176,25 @@ namespace RHIValidation
 			break;
 
 		case EOpType::Signal:
-			check(Data_Signal.Pipeline == Queue.Pipeline);
+			check(Data_Signal.Fence->SrcPipe == Queue.Pipeline);
 			Data_Signal.Fence->bSignaled = true;
+			Data_Signal.Fence->FenceValue = ++Queue.FenceValue;
 			break;
 
 		case EOpType::Wait:
-			check(Data_Wait.Pipeline == Queue.Pipeline);
-			if (!Data_Wait.Fence->bSignaled)
-				return false;
+			{
+				FFence* Fence = Data_Wait.Fence;
+				check(Fence->DstPipe == Queue.Pipeline);
+				if (!Fence->bSignaled)
+				{
+					return false;
+				}
+
+				Queue.MaxAwaitedFenceValues[Fence->SrcPipe] = FMath::Max(Fence->FenceValue, Queue.MaxAwaitedFenceValues[Fence->SrcPipe]);
 
-			// The fence has been completed. Free it now.
-			delete Data_Wait.Fence;
+				// The fence has been completed. Free it now.
+				delete Fence;
+			}
 			break;
 
 		case EOpType::AllUAVsOverlap:
@@ -2317,7 +2358,7 @@ namespace RHIValidation
 		return Backtrace;
 	}
 
-	bool ValidateDimension(EShaderCodeResourceBindingType Type, FRHIViewDesc::EDimension Dimension, bool SRV)
+	bool ValidateDimension(EShaderCodeResourceBindingType Type, FRHIViewDesc::EDimension Dimension, ERHITexturePlane TexturePlane, bool SRV)
 	{
 		// Ignore invalid types
 		if (Type == EShaderCodeResourceBindingType::Invalid)
@@ -2330,23 +2371,42 @@ namespace RHIValidation
 			return false;
 		}
 
+		if (Type == EShaderCodeResourceBindingType::RWStructuredBuffer || Type == EShaderCodeResourceBindingType::StructuredBuffer)
+		{
+			return TexturePlane == ERHITexturePlane::HTile;
+		}
+
+		if (Type == EShaderCodeResourceBindingType::RWByteAddressBuffer || Type == EShaderCodeResourceBindingType::ByteAddressBuffer)
+		{
+			return TexturePlane == ERHITexturePlane::CMask;
+		}
+		
+		if (Type == EShaderCodeResourceBindingType::RWBuffer || Type == EShaderCodeResourceBindingType::Buffer)
+		{
+			return TexturePlane == ERHITexturePlane::PrimaryCompressed || TexturePlane == ERHITexturePlane::CMask;
+		}
+
 		if (Type == EShaderCodeResourceBindingType::Texture2D || Type == EShaderCodeResourceBindingType::RWTexture2D || Type == EShaderCodeResourceBindingType::Texture2DMS)
 		{
 			return Dimension == FRHIViewDesc::EDimension::Texture2D;
 		}
-		else if (Type == EShaderCodeResourceBindingType::Texture2DArray || Type == EShaderCodeResourceBindingType::RWTexture2DArray)
+
+		if (Type == EShaderCodeResourceBindingType::Texture2DArray || Type == EShaderCodeResourceBindingType::RWTexture2DArray)
 		{
-			return Dimension == FRHIViewDesc::EDimension::Texture2DArray;
+			return Dimension == FRHIViewDesc::EDimension::Texture2DArray || Dimension == FRHIViewDesc::EDimension::TextureCube;
 		}
-		else if (Type == EShaderCodeResourceBindingType::Texture3D || Type == EShaderCodeResourceBindingType::RWTexture3D)
+
+		if (Type == EShaderCodeResourceBindingType::Texture3D || Type == EShaderCodeResourceBindingType::RWTexture3D)
 		{
 			return Dimension == FRHIViewDesc::EDimension::Texture3D;
 		}
-		else if (Type == EShaderCodeResourceBindingType::TextureCube || Type == EShaderCodeResourceBindingType::RWTextureCube)
+
+		if (Type == EShaderCodeResourceBindingType::TextureCube || Type == EShaderCodeResourceBindingType::RWTextureCube)
 		{
 			return Dimension == FRHIViewDesc::EDimension::TextureCube;
 		}
-		else if (Type == EShaderCodeResourceBindingType::TextureCubeArray)
+
+		if (Type == EShaderCodeResourceBindingType::TextureCubeArray)
 		{
 			return Dimension == FRHIViewDesc::EDimension::TextureCubeArray;
 		}
@@ -2366,19 +2426,23 @@ namespace RHIValidation
 		{
 			return Dimension == ETextureDimension::Texture2D;
 		}
-		else if (Type == EShaderCodeResourceBindingType::Texture2DArray || Type == EShaderCodeResourceBindingType::RWTexture2DArray)
+
+		if (Type == EShaderCodeResourceBindingType::Texture2DArray || Type == EShaderCodeResourceBindingType::RWTexture2DArray)
 		{
-			return Dimension == ETextureDimension::Texture2DArray;
+			return Dimension == ETextureDimension::Texture2DArray || Dimension == ETextureDimension::TextureCube;
 		}
-		else if (Type == EShaderCodeResourceBindingType::Texture3D || Type == EShaderCodeResourceBindingType::RWTexture3D)
+
+		if (Type == EShaderCodeResourceBindingType::Texture3D || Type == EShaderCodeResourceBindingType::RWTexture3D)
 		{
 			return Dimension == ETextureDimension::Texture3D;
 		}
-		else if (Type == EShaderCodeResourceBindingType::TextureCube || Type == EShaderCodeResourceBindingType::RWTextureCube)
+
+		if (Type == EShaderCodeResourceBindingType::TextureCube || Type == EShaderCodeResourceBindingType::RWTextureCube)
 		{
 			return Dimension == ETextureDimension::TextureCube;
 		}
-		else if (Type == EShaderCodeResourceBindingType::TextureCubeArray)
+
+		if (Type == EShaderCodeResourceBindingType::TextureCubeArray)
 		{
 			return Dimension == ETextureDimension::TextureCubeArray;
 		}
@@ -2405,7 +2469,7 @@ namespace RHIValidation
 		}
 		else if (Type == EShaderCodeResourceBindingType::StructuredBuffer || Type == EShaderCodeResourceBindingType::RWStructuredBuffer)
 		{
-			return BufferType == FRHIViewDesc::EBufferType::Structured;
+			return BufferType == FRHIViewDesc::EBufferType::Structured || BufferType == FRHIViewDesc::EBufferType::AccelerationStructure;
 		}
 		else if (Type == EShaderCodeResourceBindingType::Buffer || Type == EShaderCodeResourceBindingType::RWBuffer)
 		{
@@ -2451,7 +2515,7 @@ namespace RHIValidation
 			{
 				FString SRVName = GetSRVName(SRV, ViewIdentity);
 				uint16 ExpectedStride = RHIShaderBase->DebugStrideValidationData[FoundIndex].Stride;
-				if (ExpectedStride != SRVValidationStride.Stride)
+				if (ExpectedStride != SRVValidationStride.Stride && SRV->GetDesc().Buffer.SRV.BufferType != FRHIViewDesc::EBufferType::AccelerationStructure)
 				{
 					
 					FString ErrorMessage = FString::Printf(TEXT("Shader %s: Buffer stride for \"%s\" must match structure size declared in the shader"), RHIShaderBase->GetShaderName(), *SRVName);
@@ -2475,7 +2539,7 @@ namespace RHIValidation
 
 				if (SRV->IsTexture())
 				{
-					if (!ValidateDimension(ExpectedType, SRV->GetDesc().Texture.SRV.Dimension, true))
+					if (!ValidateDimension(ExpectedType, SRV->GetDesc().Texture.SRV.Dimension, SRV->GetDesc().Texture.SRV.Plane, true))
 					{
 						FString SRVName = GetSRVName(SRV, ViewIdentity);
 						FString ErrorMessage = FString::Printf(TEXT("Shader %s: Dimension for SRV \"%s\" must match type declared in the shader"), RHIShaderBase->GetShaderName(), *SRVName);
@@ -2604,7 +2668,7 @@ namespace RHIValidation
 
 				if (UAV->IsTexture())
 				{
-					if (!ValidateDimension(ExpectedType, UAV->GetDesc().Texture.UAV.Dimension, false))
+					if (!ValidateDimension(ExpectedType, UAV->GetDesc().Texture.UAV.Dimension, UAV->GetDesc().Texture.UAV.Plane, false))
 					{
 						FString UAVName = GetUAVName(UAV, ViewIdentity);
 						FString ErrorMessage = FString::Printf(TEXT("Shader %s: Dimension for UAV \"%s\" must match type declared in the shader"), RHIShaderBase->GetShaderName(), *UAVName);
diff --git a/Engine/Source/Runtime/RHI/Public/RHICommandList.h b/Engine/Source/Runtime/RHI/Public/RHICommandList.h
--- a/Engine/Source/Runtime/RHI/Public/RHICommandList.h
+++ b/Engine/Source/Runtime/RHI/Public/RHICommandList.h
@@ -2688,7 +2688,7 @@ public:
 
 			if (Bypass())
 			{
-				GetContext().RHISetShaderParameters(InShader, InBatchedParameters.ParametersData, InBatchedParameters.Parameters, InBatchedParameters.ResourceParameters, InBatchedParameters.BindlessParameters);
+				GetComputeContext().RHISetShaderParameters(InShader, InBatchedParameters.ParametersData, InBatchedParameters.Parameters, InBatchedParameters.ResourceParameters, InBatchedParameters.BindlessParameters);
 				return;
 			}
 
@@ -2838,7 +2838,7 @@ public:
 		}
 	}
 
-	RHI_API void Transition(TArrayView<const FRHITransitionInfo> Infos);
+	RHI_API void Transition(TArrayView<const FRHITransitionInfo> Infos, ERHITransitionCreateFlags CreateFlags = ERHITransitionCreateFlags::None);
 
 	FORCEINLINE_DEBUGGABLE void BeginTransition(const FRHITransition* Transition)
 	{
@@ -2850,11 +2850,17 @@ public:
 		EndTransitions(MakeArrayView(&Transition, 1));
 	}
 
-	FORCEINLINE_DEBUGGABLE void Transition(const FRHITransitionInfo& Info)
+	FORCEINLINE_DEBUGGABLE void Transition(const FRHITransitionInfo& Info, ERHITransitionCreateFlags CreateFlags = ERHITransitionCreateFlags::None)
 	{
-		Transition(MakeArrayView(&Info, 1));
+		Transition(MakeArrayView(&Info, 1), CreateFlags);
 	}
 
+	//
+	// Performs an immediate transition with the option of broadcasting to multiple pipelines.
+	// Uses both the immediate and async compute contexts. Falls back to graphics-only if async compute is not supported.
+	//
+	RHI_API void Transition(TArrayView<const FRHITransitionInfo> Infos, ERHIPipeline SrcPipelines, ERHIPipeline DstPipelines, ERHITransitionCreateFlags TransitionCreateFlags = ERHITransitionCreateFlags::None);
+
 	FORCEINLINE_DEBUGGABLE void SetTrackedAccess(TArrayView<const FRHITrackedAccessInfo> Infos)
 	{
 		bUsesSetTrackedAccess = true;
@@ -4554,13 +4560,6 @@ public:
 
 	RHI_API void InitializeImmediateContexts();
 
-	//
-	// Performs an immediate transition with the option of broadcasting to multiple pipelines.
-	// Uses both the immediate and async compute contexts. Falls back to graphics-only if async compute is not supported.
-	//
-	RHI_API void Transition(TArrayView<const FRHITransitionInfo> Infos, ERHIPipeline SrcPipelines, ERHIPipeline DstPipelines);
-	using FRHIComputeCommandList::Transition;
-
 	template <typename LAMBDA>
 	FORCEINLINE_DEBUGGABLE void EnqueueLambda(const TCHAR* LambdaName, LAMBDA&& Lambda)
 	{
diff --git a/Engine/Source/Runtime/RHI/Public/RHIValidationCommon.h b/Engine/Source/Runtime/RHI/Public/RHIValidationCommon.h
--- a/Engine/Source/Runtime/RHI/Public/RHIValidationCommon.h
+++ b/Engine/Source/Runtime/RHI/Public/RHIValidationCommon.h
@@ -186,10 +186,11 @@ namespace RHIValidation
 			void* BeginTransitionBacktrace = nullptr;
 		};
 
+		TRHIPipelineArray<uint64> LastTransitionFences{InPlace, 0};
 		TRHIPipelineArray<FPipelineState> States;
 
-		void BeginTransition   (FResource* Resource, FSubresourceIndex const& SubresourceIndex, const FState& CurrentStateFromRHI, const FState& TargetState, EResourceTransitionFlags NewFlags, ERHIPipeline Pipeline, void* CreateTrace);
-		void EndTransition     (FResource* Resource, FSubresourceIndex const& SubresourceIndex, const FState& CurrentStateFromRHI, const FState& TargetState, ERHIPipeline Pipeline, void* CreateTrace);
+		void BeginTransition   (FResource* Resource, FSubresourceIndex const& SubresourceIndex, const FState& CurrentStateFromRHI, const FState& TargetState, EResourceTransitionFlags NewFlags, ERHITransitionCreateFlags CreateFlags, ERHIPipeline Pipeline, const TRHIPipelineArray<uint64>& PipelineMaxAwaitedFenceValues, void* CreateTrace);
+		void EndTransition     (FResource* Resource, FSubresourceIndex const& SubresourceIndex, const FState& CurrentStateFromRHI, const FState& TargetState, ERHIPipeline Pipeline, uint64 PipelineFenceValue, void* CreateTrace);
 		void Assert            (FResource* Resource, FSubresourceIndex const& SubresourceIndex, const FState& RequiredState, bool bAllowAllUAVsOverlap);
 		void AssertTracked     (FResource* Resource, FSubresourceIndex const& SubresourceIndex, const FState& TrackedState);
 		void SpecificUAVOverlap(FResource* Resource, FSubresourceIndex const& SubresourceIndex, ERHIPipeline Pipeline, bool bAllow);
@@ -505,6 +506,9 @@ namespace RHIValidation
 	struct FFence
 	{
 		bool bSignaled = false;
+		ERHIPipeline SrcPipe = ERHIPipeline::None;
+		ERHIPipeline DstPipe = ERHIPipeline::None;
+		uint64 FenceValue = 0;
 	};
 
 	enum class EOpType
@@ -555,6 +559,7 @@ namespace RHIValidation
 				FState PreviousState;
 				FState NextState;
 				EResourceTransitionFlags Flags;
+				ERHITransitionCreateFlags CreateFlags;
 				void* CreateBacktrace;
 			} Data_BeginTransition;
 
@@ -607,13 +612,11 @@ namespace RHIValidation
 			struct
 			{
 				FFence* Fence;
-				ERHIPipeline Pipeline;
 			} Data_Signal;
 
 			struct
 			{
 				FFence* Fence;
-				ERHIPipeline Pipeline;
 			} Data_Wait;
 
 			struct
@@ -643,7 +646,7 @@ namespace RHIValidation
 		// Returns true if the operation is complete
 		RHI_API bool Replay(FOpQueueState& Queue) const;
 
-		static inline FOperation BeginTransitionResource(FResourceIdentity Identity, FState PreviousState, FState NextState, EResourceTransitionFlags Flags, void* CreateBacktrace)
+		static inline FOperation BeginTransitionResource(FResourceIdentity Identity, FState PreviousState, FState NextState, EResourceTransitionFlags Flags, ERHITransitionCreateFlags CreateFlags, void* CreateBacktrace)
 		{
 			for (ERHIPipeline Pipeline : MakeFlagsRange(PreviousState.Pipelines))
 			{
@@ -656,6 +659,7 @@ namespace RHIValidation
 			Op.Data_BeginTransition.PreviousState = PreviousState;
 			Op.Data_BeginTransition.NextState = NextState;
 			Op.Data_BeginTransition.Flags = Flags;
+			Op.Data_BeginTransition.CreateFlags = CreateFlags;
 			Op.Data_BeginTransition.CreateBacktrace = CreateBacktrace;
 			return MoveTemp(Op);
 		}
@@ -746,21 +750,19 @@ namespace RHIValidation
 			return MoveTemp(Op);
 		}
 
-		static inline FOperation Signal(FFence* Fence, ERHIPipeline Pipeline)
+		static inline FOperation Signal(FFence* Fence)
 		{
 			FOperation Op;
 			Op.Type = EOpType::Signal;
 			Op.Data_Signal.Fence = Fence;
-			Op.Data_Signal.Pipeline = Pipeline;
 			return MoveTemp(Op);
 		}
 
-		static inline FOperation Wait(FFence* Fence, ERHIPipeline Pipeline)
+		static inline FOperation Wait(FFence* Fence)
 		{
 			FOperation Op;
 			Op.Type = EOpType::Wait;
 			Op.Data_Wait.Fence = Fence;
-			Op.Data_Wait.Pipeline = Pipeline;
 			return MoveTemp(Op);
 		}
 
@@ -848,6 +850,8 @@ namespace RHIValidation
 	struct FOpQueueState
 	{
 		ERHIPipeline const Pipeline;
+		uint64 FenceValue = 0;
+		TRHIPipelineArray<uint64> MaxAwaitedFenceValues{InPlace, 0};
 
 #if WITH_RHI_BREADCRUMBS
 		struct
@@ -876,7 +880,7 @@ namespace RHIValidation
 		{}
 
 		void AppendOps(FValidationCommandList* CommandList);
-		
+
 		// Returns true if progress was made
 		bool Execute();
 	};
@@ -974,7 +978,7 @@ namespace RHIValidation
 		{
 			// This function exists due to the implicit transitions that RHI functions make (e.g. RHICopyToResolveTarget).
 			// It should be removed when we eventually remove all implicit transitions from the RHI.
-			AddOp(FOperation::BeginTransitionResource(Identity, PreviousState, NextState, Flags, nullptr));
+			AddOp(FOperation::BeginTransitionResource(Identity, PreviousState, NextState, Flags, ERHITransitionCreateFlags::None, nullptr));
 			AddOp(FOperation::EndTransitionResource(Identity, PreviousState, NextState, nullptr));
 		}
 
diff --git a/Engine/Source/Runtime/RenderCore/Private/RenderGraphBuilder.cpp b/Engine/Source/Runtime/RenderCore/Private/RenderGraphBuilder.cpp
--- a/Engine/Source/Runtime/RenderCore/Private/RenderGraphBuilder.cpp
+++ b/Engine/Source/Runtime/RenderCore/Private/RenderGraphBuilder.cpp
@@ -553,6 +553,12 @@ FRDGBuilder::FRDGBuilder(FRHICommandListImmediate& InRHICmdList, FRDGEventName I
 	, BarrierValidation(&Passes, BuilderName)
 #endif
 {
+	if (GSupportsEfficientAsyncCompute)
+	{
+		// Insert a manual fence from async compute to graphics to synchronize any all pipeline external access resources from the last run.
+		RHICmdList.Transition({}, ERHIPipeline::AsyncCompute, ERHIPipeline::Graphics);
+	}
+
 	ProloguePass = SetupEmptyPass(Passes.Allocate<FRDGSentinelPass>(Allocators.Root, RDG_EVENT_NAME("Graph Prologue (Graphics)")));
 
 	ParallelExecute.bEnabled = ::IsParallelExecuteEnabled() && EnumHasAnyFlags(InFlags, ERDGBuilderFlags::AllowParallelExecute);
@@ -1638,14 +1644,7 @@ void FRDGBuilder::Execute()
 
 	GRDGTransientResourceAllocator.ReleasePendingDeallocations();
 
-	{
-		SCOPED_NAMED_EVENT_TEXT("FRDGBuilder::FlushAccessModeQueue", FColor::Magenta);
-		for (FRDGViewableResource* Resource : ExternalAccessResources)
-		{
-			UseInternalAccessMode(Resource);
-		}
-		FlushAccessModeQueue();
-	}
+	FlushAccessModeQueue();
 
 	// Create the epilogue pass at the end of the graph just prior to compilation.
 	SetupEmptyPass(EpiloguePass = Passes.Allocate<FRDGSentinelPass>(Allocators.Root, RDG_EVENT_NAME("Graph Epilogue")));
@@ -3642,6 +3641,13 @@ void FRDGBuilder::AddLastTextureTransition(FRDGTexture* Texture)
 	check(IsImmediateMode() || Texture->bExtracted || Texture->ReferenceCount == FRDGViewableResource::DeallocatedReferenceCount);
 	check(Texture->HasRHI());
 
+	if (Texture->AccessModeState.ActiveMode == FRDGViewableResource::EAccessMode::External)
+	{
+		// Assign the final state that was enqueued by the external access pass, which may include merged states.
+		EpilogueResourceAccesses.Emplace(Texture->GetRHI(), Texture->State[0]->Access);
+		return;
+	}
+
 	const FRDGPassHandle EpiloguePassHandle = GetEpiloguePassHandle();
 
 	FRDGSubresourceState* SubresourceStateBefore = nullptr;
@@ -3724,6 +3730,13 @@ void FRDGBuilder::AddLastBufferTransition(FRDGBuffer* Buffer)
 	check(IsImmediateMode() || Buffer->bExtracted || Buffer->ReferenceCount == FRDGViewableResource::DeallocatedReferenceCount);
 	check(Buffer->HasRHI());
 
+	if (Buffer->AccessModeState.ActiveMode == FRDGViewableResource::EAccessMode::External)
+	{
+		// Assign the final state that was enqueued by the external access pass, which may include merged states.
+		EpilogueResourceAccesses.Emplace(Buffer->GetRHI(), Buffer->State->Access);
+		return;
+	}
+
 	const FRDGPassHandle EpiloguePassHandle = GetEpiloguePassHandle();
 
 	FRDGSubresourceState* StateAfter = AllocSubresource();
diff --git a/Engine/Source/Runtime/Renderer/Private/CompositionLighting/PostProcessAmbientOcclusion.cpp b/Engine/Source/Runtime/Renderer/Private/CompositionLighting/PostProcessAmbientOcclusion.cpp
--- a/Engine/Source/Runtime/Renderer/Private/CompositionLighting/PostProcessAmbientOcclusion.cpp
+++ b/Engine/Source/Runtime/Renderer/Private/CompositionLighting/PostProcessAmbientOcclusion.cpp
@@ -181,7 +181,7 @@ int32 FSSAOHelper::GetAmbientOcclusionShaderLevel(const FSceneView& View)
 
 bool FSSAOHelper::IsAmbientOcclusionCompute(const ERHIFeatureLevel::Type FeatureLevel)
 {
-	return FeatureLevel >= ERHIFeatureLevel::SM5 && CVarAmbientOcclusionCompute.GetValueOnRenderThread() >= 1;
+	return FeatureLevel >= ERHIFeatureLevel::SM5 && CVarAmbientOcclusionCompute.GetValueOnAnyThread() >= 1;
 }
 
 int32 FSSAOHelper::GetNumAmbientOcclusionLevels()
diff --git a/Engine/Source/Runtime/Renderer/Private/IESTextureManager.cpp b/Engine/Source/Runtime/Renderer/Private/IESTextureManager.cpp
--- a/Engine/Source/Runtime/Renderer/Private/IESTextureManager.cpp
+++ b/Engine/Source/Runtime/Renderer/Private/IESTextureManager.cpp
@@ -1,4 +1,4 @@
-﻿// Copyright Epic Games, Inc. All Rights Reserved.
+// Copyright Epic Games, Inc. All Rights Reserved.
 
 #include "IESTextureManager.h"
 #include "Shader.h"
@@ -267,8 +267,6 @@ static void AddSlotsPassCS(
 		const FIntVector DispatchCount = FComputeShaderUtils::GetGroupCount(FIntVector(Parameters->AtlasResolution.X, Parameters->AtlasResolution.Y, SlotCount), FIntVector(8, 8, 1));
 		FComputeShaderUtils::AddPass(GraphBuilder, RDG_EVENT_NAME("IESAtlas::AddTexture"), ComputeShader, Parameters, DispatchCount);
 	}
-
-	GraphBuilder.UseExternalAccessMode(OutAtlas, ERHIAccess::SRVMask);
 }
 
 class FIESAtlasAddTexturePS : public FGlobalShader
@@ -605,6 +603,12 @@ void UpdateAtlasTexture(FRDGBuilder& GraphBuilder, const FStaticShaderPlatform S
 		AddSlotsPass(GraphBuilder, ShaderPlatform, ShaderMap, RefreshSlots, AtlasTexture);
 		GIESTextureManager.bHasPendingRefreshes = false;
 	}
+
+	// If the atlas was registered into RDG make sure it gets transitioned to external access on all pipes.
+	if (FRDGTexture* AtlasTexture = GraphBuilder.FindExternalTexture(GIESTextureManager.AtlasTexture))
+	{
+		GraphBuilder.UseExternalAccessMode(AtlasTexture, ERHIAccess::SRVMask, ERHIPipeline::All);
+	}
 }
 
 void AddDebugPass(FRDGBuilder& GraphBuilder, const FViewInfo& View, FRDGTextureRef OutputTexture)
diff --git a/Engine/Source/Runtime/Renderer/Private/ReflectionEnvironmentCapture.cpp b/Engine/Source/Runtime/Renderer/Private/ReflectionEnvironmentCapture.cpp
--- a/Engine/Source/Runtime/Renderer/Private/ReflectionEnvironmentCapture.cpp
+++ b/Engine/Source/Runtime/Renderer/Private/ReflectionEnvironmentCapture.cpp
@@ -1573,6 +1573,7 @@ void CopyToSkyTexture(FRDGBuilder& GraphBuilder, FScene* Scene, FRDGTexture* Inp
 		CopyInfo.NumMips = GetNumMips(ProcessedTexture->GetSizeX());
 
 		AddCopyTexturePass(GraphBuilder, InputTexture, OutputTexture, CopyInfo);
+		GraphBuilder.UseExternalAccessMode(OutputTexture, ERHIAccess::SRVMask, ERHIPipeline::All);
 	}
 }
 
diff --git a/Engine/Source/Runtime/Renderer/Private/VT/VirtualTexturePhysicalSpace.cpp b/Engine/Source/Runtime/Renderer/Private/VT/VirtualTexturePhysicalSpace.cpp
--- a/Engine/Source/Runtime/Renderer/Private/VT/VirtualTexturePhysicalSpace.cpp
+++ b/Engine/Source/Runtime/Renderer/Private/VT/VirtualTexturePhysicalSpace.cpp
@@ -12,7 +12,7 @@
 #include "VT/VirtualTextureScalability.h"
 #include "VT/VirtualTextureSystem.h"
 #include "RHIUtilities.h"
-#include "RenderGraphBuilder.h"
+#include "RenderGraphUtils.h"
 
 DECLARE_MEMORY_STAT_POOL(TEXT("Total Physical Memory"), STAT_TotalPhysicalMemory, STATGROUP_VirtualTextureMemory, FPlatformMemory::MCR_GPU);
 
@@ -195,14 +195,14 @@ void FVirtualTexturePhysicalSpace::ReleaseRHI()
 	}
 }
 
-void FVirtualTexturePhysicalSpace::FinalizeTextures(FRDGBuilder& GraphBuilder)
+void FVirtualTexturePhysicalSpace::FinalizeTextures(FRDGBuilder& GraphBuilder, FRDGExternalAccessQueue& ExternalAccessQueue)
 {
 	for (int32 Layer = 0; Layer < Description.NumLayers; ++Layer)
 	{
 		// It's only necessary to enable external access mode on textures modified by RDG this frame.
 		if (FRDGTexture* Texture = GraphBuilder.FindExternalTexture(PooledRenderTarget[Layer]))
 		{
-			GraphBuilder.UseExternalAccessMode(Texture, ERHIAccess::SRVMask);
+			ExternalAccessQueue.Add(Texture, ERHIAccess::SRVMask);
 		}
 	}
 }
diff --git a/Engine/Source/Runtime/Renderer/Private/VT/VirtualTexturePhysicalSpace.h b/Engine/Source/Runtime/Renderer/Private/VT/VirtualTexturePhysicalSpace.h
--- a/Engine/Source/Runtime/Renderer/Private/VT/VirtualTexturePhysicalSpace.h
+++ b/Engine/Source/Runtime/Renderer/Private/VT/VirtualTexturePhysicalSpace.h
@@ -10,6 +10,8 @@
 #include "VirtualTextureShared.h"
 #include "VirtualTexturing.h"
 
+class FRDGExternalAccessQueue;
+
 /** Description that determines a unique physical space. This is filled out by a producer. */
 struct FVTPhysicalSpaceDescription
 {
@@ -136,7 +138,7 @@ public:
 		return PooledRenderTarget[Layer];
 	}
 
-	void FinalizeTextures(FRDGBuilder& GraphBuilder);
+	void FinalizeTextures(FRDGBuilder& GraphBuilder, FRDGExternalAccessQueue& ExternalAccessQueue);
 
 	/** Update internal tracking of residency. This is used to update stats and to calculate a mip bias to keep within the pool budget. */
 	void UpdateResidencyTracking(uint32 Frame);
diff --git a/Engine/Source/Runtime/Renderer/Private/VT/VirtualTextureSpace.cpp b/Engine/Source/Runtime/Renderer/Private/VT/VirtualTextureSpace.cpp
--- a/Engine/Source/Runtime/Renderer/Private/VT/VirtualTextureSpace.cpp
+++ b/Engine/Source/Runtime/Renderer/Private/VT/VirtualTextureSpace.cpp
@@ -365,11 +365,11 @@ void FVirtualTextureSpace::AllocateTextures(FRDGBuilder& GraphBuilder)
 }
 
 
-void FVirtualTextureSpace::ApplyUpdates(FVirtualTextureSystem* System, FRDGBuilder& GraphBuilder)
+void FVirtualTextureSpace::ApplyUpdates(FVirtualTextureSystem* System, FRDGBuilder& GraphBuilder, FRDGExternalAccessQueue& ExternalAccessQueue)
 {
 	ON_SCOPE_EXIT
 	{
-		FinalizeTextures(GraphBuilder);
+		FinalizeTextures(GraphBuilder, ExternalAccessQueue);
 	};
 
 	static TArray<FPageTableUpdate> ExpandedUpdates[VIRTUALTEXTURE_SPACE_MAXLAYERS][16];
@@ -566,7 +566,7 @@ void FVirtualTextureSpace::ApplyUpdates(FVirtualTextureSystem* System, FRDGBuild
 	}
 }
 
-void FVirtualTextureSpace::FinalizeTextures(FRDGBuilder& GraphBuilder)
+void FVirtualTextureSpace::FinalizeTextures(FRDGBuilder& GraphBuilder, FRDGExternalAccessQueue& ExternalAccessQueue)
 {
 	for (uint32 LayerIndex = 0u; LayerIndex < Description.NumPageTableLayers; ++LayerIndex)
 	{
@@ -577,7 +577,7 @@ void FVirtualTextureSpace::FinalizeTextures(FRDGBuilder& GraphBuilder)
 			// It's only necessary to enable external access mode on textures modified by RDG this frame.
 			if (FRDGTexture* Texture = GraphBuilder.FindExternalTexture(PageTableEntry.RenderTarget))
 			{
-				GraphBuilder.UseExternalAccessMode(Texture, ERHIAccess::SRVMask);
+				ExternalAccessQueue.Add(Texture, ERHIAccess::SRVMask, ERHIPipeline::All);
 			}
 		}
 	}
diff --git a/Engine/Source/Runtime/Renderer/Private/VT/VirtualTextureSpace.h b/Engine/Source/Runtime/Renderer/Private/VT/VirtualTextureSpace.h
--- a/Engine/Source/Runtime/Renderer/Private/VT/VirtualTextureSpace.h
+++ b/Engine/Source/Runtime/Renderer/Private/VT/VirtualTextureSpace.h
@@ -92,8 +92,8 @@ public:
 
 	void				QueueUpdate( uint8 Layer, uint8 vLogSize, uint32 vAddress, uint8 vLevel, const FPhysicalTileLocation& pTileLocation);
 	void				AllocateTextures(FRDGBuilder& GraphBuilder);
-	void				FinalizeTextures(FRDGBuilder& GraphBuilder);
-	void				ApplyUpdates(FVirtualTextureSystem* System, FRDGBuilder& GraphBuilder);
+	void				FinalizeTextures(FRDGBuilder& GraphBuilder, FRDGExternalAccessQueue& ExternalAccessQueue);
+	void				ApplyUpdates(FVirtualTextureSystem* System, FRDGBuilder& GraphBuilder, FRDGExternalAccessQueue& ExternalAccessQueue);
 	void				QueueUpdateEntirePageTable();
 
 	void DumpToConsole(bool verbose);
diff --git a/Engine/Source/Runtime/Renderer/Private/VT/VirtualTextureSystem.cpp b/Engine/Source/Runtime/Renderer/Private/VT/VirtualTextureSystem.cpp
--- a/Engine/Source/Runtime/Renderer/Private/VT/VirtualTextureSystem.cpp
+++ b/Engine/Source/Runtime/Renderer/Private/VT/VirtualTextureSystem.cpp
@@ -2440,6 +2440,8 @@ void FVirtualTextureSystem::FinalizeRequests(FRDGBuilder& GraphBuilder)
 		Finalizers.Reset();
 	}
 
+	FRDGExternalAccessQueue ExternalAccessQueue;
+
 	// Update page tables
 	{
 		SCOPE_CYCLE_COUNTER(STAT_PageTableUpdates);
@@ -2450,7 +2452,7 @@ void FVirtualTextureSystem::FinalizeRequests(FRDGBuilder& GraphBuilder)
 		{
 			if (Spaces[ID])
 			{
-				Spaces[ID]->ApplyUpdates(this, GraphBuilder);
+				Spaces[ID]->ApplyUpdates(this, GraphBuilder, ExternalAccessQueue);
 			}
 		}
 	}
@@ -2459,10 +2461,12 @@ void FVirtualTextureSystem::FinalizeRequests(FRDGBuilder& GraphBuilder)
 	{
 		if (PhysicalSpace != nullptr)
 		{
-			PhysicalSpace->FinalizeTextures(GraphBuilder);
+			PhysicalSpace->FinalizeTextures(GraphBuilder, ExternalAccessQueue);
 		}
 	}
 
+	ExternalAccessQueue.Submit(GraphBuilder);
+
 	Frame++;
 }
 
@@ -2712,14 +2716,17 @@ TUniquePtr<FVirtualTextureUpdater> FVirtualTextureSystem::BeginUpdate(FRDGBuilde
 
 	if (!Settings.bEnablePageRequests)
 	{
+		FRDGExternalAccessQueue ExternalAccessQueue;
+
 		for (uint32 ID = 0; ID < MaxSpaces; ID++)
 		{
 			if (Spaces[ID])
 			{
-				Spaces[ID]->FinalizeTextures(GraphBuilder);
+				Spaces[ID]->FinalizeTextures(GraphBuilder, ExternalAccessQueue);
 			}
 		}
 
+		ExternalAccessQueue.Submit(GraphBuilder);
 		return {};
 	}
 
diff --git a/Engine/Source/Runtime/Renderer/Private/VirtualShadowMaps/VirtualShadowMapArray.cpp b/Engine/Source/Runtime/Renderer/Private/VirtualShadowMaps/VirtualShadowMapArray.cpp
--- a/Engine/Source/Runtime/Renderer/Private/VirtualShadowMaps/VirtualShadowMapArray.cpp
+++ b/Engine/Source/Runtime/Renderer/Private/VirtualShadowMaps/VirtualShadowMapArray.cpp
@@ -3539,8 +3539,11 @@ FScreenPassTexture FVirtualShadowMapArray::AddVisualizePass(FRDGBuilder& GraphBu
 			TShaderMapRef<FDesaturatePS> DesaturatePixelShader(View.ShaderMap);
 			TShaderMapRef<FCopyRectPS> PixelShader(View.ShaderMap);
 
+			FRDGTextureRef SceneColorCopy = GraphBuilder.CreateTexture(SceneColor.Texture->Desc, TEXT("SceneColorCopy"));
+			AddCopyTexturePass(GraphBuilder, SceneColor.Texture, SceneColorCopy);
+
 			FDesaturatePS::FParameters* Parameters = GraphBuilder.AllocParameters<FDesaturatePS::FParameters>();
-			Parameters->InputTexture = SceneColor.Texture;
+			Parameters->InputTexture = SceneColorCopy;
 			Parameters->InputSampler = TStaticSamplerState<SF_Bilinear, AM_Clamp, AM_Clamp, AM_Clamp>::GetRHI();
 			Parameters->RenderTargets[0] = FRenderTargetBinding(Output.Texture, ERenderTargetLoadAction::ENoAction);
 			AddDrawScreenPass(GraphBuilder, RDG_EVENT_NAME("Desaturate"), View, OutputViewport, OutputViewport, VertexShader, DesaturatePixelShader, Parameters, EScreenPassDrawFlags::None);
