Remove RHIBeginScene and RHIEndScene
 - None of the platform RHIs provide an implementation of this function any more.

diff --git a/Engine/Plugins/FX/Niagara/Source/Niagara/Private/NiagaraGpuComputeDispatch.cpp b/Engine/Plugins/FX/Niagara/Source/Niagara/Private/NiagaraGpuComputeDispatch.cpp
--- a/Engine/Plugins/FX/Niagara/Source/Niagara/Private/NiagaraGpuComputeDispatch.cpp
+++ b/Engine/Plugins/FX/Niagara/Source/Niagara/Private/NiagaraGpuComputeDispatch.cpp
@@ -432,10 +432,6 @@ void FNiagaraGpuComputeDispatch::ProcessPendingTicksFlush(FRHICommandListImmedia
 			TConstStridedView<FSceneView> DummyViews = MakeStridedView<const FSceneView>(0, ViewFamily.Views[0], 1);
 			const bool bAllowGPUParticleUpdate = true;
 
-			// Notify that we are about to begin rendering the 'scene' this is required because some RHIs will ClearState
-			// in the event of submitting commands, i.e. when we write a fence, or indeed perform a manual flush.
-			RHICmdList.BeginScene();
-
 			// Ensure system textures are initialized
 			GetRendererModule().InitializeSystemTextures(RHICmdList);
 
@@ -469,9 +465,6 @@ void FNiagaraGpuComputeDispatch::ProcessPendingTicksFlush(FRHICommandListImmedia
 			MaxTicksToFlush = TNumericLimits<int32>::Max();
 
 			bIsOutsideSceneRenderer = false;
-
-			// We have completed flushing the commands
-			RHICmdList.EndScene();
 			break;
 		}
 
diff --git a/Engine/Source/Runtime/Apple/MetalRHI/Private/MetalRHI.cpp b/Engine/Source/Runtime/Apple/MetalRHI/Private/MetalRHI.cpp
--- a/Engine/Source/Runtime/Apple/MetalRHI/Private/MetalRHI.cpp
+++ b/Engine/Source/Runtime/Apple/MetalRHI/Private/MetalRHI.cpp
@@ -1242,14 +1242,6 @@ void FMetalRHICommandContext::RHIEndFrame()
     ((FMetalDeviceContext*)Context)->EndFrame();
 }
 
-void FMetalRHICommandContext::RHIBeginScene()
-{
-}
-
-void FMetalRHICommandContext::RHIEndScene()
-{
-}
-
 #if WITH_RHI_BREADCRUMBS
 	void FMetalRHICommandContext::RHIBeginBreadcrumbGPU(FRHIBreadcrumbNode* Breadcrumb)
 	{
diff --git a/Engine/Source/Runtime/Apple/MetalRHI/Public/MetalRHIContext.h b/Engine/Source/Runtime/Apple/MetalRHI/Public/MetalRHIContext.h
--- a/Engine/Source/Runtime/Apple/MetalRHI/Public/MetalRHIContext.h
+++ b/Engine/Source/Runtime/Apple/MetalRHI/Public/MetalRHIContext.h
@@ -74,12 +74,6 @@ public:
 	// This method is queued with an RHIThread, otherwise it will flush after it is queued; without an RHI thread there is no benefit to queuing this frame advance commands
 	virtual void RHIEndFrame() override;
 	
-	// This method is queued with an RHIThread, otherwise it will flush after it is queued; without an RHI thread there is no benefit to queuing this frame advance commands
-	virtual void RHIBeginScene() override;
-
-	// This method is queued with an RHIThread, otherwise it will flush after it is queued; without an RHI thread there is no benefit to queuing this frame advance commands
-	virtual void RHIEndScene() override;
-	
 	virtual void RHISetStreamSource(uint32 StreamIndex, FRHIBuffer* VertexBuffer, uint32 Offset) final override;
 
 	// @param MinX including like Win32 RECT
diff --git a/Engine/Source/Runtime/D3D12RHI/Private/D3D12CommandContext.cpp b/Engine/Source/Runtime/D3D12RHI/Private/D3D12CommandContext.cpp
--- a/Engine/Source/Runtime/D3D12RHI/Private/D3D12CommandContext.cpp
+++ b/Engine/Source/Runtime/D3D12RHI/Private/D3D12CommandContext.cpp
@@ -803,16 +803,6 @@ void FD3D12CommandContextBase::UpdateMemoryStats()
 #endif // PLATFORM_WINDOWS && (STATS || CSV_PROFILER)
 }
 
-void FD3D12CommandContext::RHIBeginScene()
-{
-	// Nothing to do
-}
-
-void FD3D12CommandContext::RHIEndScene()
-{
-	// Nothing to do
-}
-
 IRHIComputeContext* FD3D12DynamicRHI::RHIGetCommandContext(ERHIPipeline Pipeline, FRHIGPUMask GPUMask)
 {
 	if (GPUMask.HasSingleIndex())
diff --git a/Engine/Source/Runtime/D3D12RHI/Private/D3D12CommandContext.h b/Engine/Source/Runtime/D3D12RHI/Private/D3D12CommandContext.h
--- a/Engine/Source/Runtime/D3D12RHI/Private/D3D12CommandContext.h
+++ b/Engine/Source/Runtime/D3D12RHI/Private/D3D12CommandContext.h
@@ -656,8 +656,6 @@ public:
 #if (RHI_NEW_GPU_PROFILER == 0)
 	virtual void RHICalibrateTimers(FRHITimestampCalibrationQuery* CalibrationQuery) final override;
 #endif
-	virtual void RHIBeginScene() final override;
-	virtual void RHIEndScene() final override;
 	virtual void RHISetStreamSource(uint32 StreamIndex, FRHIBuffer* VertexBuffer, uint32 Offset) final override;
 	virtual void RHISetViewport(float MinX, float MinY, float MinZ, float MaxX, float MaxY, float MaxZ) final override;
 	virtual void RHISetStereoViewport(float LeftMinX, float RightMinX, float LeftMinY, float RightMinY, float MinZ, float LeftMaxX, float RightMaxX, float LeftMaxY, float RightMaxY, float MaxZ) override;
@@ -957,14 +955,6 @@ public:
 		ContextRedirect(RHICalibrateTimers(CalibrationQuery));
 	}
 #endif
-	FORCEINLINE virtual void RHIBeginScene() final override
-	{
-		ContextRedirect(RHIBeginScene());
-	}
-	FORCEINLINE virtual void RHIEndScene() final override
-	{
-		ContextRedirect(RHIEndScene());
-	}
 	FORCEINLINE virtual void RHISetStreamSource(uint32 StreamIndex, FRHIBuffer* VertexBuffer, uint32 Offset) final override
 	{
 		ContextRedirect(RHISetStreamSource(StreamIndex, VertexBuffer, Offset));
diff --git a/Engine/Source/Runtime/NullDrv/Public/NullRHI.h b/Engine/Source/Runtime/NullDrv/Public/NullRHI.h
--- a/Engine/Source/Runtime/NullDrv/Public/NullRHI.h
+++ b/Engine/Source/Runtime/NullDrv/Public/NullRHI.h
@@ -351,19 +351,11 @@ public:
 
 	}
 
-
 	virtual void RHIEndFrame() final override
 	{
 
 	}
-	virtual void RHIBeginScene() final override
-	{
-
-	}
-	virtual void RHIEndScene() final override
-	{
 
-	}
 	virtual void RHIAliasTextureResources(FTextureRHIRef& DestTexture, FTextureRHIRef& SrcTexture) final override
 	{
 
diff --git a/Engine/Source/Runtime/OpenGLDrv/Private/OpenGLDevice.cpp b/Engine/Source/Runtime/OpenGLDrv/Private/OpenGLDevice.cpp
--- a/Engine/Source/Runtime/OpenGLDrv/Private/OpenGLDevice.cpp
+++ b/Engine/Source/Runtime/OpenGLDrv/Private/OpenGLDevice.cpp
@@ -255,15 +255,6 @@ void FOpenGLDynamicRHI::RHIAdvanceFrameFence()
 	OpenGL_PollAllFences();
 }
 
-
-void FOpenGLDynamicRHI::RHIBeginScene()
-{
-}
-
-void FOpenGLDynamicRHI::RHIEndScene()
-{
-}
-
 #if PLATFORM_ANDROID
 
 JNI_METHOD void Java_com_epicgames_unreal_MediaPlayer14_nativeClearCachedAttributeState(JNIEnv* jenv, jobject thiz, jint PositionAttrib, jint TexCoordsAttrib)
diff --git a/Engine/Source/Runtime/OpenGLDrv/Private/OpenGLDrv.h b/Engine/Source/Runtime/OpenGLDrv/Private/OpenGLDrv.h
--- a/Engine/Source/Runtime/OpenGLDrv/Private/OpenGLDrv.h
+++ b/Engine/Source/Runtime/OpenGLDrv/Private/OpenGLDrv.h
@@ -429,8 +429,6 @@ public:
 	using FDynamicRHI::RHIBeginFrame;
 	virtual void RHIBeginFrame() final override;
 	virtual void RHIEndFrame() final override;
-	virtual void RHIBeginScene() final override;
-	virtual void RHIEndScene() final override;
 	virtual void RHISetStreamSource(uint32 StreamIndex, FRHIBuffer* VertexBuffer, uint32 Offset) final override;
 	virtual void RHISetRasterizerState(FRHIRasterizerState* NewState) final override;
 	virtual void RHISetViewport(float MinX, float MinY, float MinZ, float MaxX, float MaxY, float MaxZ) final override;
@@ -826,7 +824,6 @@ private:
 	/** Per-context state caching */
 	FOpenGLContextState	SharedContextState;
 	FOpenGLContextState	RenderingContextState;
-	// Cached context type on BeginScene
 	FOpenGLContextState* CachedContextState = nullptr;
 
 	template <typename TRHIShader>
diff --git a/Engine/Source/Runtime/RHI/Private/RHICommandList.cpp b/Engine/Source/Runtime/RHI/Private/RHICommandList.cpp
--- a/Engine/Source/Runtime/RHI/Private/RHICommandList.cpp
+++ b/Engine/Source/Runtime/RHI/Private/RHICommandList.cpp
@@ -1566,42 +1566,6 @@ void FRHICommandListImmediate::UnStallRHIThread()
 	}
 }
 
-void FRHICommandListImmediate::BeginScene()
-{
-	check(IsImmediate() && IsInRenderingThread());
-	if (Bypass())
-	{
-		GetContext().RHIBeginScene();
-		return;
-	}
-	ALLOC_COMMAND(FRHICommandBeginScene)();
-	if (!IsRunningRHIInSeparateThread())
-	{
-		// if we aren't running an RHIThread, there is no good reason to buffer this frame advance stuff and that complicates state management, so flush everything out now
-		QUICK_SCOPE_CYCLE_COUNTER(BeginScene_Flush);
-		CSV_SCOPED_TIMING_STAT(RHITFlushes, BeginScene);
-		ImmediateFlush(EImmediateFlushType::FlushRHIThread);
-	}
-}
-
-void FRHICommandListImmediate::EndScene()
-{
-	check(IsImmediate() && IsInRenderingThread());
-	if (Bypass())
-	{
-		GetContext().RHIEndScene();
-		return;
-	}
-	ALLOC_COMMAND(FRHICommandEndScene)();
-	if (!IsRunningRHIInSeparateThread())
-	{
-		// if we aren't running an RHIThread, there is no good reason to buffer this frame advance stuff and that complicates state management, so flush everything out now
-		QUICK_SCOPE_CYCLE_COUNTER(EndScene_Flush);
-		CSV_SCOPED_TIMING_STAT(RHITFlushes, EndScene);
-		ImmediateFlush(EImmediateFlushType::FlushRHIThread);
-	}
-}
-
 void FRHICommandListImmediate::BeginDrawingViewport(FRHIViewport* Viewport, FRHITexture* RenderTargetRHI)
 {
 	check(IsImmediate() && IsInRenderingThread());
diff --git a/Engine/Source/Runtime/RHI/Public/RHICommandList.h b/Engine/Source/Runtime/RHI/Public/RHICommandList.h
--- a/Engine/Source/Runtime/RHI/Public/RHICommandList.h
+++ b/Engine/Source/Runtime/RHI/Public/RHICommandList.h
@@ -2285,22 +2285,6 @@ FRHICOMMAND_MACRO(FRHICommandPollOcclusionQueries)
 	RHI_API void Execute(FRHICommandListBase& CmdList);
 };
 
-FRHICOMMAND_MACRO(FRHICommandBeginScene)
-{
-	FORCEINLINE_DEBUGGABLE FRHICommandBeginScene()
-	{
-	}
-	RHI_API void Execute(FRHICommandListBase& CmdList);
-};
-
-FRHICOMMAND_MACRO(FRHICommandEndScene)
-{
-	FORCEINLINE_DEBUGGABLE FRHICommandEndScene()
-	{
-	}
-	RHI_API void Execute(FRHICommandListBase& CmdList);
-};
-
 FRHICOMMAND_MACRO(FRHICommandBeginFrame)
 {
 	FORCEINLINE_DEBUGGABLE FRHICommandBeginFrame()
@@ -4509,8 +4493,12 @@ public:
 		return static_cast<FRHICommandListImmediate&>(RHICmdList);
 	}
 
-	RHI_API void BeginScene();
-	RHI_API void EndScene();
+	UE_DEPRECATED(5.5, "FRHICommandListImmediate::BeginScene() is deprecated and is no longer necessary. Remove any remaining calls to BeginScene(). There is no replacement.")
+	void BeginScene() {}
+
+	UE_DEPRECATED(5.5, "FRHICommandListImmediate::EndScene() is deprecated and is no longer necessary. Remove any remaining calls to EndScene(). There is no replacement.")
+	void EndScene() {}
+
 	RHI_API void BeginDrawingViewport(FRHIViewport* Viewport, FRHITexture* RenderTargetRHI);
 	RHI_API void EndDrawingViewport(FRHIViewport* Viewport, bool bPresent, bool bLockToVsync);
 	RHI_API void BeginFrame();
diff --git a/Engine/Source/Runtime/RHI/Public/RHICommandListCommandExecutes.inl b/Engine/Source/Runtime/RHI/Public/RHICommandListCommandExecutes.inl
--- a/Engine/Source/Runtime/RHI/Public/RHICommandListCommandExecutes.inl
+++ b/Engine/Source/Runtime/RHI/Public/RHICommandListCommandExecutes.inl
@@ -21,7 +21,6 @@ struct FMemory;
 struct FRHICommandBeginDrawingViewport;
 struct FRHICommandBeginFrame;
 struct FRHICommandBeginRenderQuery;
-struct FRHICommandBeginScene;
 struct FRHICommandBuildLocalBoundShaderState;
 struct FRHICommandBuildLocalGraphicsPipelineState;
 struct FRHICommandBuildLocalUniformBuffer;
@@ -36,7 +35,6 @@ struct FRHICommandEndDrawingViewport;
 struct FRHICommandEndFrame;
 struct FRHICommandEndOcclusionQueryBatch;
 struct FRHICommandEndRenderQuery;
-struct FRHICommandEndScene;
 struct FRHICommandSetBlendFactor;
 struct FRHICommandSetBoundShaderState;
 struct FRHICommandSetLocalGraphicsPipelineState;
@@ -573,18 +571,6 @@ void FRHICommandSetBindingsOnShaderBindingTable::Execute(FRHICommandListBase& Cm
 	INTERNAL_DECORATOR(RHISetBindingsOnShaderBindingTable)(SBT, RayTracingPipelineState, NumBindings, Bindings, BindingType);
 }
 
-void FRHICommandBeginScene::Execute(FRHICommandListBase& CmdList)
-{
-	RHISTAT(BeginScene);
-	INTERNAL_DECORATOR(RHIBeginScene)();
-}
-
-void FRHICommandEndScene::Execute(FRHICommandListBase& CmdList)
-{
-	RHISTAT(EndScene);
-	INTERNAL_DECORATOR(RHIEndScene)();
-}
-
 void FRHICommandBeginFrame::Execute(FRHICommandListBase& CmdList)
 {
 	RHISTAT(BeginFrame);
diff --git a/Engine/Source/Runtime/RHI/Public/RHIContext.h b/Engine/Source/Runtime/RHI/Public/RHIContext.h
--- a/Engine/Source/Runtime/RHI/Public/RHIContext.h
+++ b/Engine/Source/Runtime/RHI/Public/RHIContext.h
@@ -640,20 +640,6 @@ public:
 	// This method is queued with an RHIThread, otherwise it will flush after it is queued; without an RHI thread there is no benefit to queuing this frame advance commands
 	virtual void RHIEndFrame() = 0;
 
-	/**
-	* Signals the beginning of scene rendering. The RHI makes certain caching assumptions between
-	* calls to BeginScene/EndScene. Currently the only restriction is that you can't update texture
-	* references.
-	*/
-	// This method is queued with an RHIThread, otherwise it will flush after it is queued; without an RHI thread there is no benefit to queuing this frame advance commands
-	virtual void RHIBeginScene() = 0;
-
-	/**
-	* Signals the end of scene rendering. See RHIBeginScene.
-	*/
-	// This method is queued with an RHIThread, otherwise it will flush after it is queued; without an RHI thread there is no benefit to queuing this frame advance commands
-	virtual void RHIEndScene() = 0;
-
 	virtual void RHISetStreamSource(uint32 StreamIndex, FRHIBuffer* VertexBuffer, uint32 Offset) = 0;
 
 	// @param MinX including like Win32 RECT
diff --git a/Engine/Source/Runtime/RHI/Public/RHIValidationContext.h b/Engine/Source/Runtime/RHI/Public/RHIValidationContext.h
--- a/Engine/Source/Runtime/RHI/Public/RHIValidationContext.h
+++ b/Engine/Source/Runtime/RHI/Public/RHIValidationContext.h
@@ -696,26 +696,6 @@ public:
 	// This method is queued with an RHIThread, otherwise it will flush after it is queued; without an RHI thread there is no benefit to queuing this frame advance commands
 	virtual void RHIEndFrame() override final;
 
-	/**
-	* Signals the beginning of scene rendering. The RHI makes certain caching assumptions between
-	* calls to BeginScene/EndScene. Currently the only restriction is that you can't update texture
-	* references.
-	*/
-	// This method is queued with an RHIThread, otherwise it will flush after it is queued; without an RHI thread there is no benefit to queuing this frame advance commands
-	virtual void RHIBeginScene() override final
-	{
-		RHIContext->RHIBeginScene();
-	}
-
-	/**
-	* Signals the end of scene rendering. See RHIBeginScene.
-	*/
-	// This method is queued with an RHIThread, otherwise it will flush after it is queued; without an RHI thread there is no benefit to queuing this frame advance commands
-	virtual void RHIEndScene() override final
-	{
-		RHIContext->RHIEndScene();
-	}
-
 	virtual void RHISetStreamSource(uint32 StreamIndex, FRHIBuffer* VertexBuffer, uint32 Offset) override final
 	{
 		//#todo-rco: Decide if this is needed or not...
diff --git a/Engine/Source/Runtime/Renderer/Private/SceneRendering.cpp b/Engine/Source/Runtime/Renderer/Private/SceneRendering.cpp
--- a/Engine/Source/Runtime/Renderer/Private/SceneRendering.cpp
+++ b/Engine/Source/Runtime/Renderer/Private/SceneRendering.cpp
@@ -3781,8 +3781,6 @@ IVisibilityTaskData* FSceneRenderer::OnRenderBegin(FRDGBuilder& GraphBuilder)
 		{
 			LightFunctionAtlas::OnRenderBegin(LightFunctionAtlas, *Scene, Views, ViewFamily);
 		}
-	
-		GraphBuilder.RHICmdList.BeginScene();
 
 		VisibilityTaskData = LaunchVisibilityTasks(GraphBuilder.RHICmdList, *this, StaticMeshUpdateTask);
 
@@ -4307,9 +4305,6 @@ void FSceneRenderer::OnRenderFinish(FRDGBuilder& GraphBuilder, FRDGTextureRef Vi
 
 	AddPass(GraphBuilder, RDG_EVENT_NAME("EndScene"), [this](FRHICommandListImmediate& InRHICmdList)
 	{
-		// Notify the RHI we are done rendering a scene.
-		InRHICmdList.EndScene();
-
 		if (GDumpMeshDrawCommandMemoryStats)
 		{
 			GDumpMeshDrawCommandMemoryStats = 0;
diff --git a/Engine/Source/Runtime/VulkanRHI/Private/VulkanContext.h b/Engine/Source/Runtime/VulkanRHI/Private/VulkanContext.h
--- a/Engine/Source/Runtime/VulkanRHI/Private/VulkanContext.h
+++ b/Engine/Source/Runtime/VulkanRHI/Private/VulkanContext.h
@@ -104,9 +104,6 @@ public:
 	virtual void RHIBeginFrame() final override;
 	virtual void RHIEndFrame() final override;
 
-	virtual void RHIBeginScene() final override;
-	virtual void RHIEndScene() final override;
-
 	virtual void RHIBeginRenderPass(const FRHIRenderPassInfo& InInfo, const TCHAR* InName) final override;
 	virtual void RHIEndRenderPass() final override;
 	virtual void RHINextSubpass() final override;
diff --git a/Engine/Source/Runtime/VulkanRHI/Private/VulkanRHI.cpp b/Engine/Source/Runtime/VulkanRHI/Private/VulkanRHI.cpp
--- a/Engine/Source/Runtime/VulkanRHI/Private/VulkanRHI.cpp
+++ b/Engine/Source/Runtime/VulkanRHI/Private/VulkanRHI.cpp
@@ -1109,15 +1109,6 @@ void FVulkanCommandListContext::RHIBeginFrame()
 	}
 }
 
-
-void FVulkanCommandListContext::RHIBeginScene()
-{
-}
-
-void FVulkanCommandListContext::RHIEndScene()
-{
-}
-
 void FVulkanCommandListContext::RHIBeginDrawingViewport(FRHIViewport* ViewportRHI, FRHITexture* RenderTargetRHI)
 {
 	//FRCLog::Printf(FString::Printf(TEXT("FVulkanCommandListContext::RHIBeginDrawingViewport\n")));
diff --git a/Engine/Source/Runtime/Windows/D3D11RHI/Private/D3D11RHI.cpp b/Engine/Source/Runtime/Windows/D3D11RHI/Private/D3D11RHI.cpp
--- a/Engine/Source/Runtime/Windows/D3D11RHI/Private/D3D11RHI.cpp
+++ b/Engine/Source/Runtime/Windows/D3D11RHI/Private/D3D11RHI.cpp
@@ -37,16 +37,6 @@ void FD3D11DynamicRHI::RHIBeginFrame()
 #endif
 }
 
-void FD3D11DynamicRHI::RHIBeginScene()
-{
-	// Nothing to do
-}
-
-void FD3D11DynamicRHI::RHIEndScene()
-{
-	// Nothing to do
-}
-
 void FD3D11DynamicRHI::RHIEndFrame()
 {
 #if (RHI_NEW_GPU_PROFILER == 0)
diff --git a/Engine/Source/Runtime/Windows/D3D11RHI/Private/D3D11RHIPrivate.h b/Engine/Source/Runtime/Windows/D3D11RHI/Private/D3D11RHIPrivate.h
--- a/Engine/Source/Runtime/Windows/D3D11RHI/Private/D3D11RHIPrivate.h
+++ b/Engine/Source/Runtime/Windows/D3D11RHI/Private/D3D11RHIPrivate.h
@@ -529,8 +529,6 @@ public:
 	using FDynamicRHI::RHIBeginFrame;
 	virtual void RHIBeginFrame() override;
 	virtual void RHIEndFrame() override;
-	virtual void RHIBeginScene() final override;
-	virtual void RHIEndScene() final override;
 	virtual void RHISetStreamSource(uint32 StreamIndex, FRHIBuffer* VertexBuffer, uint32 Offset) final override;
 	virtual void RHISetRasterizerState(FRHIRasterizerState* NewState) final override;
 	virtual void RHISetViewport(float MinX, float MinY, float MinZ, float MaxX, float MaxY, float MaxZ) final override;
