Log a fatal error when trying to create a buffer with size 0.

Most RHIs will crash further down the line anyway, and in very obscure ways in test or shipping because checks are disabled. It's better to crash at the source, and log all the debug information we might have (such as the owner name).

#rnx
#tests Tried PIE and played a couple of games in a cooked client.

diff --git a/Engine/Source/Runtime/RHI/Private/RHICommandList.cpp b/Engine/Source/Runtime/RHI/Private/RHICommandList.cpp
--- a/Engine/Source/Runtime/RHI/Private/RHICommandList.cpp
+++ b/Engine/Source/Runtime/RHI/Private/RHICommandList.cpp
@@ -2131,6 +2131,11 @@ void FRHICommandList::GenerateMips(FRHITexture*)
 	UE_LOG(LogRHI, Fatal, TEXT("GenerateMips on RHI command lists is deprecated and no longer functions. Use the FGenerateMips helper class from the RenderCore module to generate mips on textures."));
 }
 
+void FRHICommandListBase::InvalidBufferFatalError(const FRHIResourceCreateInfo& CreateInfo, EBufferUsageFlags Usage, uint32 Stride)
+{
+	UE_LOG(LogRHI, Fatal, TEXT("Attempt to create zero-sized buffer '%s', owner '%s', usage 0x%x, stride %u."), CreateInfo.DebugName ? CreateInfo.DebugName : TEXT("(nullptr)"), *CreateInfo.OwnerName.ToString(), (uint32)Usage, Stride);
+}
+
 FRayTracingShaderBindings UE::RHI::ConvertRayTracingShaderBindings(const FRHIBatchedShaderParameters& BatchedParameters)
 {
 	FRayTracingShaderBindings Result;
@@ -2227,4 +2232,3 @@ void FRHICommandList::RayTraceDispatchIndirect(
 	GlobalResourceBindings.Reset();
 }
 
-
diff --git a/Engine/Source/Runtime/RHI/Public/RHICommandList.h b/Engine/Source/Runtime/RHI/Public/RHICommandList.h
--- a/Engine/Source/Runtime/RHI/Public/RHICommandList.h
+++ b/Engine/Source/Runtime/RHI/Public/RHICommandList.h
@@ -664,6 +664,7 @@ public:
 
 private:
 	RHI_API void ActivatePipelines(ERHIPipeline Pipelines);
+	RHI_API void InvalidBufferFatalError(const FRHIResourceCreateInfo& CreateInfo, EBufferUsageFlags Usage, uint32 Stride);
 
 public:
 	RHI_API ERHIPipeline SwitchPipeline(ERHIPipeline Pipeline);
@@ -727,6 +728,11 @@ public:
 			? FRHIBufferDesc::Null()
 			: FRHIBufferDesc(Size, Stride, Usage);
 
+		if (Size == 0 && !BufferDesc.IsNull())
+		{
+			InvalidBufferFatalError(CreateInfo, Usage, Stride);
+		}
+
 		FRHICommandListScopedPipelineGuard ScopedPipeline(*this);
 		FBufferRHIRef Buffer = GDynamicRHI->RHICreateBuffer(*this, BufferDesc, ResourceState, CreateInfo);
 		Buffer->SetTrackedAccess_Unsafe(ResourceState);
