[UE] Fixed CIS error related to compile configurations differences. DebugGame and Development can't differ.
#rnx

diff --git a/Engine/Source/Programs/UnrealBuildTool/Configuration/TargetRules.cs b/Engine/Source/Programs/UnrealBuildTool/Configuration/TargetRules.cs
--- a/Engine/Source/Programs/UnrealBuildTool/Configuration/TargetRules.cs
+++ b/Engine/Source/Programs/UnrealBuildTool/Configuration/TargetRules.cs
@@ -1016,7 +1016,7 @@ namespace UnrealBuildTool
 		public bool bUseIris { get; set; } = true;
 
 		/// <summary>
-		/// Whether to track owner (asset name) of RHI resource for non Shipping configurations.
+		/// Whether to track owner (asset name) of RHI resource for Test configuration.
 		/// Useful for ListShaderMaps and ListShaderLibraries commands.
 		/// </summary>
 		[RequiresUniqueBuildEnvironment]
diff --git a/Engine/Source/Runtime/RHI/RHI.Build.cs b/Engine/Source/Runtime/RHI/RHI.Build.cs
--- a/Engine/Source/Runtime/RHI/RHI.Build.cs
+++ b/Engine/Source/Runtime/RHI/RHI.Build.cs
@@ -20,8 +20,11 @@ public class RHI : ModuleRules
 
 			if (Target.Type != TargetRules.TargetType.Server)   // Dedicated servers should skip loading everything but NullDrv
 			{
-				// Always disable for Shipping, always enable for Development, let the user decide to use for Test and DebugGame
-				if (Target.Configuration != UnrealTargetConfiguration.Shipping && Target.bTrackRHIResourceInfo || Target.Configuration == UnrealTargetConfiguration.Development)
+				// Always disable for Shipping, always enable for Development and DebugGame,
+				// let the user decide to use for Test
+				else if (Target.Configuration == UnrealTargetConfiguration.DebugGame ||
+					Target.Configuration == UnrealTargetConfiguration.Development ||
+					Target.Configuration == UnrealTargetConfiguration.Test && Target.bTrackRHIResourceInfo)
 				{
 					PublicDefinitions.Add("RHI_WANT_RESOURCE_INFO=1");
 				}
