[UE] Fixed CIS error else keyword
#rnx

diff --git a/Engine/Source/Runtime/RHI/RHI.Build.cs b/Engine/Source/Runtime/RHI/RHI.Build.cs
--- a/Engine/Source/Runtime/RHI/RHI.Build.cs
+++ b/Engine/Source/Runtime/RHI/RHI.Build.cs
@@ -22,7 +22,7 @@ public class RHI : ModuleRules
 			{
 				// Always disable for Shipping, always enable for Development and DebugGame,
 				// let the user decide to use for Test
-				else if (Target.Configuration == UnrealTargetConfiguration.DebugGame ||
+				if (Target.Configuration == UnrealTargetConfiguration.DebugGame ||
 					Target.Configuration == UnrealTargetConfiguration.Development ||
 					Target.Configuration == UnrealTargetConfiguration.Test && Target.bTrackRHIResourceInfo)
 				{
