Copy over bindless resources from the batched paramaters to the FRayTracingShaderBindings

diff --git a/Engine/Source/Runtime/RHI/Private/RHICommandList.cpp b/Engine/Source/Runtime/RHI/Private/RHICommandList.cpp
--- a/Engine/Source/Runtime/RHI/Private/RHICommandList.cpp
+++ b/Engine/Source/Runtime/RHI/Private/RHICommandList.cpp
@@ -2181,6 +2181,8 @@ FRayTracingShaderBindings UE::RHI::ConvertRayTracingShaderBindings(const FRHIBat
 		}
 	}
 
+	Result.BindlessParameters = BatchedParameters.BindlessParameters;
+
 #else // // RHI_RAYTRACING
 
 	checkNoEntry();
