Add additional 'highest' PSO precaching priority level
Add 'r.PSOPrecache.DrawnComponentBoostStrategy' cvar to boost pso precache requests to the new highest priority when a scene proxy is drawn. (default off)
Include PSO precache priorities in RHI PSO create info.

diff --git a/Engine/Plugins/FX/Niagara/Source/Niagara/Private/NiagaraComponent.cpp b/Engine/Plugins/FX/Niagara/Source/Niagara/Private/NiagaraComponent.cpp
--- a/Engine/Plugins/FX/Niagara/Source/Niagara/Private/NiagaraComponent.cpp
+++ b/Engine/Plugins/FX/Niagara/Source/Niagara/Private/NiagaraComponent.cpp
@@ -101,6 +101,16 @@ static FAutoConsoleVariableRef CVarNiagaraForceWaitForCompilationOnActivate(
 	ECVF_Default
 );
 
+static bool GNiagaraComponentBoostPSOPrecacheHighest = false;
+static FAutoConsoleVariableRef CVarNiagaraComponentBoostPSOPrecacheHighest(
+	TEXT("r.PSOPrecache.NiagaraComponentPSOPrecachePriority"),
+	GNiagaraComponentBoostPSOPrecacheHighest,
+	TEXT("Niagara component PSO precache priority level.\n")
+	TEXT(" 0. Niagara component's PSO precache requests are set to high priority (default)\n")
+	TEXT(" 1. Niagara component's PSO precache requests are set to highest priority"),
+	ECVF_Default
+);
+
 FAutoConsoleCommandWithWorldAndArgs DumpNiagaraComponentsCommand(
 	TEXT("fx.Niagara.DumpComponents"),
 	TEXT("Dump Information about all Niagara Components"),
@@ -2329,7 +2339,8 @@ FPrimitiveSceneProxy* UNiagaraComponent::CreateSceneProxy()
 			PrecacheAssetPSOs(Asset);
 		}
 
-		if (CheckPSOPrecachingAndBoostPriority() && GetPSOPrecacheProxyCreationStrategy() != EPSOPrecacheProxyCreationStrategy::AlwaysCreate)
+		EPSOPrecachePriority PSOPrecachePriority = GNiagaraComponentBoostPSOPrecacheHighest ? EPSOPrecachePriority::Highest : EPSOPrecachePriority::High;
+		if (CheckPSOPrecachingAndBoostPriority(PSOPrecachePriority) && GetPSOPrecacheProxyCreationStrategy() != EPSOPrecacheProxyCreationStrategy::AlwaysCreate)
 		{
 			UE_LOG(LogNiagara, Verbose, TEXT("Skipping CreateSceneProxy for UNiagaraComponent %s (UNiagaraSystem PSOs are still compiling)"), *GetFullName());
 			return nullptr;
diff --git a/Engine/Source/Runtime/Engine/Classes/Components/PrimitiveComponent.h b/Engine/Source/Runtime/Engine/Classes/Components/PrimitiveComponent.h
--- a/Engine/Source/Runtime/Engine/Classes/Components/PrimitiveComponent.h
+++ b/Engine/Source/Runtime/Engine/Classes/Components/PrimitiveComponent.h
@@ -670,8 +670,9 @@ protected:
 	/** Helper flag to check if PSOs have been precached already */
 	uint8 bPSOPrecacheCalled : 1;
 
-	/** Have the PSO requests already been priority boosted? */
-	uint8 bPSOPrecacheRequestBoosted : 1;
+	/** PSOs requested priority */
+	EPSOPrecachePriority PSOPrecacheRequestPriority : 2;
+	static_assert((int)EPSOPrecachePriority::Highest < 1 << 2);
 
 	/** Cached array of material PSO requests which can be used to boost the priority */
 	TArray<FMaterialPSOPrecacheRequestID> MaterialPSOPrecacheRequestIDs;
@@ -936,9 +937,10 @@ public:
 
 	/**
 	 * Check if PSOs are still precaching and boost priority if not done yet.
+	 * PSOPrecachePriority specifies the new priority to boost to. calls that would reduce the current priority are ignored, highest pri is used for components that are actively rendering.
 	 * Returns true if the PSOs are still precaching.
 	 */
-	ENGINE_API bool CheckPSOPrecachingAndBoostPriority();
+	ENGINE_API bool CheckPSOPrecachingAndBoostPriority(EPSOPrecachePriority PSOPrecachePriority = EPSOPrecachePriority::High);
 
 protected:
 
diff --git a/Engine/Source/Runtime/Engine/Private/Components/PrimitiveComponent.cpp b/Engine/Source/Runtime/Engine/Private/Components/PrimitiveComponent.cpp
--- a/Engine/Source/Runtime/Engine/Private/Components/PrimitiveComponent.cpp
+++ b/Engine/Source/Runtime/Engine/Private/Components/PrimitiveComponent.cpp
@@ -384,7 +384,7 @@ UPrimitiveComponent::UPrimitiveComponent(const FObjectInitializer& ObjectInitial
 
 #if UE_WITH_PSO_PRECACHING
 	bPSOPrecacheCalled = false;
-	bPSOPrecacheRequestBoosted = false;
+	PSOPrecacheRequestPriority = EPSOPrecachePriority::Medium;
 #endif // UE_WITH_PSO_PRECACHING
 	
 	bApplyImpulseOnDamage = true;
@@ -4763,7 +4763,7 @@ void UPrimitiveComponent::PrecachePSOs()
 	// clear the current request data
 	MaterialPSOPrecacheRequestIDs.Empty();
 	PSOPrecacheCompileEvent = nullptr;
-	bPSOPrecacheRequestBoosted = false;
+	PSOPrecacheRequestPriority = EPSOPrecachePriority::Medium;
 
 	// Collect the data from the derived classes
 	FPSOPrecacheParams PSOPrecacheParams;
@@ -4778,6 +4778,20 @@ void UPrimitiveComponent::PrecachePSOs()
 #endif
 }
 
+class FPSOPrecacheFinishedTask : public FMarkActorRenderStateDirtyTask
+{
+public:
+	explicit FPSOPrecacheFinishedTask(UActorComponent* InActorComponent)
+		: FMarkActorRenderStateDirtyTask(InActorComponent)
+	{ }
+
+	void DoTask(ENamedThreads::Type CurrentThread, const FGraphEventRef& MyCompletionGraphEvent)
+	{
+		QUICK_SCOPE_CYCLE_COUNTER(STAT_PSOPrecacheFinishedTask);
+		FMarkActorRenderStateDirtyTask::DoTask(CurrentThread, MyCompletionGraphEvent);
+	}
+};
+
 void UPrimitiveComponent::RequestRecreateRenderStateWhenPSOPrecacheFinished(const FGraphEventArray& PSOPrecacheCompileEvents)
 {
 #if UE_WITH_PSO_PRECACHING
@@ -4785,7 +4799,7 @@ void UPrimitiveComponent::RequestRecreateRenderStateWhenPSOPrecacheFinished(cons
 	// schedule a task to mark the render state dirty when all PSOs are compiled so the proxy gets recreated.
 	if (UsePSOPrecacheRenderProxyDelay() && GetPSOPrecacheProxyCreationStrategy() != EPSOPrecacheProxyCreationStrategy::AlwaysCreate && !PSOPrecacheCompileEvents.IsEmpty())
 	{
-		PSOPrecacheCompileEvent = TGraphTask<FMarkActorRenderStateDirtyTask>::CreateTask(&PSOPrecacheCompileEvents).ConstructAndDispatchWhenReady(this);
+		PSOPrecacheCompileEvent = TGraphTask<FPSOPrecacheFinishedTask>::CreateTask(&PSOPrecacheCompileEvents).ConstructAndDispatchWhenReady(this);
 	}
 
 	bPSOPrecacheCalled = true;
@@ -4819,17 +4833,18 @@ bool UPrimitiveComponent::ShouldRenderProxyFallbackToDefaultMaterial() const
 #endif // UE_WITH_PSO_PRECACHING
 }
 
-bool UPrimitiveComponent::CheckPSOPrecachingAndBoostPriority()
+bool UPrimitiveComponent::CheckPSOPrecachingAndBoostPriority(EPSOPrecachePriority NewPSOPrecachePriority)
 {
 #if UE_WITH_PSO_PRECACHING
 	ensure(!IsComponentPSOPrecachingEnabled() || bPSOPrecacheCalled);
+	check(NewPSOPrecachePriority == EPSOPrecachePriority::High || NewPSOPrecachePriority == EPSOPrecachePriority::Highest);
 
 	if (PSOPrecacheCompileEvent && !PSOPrecacheCompileEvent->IsComplete())
 	{
-		if (!bPSOPrecacheRequestBoosted)
+		if (PSOPrecacheRequestPriority< NewPSOPrecachePriority)
 		{
-			BoostPSOPriority(MaterialPSOPrecacheRequestIDs);
-			bPSOPrecacheRequestBoosted = true;
+			BoostPSOPriority(NewPSOPrecachePriority, MaterialPSOPrecacheRequestIDs);
+			PSOPrecacheRequestPriority = NewPSOPrecachePriority;
 		}
 	}
 	else
@@ -5121,6 +5136,13 @@ FPrimitiveSceneProxy* FActorPrimitiveComponentInterface::CreateSceneProxy()
 	UPrimitiveComponent* Component = UPrimitiveComponent::GetPrimitiveComponent(this);
 	FPrimitiveSceneProxy* SceneProxy = Component->CreateSceneProxy();
 	Component->AssignSceneProxy(SceneProxy);
+
+#if UE_WITH_PSO_PRECACHING
+	if (ShouldBoostPSOPrecachePriorityOnDraw() && Component->IsPSOPrecaching() && Component->PSOPrecacheRequestPriority == EPSOPrecachePriority::High && SceneProxy)
+	{
+		SceneProxy->SetPSORequestsToBoostOnDraw(Component->MaterialPSOPrecacheRequestIDs);
+	}
+#endif
 	return SceneProxy;
 }
 
diff --git a/Engine/Source/Runtime/Engine/Private/PSOPrecache.cpp b/Engine/Source/Runtime/Engine/Private/PSOPrecache.cpp
--- a/Engine/Source/Runtime/Engine/Private/PSOPrecache.cpp
+++ b/Engine/Source/Runtime/Engine/Private/PSOPrecache.cpp
@@ -51,6 +51,16 @@ static FAutoConsoleVariableRef CVarPSOProxyCreationDelayStrategy(
 	ECVF_ReadOnly
 );
 
+static int32 GPSODrawnComponentBoostStrategy = 0;
+static FAutoConsoleVariableRef CVarPSOComponentBoostStrategy(
+	TEXT("r.PSOPrecache.DrawnComponentBoostStrategy"),
+	GPSODrawnComponentBoostStrategy,
+	TEXT("Increase priority of queued precache PSOs which are also required by the component for rendering.\n")
+	TEXT("0 do not increase priority of drawn PSOs (default)\n")
+	TEXT("1 if the component has been rendered then increase the priority of it's PSO precache requests. (this requires r.PSOPrecache.ProxyCreationDelayStrategy == 1.)"),
+	ECVF_ReadOnly
+);
+
 bool IsComponentPSOPrecachingEnabled()
 {
 	return FApp::CanEverRender() && PipelineStateCache::IsPSOPrecachingEnabled() && GPSOPrecacheComponents && !GIsEditor;
@@ -61,6 +71,11 @@ bool IsResourcePSOPrecachingEnabled()
 	return FApp::CanEverRender() && PipelineStateCache::IsPSOPrecachingEnabled() && GPSOPrecacheResources && !GIsEditor;
 }
 
+bool ShouldBoostPSOPrecachePriorityOnDraw()
+{
+	return FApp::CanEverRender() && PipelineStateCache::IsPSOPrecachingEnabled() && GPSODrawnComponentBoostStrategy && !GIsEditor;
+}
+
 EPSOPrecacheProxyCreationStrategy GetPSOPrecacheProxyCreationStrategy()
 {
 	if (GPSOProxyCreationWhenPSOReady != 1)
diff --git a/Engine/Source/Runtime/Engine/Private/PSOPrecacheMaterial.cpp b/Engine/Source/Runtime/Engine/Private/PSOPrecacheMaterial.cpp
--- a/Engine/Source/Runtime/Engine/Private/PSOPrecacheMaterial.cpp
+++ b/Engine/Source/Runtime/Engine/Private/PSOPrecacheMaterial.cpp
@@ -118,6 +118,11 @@ public:
 
 	void DoTask(ENamedThreads::Type CurrentThread, const FGraphEventRef& MyCompletionGraphEvent);
 
+	~FMaterialPSOPrecacheCollectionTask()
+	{
+		//check(MaterialInterface == nullptr);  // TODO: reinstate this or replace TStrongObjectPtr* with TStrongObjectPtr..
+	}
+
 public:
 
 	TStrongObjectPtr<UMaterialInterface>* MaterialInterface;
@@ -281,7 +286,7 @@ public:
 		}
 
 		// Boost priority if requested already
-		if (FindResult->Priority == EPSOPrecachePriority::High)
+		if (FindResult->Priority >= EPSOPrecachePriority::High)
 		{
 			CheckCompilingPSOs(*FindResult, true /*bBoostPriority*/);
 		}
@@ -299,7 +304,7 @@ public:
 		MaterialPSORequests[MaterialPSORequestID] = FMaterialPSOPrecacheParams();
 	}
 
-	void BoostPriority(FMaterialPSOPrecacheRequestID MaterialPSORequestID)
+	void BoostPriority(EPSOPrecachePriority NewPri, FMaterialPSOPrecacheRequestID MaterialPSORequestID)
 	{
 		check(MaterialPSORequestID != INDEX_NONE);
 
@@ -315,7 +320,7 @@ public:
 			FPrecacheData* FindResult = MaterialPSORequestData.Find(Params);
 
 			// Only process if not boosted yet and not completed yet
-			if (FindResult == nullptr || FindResult->Priority == EPSOPrecachePriority::High || FindResult->State == EState::Completed)
+			if (FindResult == nullptr || NewPri <= FindResult->Priority || FindResult->State == EState::Completed)
 			{
 				return;
 			}
@@ -325,8 +330,7 @@ public:
 		const FMaterialPSOPrecacheParams& Params = MaterialPSORequests[MaterialPSORequestID];
 		FPrecacheData* FindResult = MaterialPSORequestData.Find(Params);
 		check(FindResult);
-		FindResult->Priority = EPSOPrecachePriority::High;
-
+		FindResult->Priority = NewPri;
 		// Boost PSOs which are still compiling
 		CheckCompilingPSOs(*FindResult, true /*bBoostPriority*/);
 	}
@@ -431,7 +435,7 @@ private:
 				}
 				else if (bBoostPriority)
 				{
-					PipelineStateCache::BoostPrecachePriority(RequestResult.RequestID);
+					PipelineStateCache::BoostPrecachePriority(PrecacheData.Priority, RequestResult.RequestID);
 				}
 			}
 
@@ -477,6 +481,7 @@ void FMaterialPSOPrecacheCollectionTask::DoTask(ENamedThreads::Type CurrentThrea
 	if (RequestLifecycleID != GMaterialPSORequestManager.GetLifecycleID())
 	{
 		CollectionGraphEvent->DispatchSubsequents();
+		MaterialInterface->Reset();
 		return;
 	}
 
@@ -535,13 +540,13 @@ void ReleasePSOPrecacheData(const TArray<FMaterialPSOPrecacheRequestID>& Materia
 	}
 }
 
-void BoostPSOPriority(const TArray<FMaterialPSOPrecacheRequestID>& MaterialPSORequestIDs)
+void BoostPSOPriority(EPSOPrecachePriority NewPri, const TArray<FMaterialPSOPrecacheRequestID>& MaterialPSORequestIDs)
 {
 	TRACE_CPUPROFILER_EVENT_SCOPE(BoostPSOPriority);
 
 	for (FMaterialPSOPrecacheRequestID RequestID : MaterialPSORequestIDs)
 	{
-		GMaterialPSORequestManager.BoostPriority(RequestID);
+		GMaterialPSORequestManager.BoostPriority(NewPri, RequestID);
 	}
 }
 
diff --git a/Engine/Source/Runtime/Engine/Private/Particles/ParticleComponents.cpp b/Engine/Source/Runtime/Engine/Private/Particles/ParticleComponents.cpp
--- a/Engine/Source/Runtime/Engine/Private/Particles/ParticleComponents.cpp
+++ b/Engine/Source/Runtime/Engine/Private/Particles/ParticleComponents.cpp
@@ -3551,7 +3551,7 @@ void UFXSystemComponent::PrecacheAssetPSOs(UFXSystemAsset* FXSystemAsset)
 
 	MaterialPSOPrecacheRequestIDs.Empty();
 	PSOPrecacheCompileEvent = nullptr;
-	bPSOPrecacheRequestBoosted = false;
+	PSOPrecacheRequestPriority = EPSOPrecachePriority::Medium;
 
 	// The asset will keep the Precache events alive, but these might be over. Avoid delaying scene proxy creation if everything is finished
 	bool bAllEventsDone = GraphEvent == nullptr || GraphEvent->IsComplete();
diff --git a/Engine/Source/Runtime/Engine/Private/PrimitiveSceneProxy.cpp b/Engine/Source/Runtime/Engine/Private/PrimitiveSceneProxy.cpp
--- a/Engine/Source/Runtime/Engine/Private/PrimitiveSceneProxy.cpp
+++ b/Engine/Source/Runtime/Engine/Private/PrimitiveSceneProxy.cpp
@@ -1812,6 +1812,22 @@ bool FPrimitiveSceneProxy::GetMaterialTextureScales(int32 LODIndex, int32 Sectio
 
 #endif // WITH_EDITORONLY_DATA
 
+#if UE_WITH_PSO_PRECACHING
+void FPrimitiveSceneProxy::BoostPrecachedPSORequestsOnDraw()
+{
+	if (!PSOPrecacheRequestsToBoostOnDraw.IsEmpty())
+	{
+		BoostPSOPriority(EPSOPrecachePriority::Highest, PSOPrecacheRequestsToBoostOnDraw);
+		PSOPrecacheRequestsToBoostOnDraw.Empty();
+	}
+}
+
+void FPrimitiveSceneProxy::SetPSORequestsToBoostOnDraw(const TArray<FMaterialPSOPrecacheRequestID>& PSORequestsToBoostOnDrawIN)
+{
+	PSOPrecacheRequestsToBoostOnDraw = PSORequestsToBoostOnDrawIN;
+}
+#endif // UE_WITH_PSO_PRECACHING
+
 #if RHI_RAYTRACING
 ERayTracingPrimitiveFlags FPrimitiveSceneProxy::GetCachedRayTracingInstance(FRayTracingInstance& OutRayTracingInstance)
 {
diff --git a/Engine/Source/Runtime/Engine/Public/PSOPrecache.h b/Engine/Source/Runtime/Engine/Public/PSOPrecache.h
--- a/Engine/Source/Runtime/Engine/Public/PSOPrecache.h
+++ b/Engine/Source/Runtime/Engine/Public/PSOPrecache.h
@@ -129,15 +129,6 @@ struct FPSOPrecacheParams
 // Unique ID to find the FVertexDeclarationElementList - these can be shared
 using FVertexDeclarationElementListID = uint16;
 
-/**
- * PSO Precache request priority
- */
-enum class EPSOPrecachePriority : uint8
-{
-	Medium,
-	High
-};
-
 /**
  * Wraps vertex factory data used during PSO precaching - optional element list ID can be used if manual vertex fetch is not possible for the given vertex factory type
  */
@@ -266,6 +257,11 @@ extern ENGINE_API bool IsComponentPSOPrecachingEnabled();
  */
 extern ENGINE_API bool IsResourcePSOPrecachingEnabled();
 
+/**
+ * Boost drawn PSO precache request priority
+ */
+extern ENGINE_API bool ShouldBoostPSOPrecachePriorityOnDraw();
+
 enum class EPSOPrecacheProxyCreationStrategy : uint8
 {
 	// Always create the render proxy regardless of whether the PSO has finished precaching or not. 
diff --git a/Engine/Source/Runtime/Engine/Public/PSOPrecacheFwd.h b/Engine/Source/Runtime/Engine/Public/PSOPrecacheFwd.h
--- a/Engine/Source/Runtime/Engine/Public/PSOPrecacheFwd.h
+++ b/Engine/Source/Runtime/Engine/Public/PSOPrecacheFwd.h
@@ -8,6 +8,7 @@
 
 #include "Containers/Array.h"
 #include "HAL/Platform.h"
+#include "PipelineStateCache.h"
 
 // General switch that decides whether to compile out some PSO precaching code (most importantly, reduce sizeofs of common classes)
 #ifndef UE_WITH_PSO_PRECACHING
@@ -22,4 +23,4 @@ struct FPSOPrecacheParams;
 typedef TArray<FMaterialInterfacePSOPrecacheParams, TInlineAllocator<4> > FMaterialInterfacePSOPrecacheParamsList;
 
 // Unique request ID of MaterialPSOPrecache which can be used to boost the priority of a PSO precache requests if it's needed for rendering
-using FMaterialPSOPrecacheRequestID = uint32;
\ No newline at end of file
+using FMaterialPSOPrecacheRequestID = uint32;
diff --git a/Engine/Source/Runtime/Engine/Public/PSOPrecacheMaterial.h b/Engine/Source/Runtime/Engine/Public/PSOPrecacheMaterial.h
--- a/Engine/Source/Runtime/Engine/Public/PSOPrecacheMaterial.h
+++ b/Engine/Source/Runtime/Engine/Public/PSOPrecacheMaterial.h
@@ -139,7 +139,7 @@ extern ENGINE_API void ReleasePSOPrecacheData(const TArray<FMaterialPSOPrecacheR
 /**
  * Boost priority for all the PSOs still compiling for the request material request IDs
  */
-extern ENGINE_API void BoostPSOPriority(const TArray<FMaterialPSOPrecacheRequestID>& MaterialPSORequestIDs);
+extern ENGINE_API void BoostPSOPriority(EPSOPrecachePriority NewPri, const TArray<FMaterialPSOPrecacheRequestID>& MaterialPSORequestIDs);
 
 /**
  * Invalidate & clear all the current material PSO requests
diff --git a/Engine/Source/Runtime/Engine/Public/PrimitiveSceneProxy.h b/Engine/Source/Runtime/Engine/Public/PrimitiveSceneProxy.h
--- a/Engine/Source/Runtime/Engine/Public/PrimitiveSceneProxy.h
+++ b/Engine/Source/Runtime/Engine/Public/PrimitiveSceneProxy.h
@@ -23,6 +23,7 @@
 #include "SceneDefinitions.h"
 #include "MeshDrawCommandStatsDefines.h"
 #include "InstanceDataTypes.h"
+#include "PSOPrecacheFwd.h"
 
 class FLightSceneInfo;
 class FLightSceneProxy;
@@ -1150,6 +1151,11 @@ public:
 	 */
 	FInstanceDataBufferHeader GetInstanceDataHeader() const;
 
+#if UE_WITH_PSO_PRECACHING
+	ENGINE_API void BoostPrecachedPSORequestsOnDraw();
+	ENGINE_API void SetPSORequestsToBoostOnDraw(const TArray<FMaterialPSOPrecacheRequestID>& PSORequestsToBoostOnDrawIN);
+#endif
+
 protected:
 	ENGINE_API void SetupInstanceSceneDataBuffers(const FInstanceSceneDataBuffers* InInstanceSceneDataBuffers);
 
@@ -1605,7 +1611,10 @@ private:
 
 	TArray<UMaterialInterface*> UsedMaterialsForVerification;
 #endif
-
+#if UE_WITH_PSO_PRECACHING
+	// The pso precache request IDs that will be boosted to the highest priority when drawn.
+	TArray<FMaterialPSOPrecacheRequestID> PSOPrecacheRequestsToBoostOnDraw;
+#endif
 	/**
 	 * Updates the primitive proxy's cached transforms, and calls OnUpdateTransform to notify it of the change.
 	 * Called in the thread that owns the proxy; game or rendering.
diff --git a/Engine/Source/Runtime/RHI/Private/PipelineStateCache.cpp b/Engine/Source/Runtime/RHI/Private/PipelineStateCache.cpp
--- a/Engine/Source/Runtime/RHI/Private/PipelineStateCache.cpp
+++ b/Engine/Source/Runtime/RHI/Private/PipelineStateCache.cpp
@@ -50,6 +50,9 @@ DECLARE_DWORD_ACCUMULATOR_STAT(TEXT("Active Compute PSO Precache Requests"), STA
 DECLARE_DWORD_ACCUMULATOR_STAT(TEXT("High Priority Graphics PSO Precache Requests"), STAT_HighPriorityGraphicsPSOPrecacheRequests, STATGROUP_PipelineStateCache);
 DECLARE_DWORD_ACCUMULATOR_STAT(TEXT("High Priority Compute PSO Precache Requests"), STAT_HighPriorityComputePSOPrecacheRequests, STATGROUP_PipelineStateCache);
 
+DECLARE_DWORD_ACCUMULATOR_STAT(TEXT("Highest Priority Graphics PSO Precache Requests"), STAT_HighestPriorityGraphicsPSOPrecacheRequests, STATGROUP_PipelineStateCache);
+DECLARE_DWORD_ACCUMULATOR_STAT(TEXT("Highest Priority Compute PSO Precache Requests"), STAT_HighestPriorityComputePSOPrecacheRequests, STATGROUP_PipelineStateCache);
+
 static inline uint32 GetTypeHash(const FBoundShaderStateInput& Input)
 {
 	uint32 Hash = GetTypeHash(Input.VertexDeclarationRHI);
@@ -479,51 +482,106 @@ static void HandlePipelineCreationFailure(const FRHIComputeShader* ComputeShader
 	}
 }
 
-// FAsyncTask used by the PSO threadpool that will reschedule once only.
-template<typename TTask>
-class FAsyncTaskLimitedReschedule : public FAsyncTask<TTask>
+#if PSO_TRACK_CACHE_STATS
+
+static std::atomic_uint64_t TotalPrecompileCompleteTime[(int)EQueuedWorkPriority::Count];
+static std::atomic_uint64_t TotalPrecompileCompileTime[(int)EQueuedWorkPriority::Count];
+static std::atomic_uint64_t TotalPrecompileTimeToBegin[(int)EQueuedWorkPriority::Count];
+
+static std::atomic_int64_t TotalNumPrecompileJobs[(int)EQueuedWorkPriority::Count];
+static std::atomic_int64_t TotalNumPrecompileJobsCompleted[(int)EQueuedWorkPriority::Count];
+
+static std::atomic_int64_t MaxPrecompileJobTime[(int)EQueuedWorkPriority::Count];
+static std::atomic_int64_t MaxPrecompileTimeToCompile[(int)EQueuedWorkPriority::Count];
+static std::atomic_int64_t MaxPrecompileTimeToBegin[(int)EQueuedWorkPriority::Count];
+
+void ResetPrecompileStats()
 {
+	for (int i = 0; i < (int)EQueuedWorkPriority::Count; i++)
+	{
+		TotalPrecompileCompleteTime[i] = 0;
+		TotalPrecompileCompileTime[i] = 0;
+		TotalPrecompileTimeToBegin[i] = 0;
+
+		TotalNumPrecompileJobs[i] = 0;
+		TotalNumPrecompileJobsCompleted[i] = 0;
+
+		MaxPrecompileJobTime[i] = 0;
+		MaxPrecompileTimeToCompile[i] = 0;
+		MaxPrecompileTimeToBegin[i] = 0;
+	}
+}
+
+void StatsEndPrecompile(uint64 CreateTime, uint64 RescheduleTime, uint64 TaskBeginTime, uint64 EndTime, EQueuedWorkPriority TaskPri)
+{
+	uint64 TaskIssueTime = FMath::Max(CreateTime, RescheduleTime);
+	uint64 TimeToComplete = EndTime - TaskIssueTime;
+	uint64 TimeToCompile = EndTime - TaskBeginTime;
+	uint64 TimeToBegin = TaskBeginTime - TaskIssueTime;
+	check (TaskBeginTime > TaskIssueTime);
+
+	TotalPrecompileCompleteTime[(int)TaskPri] += TimeToComplete;
+	TotalPrecompileCompileTime[(int)TaskPri] += TimeToCompile;
+	TotalPrecompileTimeToBegin[(int)TaskPri] += TimeToBegin;
+
+	MaxPrecompileJobTime[(int)TaskPri] = FMath::Max((uint64)MaxPrecompileJobTime[(int)TaskPri].load(), (uint64)TimeToComplete);
+	MaxPrecompileTimeToCompile[(int)TaskPri] = FMath::Max((uint64)MaxPrecompileTimeToCompile[(int)TaskPri].load(), (uint64)TimeToCompile);
+	MaxPrecompileTimeToBegin[(int)TaskPri] = FMath::Max((uint64)MaxPrecompileTimeToBegin[(int)TaskPri].load(), (uint64)TimeToBegin);
+
+	TotalNumPrecompileJobsCompleted[(int)TaskPri]++;
+}
+#endif
+
+class FPSOPrecacheAsyncTask
+	: public FAsyncTaskBase
+{
+	TUniqueFunction<void(const FPSOPrecacheAsyncTask*)> AsyncTaskFunc;
 public:
-	/** Forwarding constructor. */
-	template<typename...T>
-	explicit FAsyncTaskLimitedReschedule(T&&... Args) : FAsyncTask<TTask>(Forward<T>(Args)...)	{ }
+
+	FPSOPrecacheAsyncTask(TUniqueFunction<void(const FPSOPrecacheAsyncTask*)> InFunc) : AsyncTaskFunc(MoveTemp(InFunc))
+	{
+#if PSO_TRACK_CACHE_STATS
+		CreateTime = FPlatformTime::Cycles64();
+#endif
+		// Cache the StatId to remain backward compatible with TTask that declare GetStatId as non-const.
+		Init(GetStatId());
+	}
+	bool TryAbandonTask() final 	{ 	return false;	}
 
 	bool Reschedule(FQueuedThreadPool* InQueuedPool, EQueuedWorkPriority InQueuedWorkPriority)
 	{
-		if (bCanReschedule.exchange(false))
+		uint64 RescheduleAttemptTime = FPlatformTime::Cycles64();
+		bool bSuccess = FAsyncTaskBase::Reschedule(InQueuedPool, InQueuedWorkPriority);
+#if PSO_TRACK_CACHE_STATS
+		if (bSuccess)
 		{
-			return FAsyncTask<TTask>::Reschedule(InQueuedPool, InQueuedWorkPriority);
+			RescheduleTime = RescheduleAttemptTime;
 		}
-		return false;
+#endif
+		return bSuccess;
 	}
-private:
-	std::atomic_bool bCanReschedule = true;
-};
 
-class FPSOPrecompileTask : public FNonAbandonableTask
-{
-	friend class FAsyncTask<FPSOPrecompileTask>;
-private:
-	TUniqueFunction<void()> Func;
-	TUniqueFunction<void()> OnComplete;
-public:
-	FPSOPrecompileTask(TUniqueFunction<void()> InFunc, TUniqueFunction<void()> InOnComplete) : Func(MoveTemp(InFunc)), OnComplete(MoveTemp(InOnComplete)) {}
-
-protected:
-	void DoWork() 
-	{ 
-		QUICK_SCOPE_CYCLE_COUNTER(STAT_FPSOPrecompileTask_Work);
-
-		Func(); 
-		Func = nullptr;
-		OnComplete(); 
-		OnComplete = nullptr;
+	void DoTaskWork() final
+	{
+#if PSO_TRACK_CACHE_STATS
+		TaskBeginTime = FPlatformTime::Cycles64();
+#endif
+		AsyncTaskFunc(this);
+#if PSO_TRACK_CACHE_STATS
+		StatsEndPrecompile(CreateTime, RescheduleTime, TaskBeginTime, FPlatformTime::Cycles64(), GetPriority());
+#endif
 	}
 
 	FORCEINLINE TStatId GetStatId() const
 	{
-		RETURN_QUICK_DECLARE_CYCLE_STAT(FPSOPrecompileTask, STATGROUP_ThreadPoolAsyncTasks);
+		RETURN_QUICK_DECLARE_CYCLE_STAT(FPSOPrecacheAsyncTask, STATGROUP_ThreadPoolAsyncTasks);
 	}
+
+#if PSO_TRACK_CACHE_STATS
+	uint64 CreateTime = 0;
+	uint64 RescheduleTime = 0;
+	uint64 TaskBeginTime = 0;
+#endif
 };
 
 /**
@@ -549,7 +607,7 @@ public:
 	virtual bool IsCompute() const = 0;
 
 	FGraphEventRef CompletionEvent;
-	TUniquePtr<FAsyncTaskLimitedReschedule<FPSOPrecompileTask>> PrecompileTask;
+	TUniquePtr<FPSOPrecacheAsyncTask> PrecompileTask;
 
 	bool IsComplete()
 	{
@@ -1306,9 +1364,11 @@ enum class EPSOPrecacheStateMask : uint8
 	Succeeded = 1 << 1, // once compilation is finished
 	Failed = 1 << 2,    // once compilation is finished
 	Boosted = 1 << 3,
+	HighestPri = 1 << 4,
 };
 
 ENUM_CLASS_FLAGS(EPSOPrecacheStateMask)
+static bool GForceHighToHighestPri = false;
 
 template<class TPrecachePipelineCacheDerived, class TPrecachedPSOInitializer, class TPipelineState>
 class TPrecachePipelineCacheBase
@@ -1325,15 +1385,40 @@ public:
 	
 protected:
 
-	void RescheduleTaskToHighPriority(TPipelineState* PipelineState)
+	void RescheduleTaskToHighPriority(EPSOPrecacheStateMask NewState, EPSOPrecacheStateMask PrevState, TPipelineState* PipelineState)
 	{
+		bool bHighestPriority = EnumHasAnyFlags(NewState, EPSOPrecacheStateMask::HighestPri);
+		bool bWasPreviouslyHigh = EnumHasAnyFlags(PrevState, EPSOPrecacheStateMask::Boosted);
+
+		check(!EnumHasAnyFlags(PrevState, EPSOPrecacheStateMask::HighestPri));
+
+		bool bCompleted = EnumHasAnyFlags(PrevState, EPSOPrecacheStateMask::Failed | EPSOPrecacheStateMask::Succeeded);
+		UE_CLOG(bCompleted, LogRHI, Error, TEXT("pso request has completed? prev %x, new %x"), (uint32)PrevState, (uint32)NewState);
+
 		if (FPSOPrecacheThreadPool::UsePool())
 		{
 			check(PipelineState->PrecompileTask);
-			PipelineState->PrecompileTask->Reschedule(&GPSOPrecacheThreadPool.Get(), EQueuedWorkPriority::Highest);
+			EQueuedWorkPriority NewPriority = bHighestPriority ? EQueuedWorkPriority::Highest : EQueuedWorkPriority::High;
+			if(PipelineState->PrecompileTask)
+			{
+				EQueuedWorkPriority PrevPriority = PipelineState->PrecompileTask->GetPriority();
+				check(PrevPriority > NewPriority);
+				PipelineState->PrecompileTask->Reschedule(&GPSOPrecacheThreadPool.Get(), NewPriority);
+			}
 		}
 
-		UpdateHighPriorityCompileCount(true /*Increment*/);
+		if (bHighestPriority)
+		{
+			UpdateHighestPriorityCompileCount(true);
+			if (bWasPreviouslyHigh)
+			{
+				UpdateHighPriorityCompileCount(false /*decrement*/);
+			}
+		}
+		else
+		{
+			UpdateHighPriorityCompileCount(true);
+		}
 	}
 
 	FPSOPrecacheRequestResult TryAddNewState(const TPrecachedPSOInitializer& Initializer, const FString& PSOCompilationEventName, bool bDoAsyncCompile)
@@ -1395,9 +1480,9 @@ protected:
 			{
 				EPSOPrecacheStateMask PreviousStateMask = FindResult->AddPSOPrecacheState(EPSOPrecacheStateMask::Compiling);
 				// by the time we're here, PrecacheFinished might already have been called, so boost it only if we know we will call it
-				if ( PreviousStateMask == EPSOPrecacheStateMask::Boosted)
+				if (!IsCompilationDone(PreviousStateMask) && EnumHasAnyFlags(PreviousStateMask, EPSOPrecacheStateMask::Boosted) )
 				{
-					RescheduleTaskToHighPriority(FindResult->PipelineState);
+					RescheduleTaskToHighPriority(PreviousStateMask, (EPSOPrecacheStateMask)0, FindResult->PipelineState);
 				}
 			}
 		}
@@ -1471,7 +1556,7 @@ public:
 		return ActiveCompileCount != 0;
 	}
 
-	void BoostPriority(const FPSOPrecacheRequestID& RequestID)
+	void BoostPriority(EPSOPrecachePriority PSOPrecachePriority, const FPSOPrecacheRequestID& RequestID)
 	{
 		check(RequestID.GetType() == PSOType);
 
@@ -1480,12 +1565,17 @@ public:
 		const TPrecachedPSOInitializer& Initializer = PrecachedPSOInitializers[RequestID.RequestID];
 		FPrecacheTask* FindResult = PrecachedPSOInitializerData.Find(Initializer);
 		check(FindResult);
-		EPSOPrecacheStateMask PreviousStateMask = FindResult->AddPSOPrecacheState(EPSOPrecacheStateMask::Boosted);
+
+		EPSOPrecacheStateMask NewMask = EPSOPrecacheStateMask::Boosted | (PSOPrecachePriority == EPSOPrecachePriority::Highest ? EPSOPrecacheStateMask::HighestPri : (EPSOPrecacheStateMask)0);
+		EPSOPrecacheStateMask PreviousStateMask = FindResult->AddPSOPrecacheState(NewMask);
 		// It's possible to get a boost request while the task has not been started yet. In this case, TryAddNewState will take care of it
 		// if TryAddNewState is done, then we can proceed to boost it, if the task is not done yet
-		if (PreviousStateMask == EPSOPrecacheStateMask::Compiling)
+		if (!IsCompilationDone(PreviousStateMask) && EnumHasAnyFlags(PreviousStateMask, EPSOPrecacheStateMask::Compiling) && !EnumHasAnyFlags(PreviousStateMask, EPSOPrecacheStateMask::HighestPri) )
 		{
-			RescheduleTaskToHighPriority(FindResult->PipelineState);
+			if (!EnumHasAnyFlags(PreviousStateMask, EPSOPrecacheStateMask::Boosted) || EnumHasAnyFlags(NewMask, EPSOPrecacheStateMask::HighestPri))
+			{
+				RescheduleTaskToHighPriority(NewMask, PreviousStateMask, FindResult->PipelineState);
+			}
 		}
 	}
 
@@ -1493,7 +1583,7 @@ public:
 	{
 		if (GPSOWaitForHighPriorityRequestsOnly)
 		{
-			return FPlatformAtomics::AtomicRead(&HighPriorityCompileCount);
+			return FPlatformAtomics::AtomicRead(&HighPriorityCompileCount) + FPlatformAtomics::AtomicRead(&HighestPriorityCompileCount);
 		}
 		else
 		{
@@ -1524,7 +1614,14 @@ public:
         // yet) then we must ignore the request
 		if (EnumHasAllFlags(PreviousStateMask, EPSOPrecacheStateMask::Boosted | EPSOPrecacheStateMask::Compiling))
 		{
-			UpdateHighPriorityCompileCount(false /*Increment*/);
+			if (EnumHasAnyFlags(PreviousStateMask, EPSOPrecacheStateMask::HighestPri))
+			{
+				UpdateHighestPriorityCompileCount(false  /*Increment*/);
+			}
+			else
+			{
+				UpdateHighPriorityCompileCount(false /*Increment*/);
+			}
 		}
 		UpdateActiveCompileCount(false /*Increment*/);
 	}
@@ -1538,7 +1635,8 @@ public:
 	{
 		SET_DWORD_STAT_FName(TPrecachePipelineCacheDerived::GetActiveCompileStatName(), ActiveCompileCount);
 		SET_DWORD_STAT_FName(TPrecachePipelineCacheDerived::GetHighPriorityCompileStatName(), HighPriorityCompileCount);
-
+		SET_DWORD_STAT_FName(TPrecachePipelineCacheDerived::GetHighestPriorityCompileStatName(), HighestPriorityCompileCount);
+		
 		FRWScopeLock WriteLock(PrecachePSOsRWLock, SLT_Write);
 		for (int32 Index = 0; Index < PrecachedPSOs.Num(); ++Index)		
 		{
@@ -1628,6 +1726,20 @@ protected:
 		else
 		{
 			FPlatformAtomics::InterlockedDecrement(&HighPriorityCompileCount);
+			check(HighPriorityCompileCount >= 0);
+		}
+	}
+
+	void UpdateHighestPriorityCompileCount(bool bIncrement)
+	{
+		if (bIncrement)
+		{
+			FPlatformAtomics::InterlockedIncrement(&HighestPriorityCompileCount);
+		}
+		else
+		{
+			FPlatformAtomics::InterlockedDecrement(&HighestPriorityCompileCount);
+			check(HighestPriorityCompileCount >= 0);
 		}
 	}
 
@@ -1682,6 +1794,8 @@ protected:
 	// Number of open high priority compiles
 	volatile int32 HighPriorityCompileCount = 0;
 
+	volatile int32 HighestPriorityCompileCount = 0;
+
 	// Finished Precached PSOs which can be garbage collected
 	TArray<TPrecachedPSOInitializer> PrecachedPSOs;
 };
@@ -1721,7 +1835,10 @@ public:
 	{
 		return GET_STATFNAME(STAT_HighPriorityComputePSOPrecacheRequests);
 	}
-
+	static const FName GetHighestPriorityCompileStatName()
+	{
+		return GET_STATFNAME(STAT_HighestPriorityComputePSOPrecacheRequests);
+	}
 	static FORCEINLINE bool PipelineStateInitializerMatch(const FPrecacheComputeInitializer& ComputeShaderInitializerA, const FPrecacheComputeInitializer& ComputeShaderInitializerB)
 	{
 		// todo: would be good/more robust to have the CS equivalent of RHIMatchPrecachePSOInitializers instead of relying on pointers
@@ -1762,7 +1879,10 @@ public:
 	{
 		return GET_STATFNAME(STAT_HighPriorityGraphicsPSOPrecacheRequests);
 	}
-
+	static const FName GetHighestPriorityCompileStatName()
+	{
+		return GET_STATFNAME(STAT_HighestPriorityGraphicsPSOPrecacheRequests);
+	}
 	FPrecacheGraphicsPipelineCache() : TPrecachePipelineCacheBase(FPSOPrecacheRequestID::EType::Graphics) {}
 	FPSOPrecacheRequestResult PrecacheGraphicsPipelineState(const FGraphicsPipelineStateInitializer& Initializer);
 
@@ -2107,7 +2227,7 @@ public:
 		CompilePSO();
 	}
 
-	void CompilePSO()
+	void CompilePSO(const FGraphicsPipelineStateInitializer::EPSOPrecacheCompileType* OptionalPriorityOverride = nullptr)
 	{
 		LLM_SCOPE(ELLMTag::PSO);
 		FTaskTagScope Scope(ETaskTag::EParallelRenderingThread);
@@ -2186,6 +2306,11 @@ public:
 				    bSkipCreation = true;
 			    }
 
+			    if(OptionalPriorityOverride)
+			    {
+				    Initializer.PrecacheCompileType = (uint32)FMath::Clamp((uint32)*OptionalPriorityOverride, (uint32)FGraphicsPipelineStateInitializer::EPSOPrecacheCompileType::MinPri, (uint32)FGraphicsPipelineStateInitializer::EPSOPrecacheCompileType::MaxPri);
+			    }
+
 				FGraphicsPipelineState* GfxPipeline = static_cast<FGraphicsPipelineState*>(Pipeline);
 
 				uint64 StartTime = FPlatformTime::Cycles64();
@@ -2831,24 +2956,6 @@ FRHIComputePipelineState* ExecuteSetComputePipelineState(FComputePipelineState*
 	return PipelineState->RHIPipeline;
 }
 
-#if PSO_TRACK_CACHE_STATS
-
-static std::atomic_uint64_t TotalPrecompileWaitTime;
-static std::atomic_int64_t TotalNumPrecompileJobs;
-static std::atomic_int64_t TotalNumPrecompileJobsCompleted;
-
-void StatsStartPrecompile()
-{
-	TotalNumPrecompileJobs++;
-}
-
-void StatsEndPrecompile(uint64 TimeToComplete)
-{
-	TotalPrecompileWaitTime += TimeToComplete;
-	TotalNumPrecompileJobsCompleted++;
-}
-
-#endif
 
 inline void ValidateGraphicsPipelineStateInitializer(const FGraphicsPipelineStateInitializer& Initializer)
 {
@@ -2882,22 +2989,30 @@ static void InternalCreateGraphicsPipelineState(const FGraphicsPipelineStateInit
 			// Here, PSO precompiles use a separate thread pool.
 			// Note that we do not add precompile tasks as cmdlist prerequisites.
 			TUniquePtr<FCompilePipelineStateTask> ThreadPoolTask = MakeUnique<FCompilePipelineStateTask>(CachedState, Initializer, PSOPrecacheResult, bInImmediateCmdList, PSOCompilationDebugData);
-			uint64 StartTime = FPlatformTime::Cycles64();
-#if	PSO_TRACK_CACHE_STATS
-			StatsStartPrecompile();
-#endif
-			CachedState->PrecompileTask = MakeUnique<FAsyncTaskLimitedReschedule<FPSOPrecompileTask>>(
-				[ThreadPoolTask = MoveTemp(ThreadPoolTask)]()
-			{
-				ThreadPoolTask->CompilePSO();
-			}
-			,
-				[StartTime]()
-			{
-#if PSO_TRACK_CACHE_STATS
-				StatsEndPrecompile(FPlatformTime::Cycles64() - StartTime);
-#endif
-			}
+			CachedState->PrecompileTask = MakeUnique<FPSOPrecacheAsyncTask>(
+				[ThreadPoolTask = MoveTemp(ThreadPoolTask)](const FPSOPrecacheAsyncTask* ThisTask)
+				{
+					// Convert the task priority to PSO precompile priority.
+					// Update here as the task's priority may have changed since creation.
+					FGraphicsPipelineStateInitializer::EPSOPrecacheCompileType PriOverride;
+					switch (ThisTask->GetPriority())
+					{
+						case EQueuedWorkPriority::Blocking:
+						case EQueuedWorkPriority::Highest:
+							PriOverride = FGraphicsPipelineStateInitializer::EPSOPrecacheCompileType::MaxPri;
+						break;
+						case EQueuedWorkPriority::High:
+						case EQueuedWorkPriority::Normal:
+							PriOverride = FGraphicsPipelineStateInitializer::EPSOPrecacheCompileType::NormalPri;
+						break;
+						default: checkNoEntry(); [[fallthrough]];
+						case EQueuedWorkPriority::Low:
+						case EQueuedWorkPriority::Lowest:
+							PriOverride = FGraphicsPipelineStateInitializer::EPSOPrecacheCompileType::MinPri;
+						break;
+					}
+					ThreadPoolTask->CompilePSO(&PriOverride);
+				}
 			);
 			CachedState->PrecompileTask->StartBackgroundTask(&GPSOPrecacheThreadPool.Get(), Initializer.bFromPSOFileCache ? EQueuedWorkPriority::Normal : EQueuedWorkPriority::Low);
 		}
@@ -3171,22 +3286,11 @@ void FPrecacheComputePipelineCache::OnNewPipelineStateCreated(const FPrecacheCom
 			// Here, PSO precompiles use a separate thread pool.
 			// Note that we do not add precompile tasks as cmdlist prerequisites.
 			TUniquePtr<FCompilePipelineStateTask> ThreadPoolTask = MakeUnique<FCompilePipelineStateTask>(CachedState, GraphicsPipelineStateInitializer, EPSOPrecacheResult::Active, false, PSOCompilationDebugData);
-			uint64 StartTime = FPlatformTime::Cycles64();
-#if	PSO_TRACK_CACHE_STATS
-			StatsStartPrecompile();
-#endif
-			CachedState->PrecompileTask = MakeUnique<FAsyncTaskLimitedReschedule<FPSOPrecompileTask>>(
-				[ThreadPoolTask = MoveTemp(ThreadPoolTask)]()
-			{
-				ThreadPoolTask->CompilePSO();
-			}
-			,
-				[StartTime]()
-			{
-#if PSO_TRACK_CACHE_STATS
-				StatsEndPrecompile(FPlatformTime::Cycles64() - StartTime);
-#endif
-			}
+			CachedState->PrecompileTask = MakeUnique<FPSOPrecacheAsyncTask>(
+				[ThreadPoolTask = MoveTemp(ThreadPoolTask)](const FPSOPrecacheAsyncTask* ThisTask)
+				{
+					ThreadPoolTask->CompilePSO();
+				}
 			);
 			CachedState->PrecompileTask->StartBackgroundTask(&GPSOPrecacheThreadPool.Get(), EQueuedWorkPriority::Normal);
 		}
@@ -3314,21 +3418,33 @@ bool PipelineStateCache::IsPrecaching()
 	return GPrecacheGraphicsPipelineCache.IsPrecaching() || GPrecacheComputePipelineCache.IsPrecaching();
 }
 
-void PipelineStateCache::BoostPrecachePriority(const FPSOPrecacheRequestID& PSOPrecacheRequestID)
+void PipelineStateCache::BoostPrecachePriority(EPSOPrecachePriority PSOPrecachePriority, const FPSOPrecacheRequestID& PSOPrecacheRequestID)
 {
 	if (IsPSOPrecachingEnabled())
 	{
+		PSOPrecachePriority = GForceHighToHighestPri && PSOPrecachePriority == EPSOPrecachePriority::High ? EPSOPrecachePriority::Highest : PSOPrecachePriority;
+
 		if (PSOPrecacheRequestID.GetType() == FPSOPrecacheRequestID::EType::Graphics)
 		{
-			GPrecacheGraphicsPipelineCache.BoostPriority(PSOPrecacheRequestID);
+			GPrecacheGraphicsPipelineCache.BoostPriority(PSOPrecachePriority, PSOPrecacheRequestID);
 		}
 		else
 		{
-			GPrecacheComputePipelineCache.BoostPriority(PSOPrecacheRequestID);
+			GPrecacheComputePipelineCache.BoostPriority(PSOPrecachePriority, PSOPrecacheRequestID);
 		}
 	}
 }
 
+
+void PipelineStateCache::PrecachePSOsBoostToHighestPriority(bool bForceHighest)
+{
+	UE_LOG(LogRHI, Log, TEXT("PipelineStateCache: PSO precaching %s highest priority boost"), bForceHighest ? TEXT("enabling") : TEXT("disabling"));
+	GForceHighToHighestPri = bForceHighest;
+#if PSO_TRACK_CACHE_STATS
+	DumpPipelineCacheStats();
+#endif
+}
+
 uint32 PipelineStateCache::NumActivePrecacheRequests()
 {
 	if (!IsPSOPrecachingEnabled())
@@ -3412,8 +3528,21 @@ void DumpPipelineCacheStats()
 	}
 
 	UE_LOG(LogRHI, Log, TEXT("Have %d GraphicsPipeline entries"), NumCachedItems);
-	UE_LOG(LogRHI, Log, TEXT("Threadpool precompile: %d GraphicsPipeline in flight, %d Jobs started, %d completed"), GPipelinePrecompileTasksInFlight.load(), TotalNumPrecompileJobs.load(), TotalNumPrecompileJobsCompleted.load());
-	UE_LOG(LogRHI, Log, TEXT("Threadpool precompile: %f s avg precompile time"), FPlatformTime::GetSecondsPerCycle64() * (TotalPrecompileWaitTime.load()/FMath::Max((int64_t)1, TotalNumPrecompileJobsCompleted.load())));
+	for (int i = 0; i < (int)EQueuedWorkPriority::Count; i++)
+	{
+		if(TotalPrecompileCompleteTime[i].load() > 0)
+		{
+			UE_LOG(LogRHI, Log, TEXT("Threadpool precompile: pri %d: %d GraphicsPipeline in flight, %d Jobs started, %d completed"), i, GPipelinePrecompileTasksInFlight.load(), TotalNumPrecompileJobs[i].load(), TotalNumPrecompileJobsCompleted[i].load());
+			UE_LOG(LogRHI, Log, TEXT("Threadpool precompile: pri %d: %f s avg precompile time"), i, FPlatformTime::GetSecondsPerCycle64() * (TotalPrecompileCompleteTime[i].load() / FMath::Max((int64_t)1, TotalNumPrecompileJobsCompleted[i].load())));
+			UE_LOG(LogRHI, Log, TEXT("Threadpool precompile: pri %d: %f s max precompile time"), i, FPlatformTime::GetSecondsPerCycle64() * (MaxPrecompileJobTime[i].load()));
+
+			UE_LOG(LogRHI, Log, TEXT("Threadpool precompile: pri %d: %f s avg precompile compile time"), i, FPlatformTime::GetSecondsPerCycle64() * (TotalPrecompileCompileTime[i].load() / FMath::Max((int64_t)1, TotalNumPrecompileJobsCompleted[i].load())));
+			UE_LOG(LogRHI, Log, TEXT("Threadpool precompile: pri %d: %f s max precompile compile time"), i, FPlatformTime::GetSecondsPerCycle64() * (MaxPrecompileTimeToCompile[i].load()));
+
+			UE_LOG(LogRHI, Log, TEXT("Threadpool precompile: pri %d: %f s avg precompile latency time"), i, FPlatformTime::GetSecondsPerCycle64() * (TotalPrecompileTimeToBegin[i].load() / FMath::Max((int64_t)1, TotalNumPrecompileJobsCompleted[i].load())));
+			UE_LOG(LogRHI, Log, TEXT("Threadpool precompile: pri %d: %f s max precompile latency time"), i, FPlatformTime::GetSecondsPerCycle64() * (MaxPrecompileTimeToBegin[i].load()));
+		}
+	}
 
 	UE_LOG(LogRHI, Log, TEXT("Secs Used: Min=%.02f, Max=%.02f, Avg=%.02f. %d used in last 30 secs"), MinTime, MaxTime, TotalTime / NumCachedItems, NumUsedLastMin);
 	UE_LOG(LogRHI, Log, TEXT("Frames Used: Min=%d, Max=%d, Avg=%d"), MinFrames, MaxFrames, TotalFrames / NumCachedItems);
diff --git a/Engine/Source/Runtime/RHI/Public/PipelineStateCache.h b/Engine/Source/Runtime/RHI/Public/PipelineStateCache.h
--- a/Engine/Source/Runtime/RHI/Public/PipelineStateCache.h
+++ b/Engine/Source/Runtime/RHI/Public/PipelineStateCache.h
@@ -25,6 +25,16 @@ enum class EApplyRendertargetOption : int
 
 ENUM_CLASS_FLAGS(EApplyRendertargetOption);
 
+/**
+ * PSO Precache request priority
+ */
+enum class EPSOPrecachePriority : uint8
+{
+	Medium,
+	High,
+	Highest,
+};
+
 enum class ERayTracingPipelineCacheFlags : uint8
 {
 	// Query the pipeline cache, create pipeline if necessary.
@@ -177,11 +187,14 @@ namespace PipelineStateCache
 	extern RHI_API bool						IsPrecaching();
 
 	/* Boost the priority of the given PSO request ID */
-	extern RHI_API void						BoostPrecachePriority(const FPSOPrecacheRequestID& PSOPrecacheRequestID);
+	extern RHI_API void						BoostPrecachePriority(EPSOPrecachePriority PSOPrecachePriority, const FPSOPrecacheRequestID& PSOPrecacheRequestID);
 
 	/* Return number of active or pending PSO precache requests */
 	extern RHI_API uint32					NumActivePrecacheRequests();
 
+	/* Set all subsequent high priority requests to highest priority, useful in non-interactive scenarios where maximum PSO throughput is preferable. */
+	extern RHI_API void						PrecachePSOsBoostToHighestPriority(bool bForceHighest);
+
 	/* Reset the PSO hitch tracking counters */
 	extern RHI_API void						ResetPSOHitchTrackingStats();
 }
diff --git a/Engine/Source/Runtime/Renderer/Private/PrimitiveSceneInfo.cpp b/Engine/Source/Runtime/Renderer/Private/PrimitiveSceneInfo.cpp
--- a/Engine/Source/Runtime/Renderer/Private/PrimitiveSceneInfo.cpp
+++ b/Engine/Source/Runtime/Renderer/Private/PrimitiveSceneInfo.cpp
@@ -2394,9 +2394,12 @@ void FPrimitiveSceneInfo::RemoveCachedReflectionCaptures()
 	bNeedsCachedReflectionCaptureUpdate = true;
 }
 
-void FPrimitiveSceneInfo::UpdateComponentLastRenderTime(float CurrentWorldTime, bool bUpdateLastRenderTimeOnScreen) const
+void FPrimitiveSceneInfo::UpdateComponentLastRenderTime(float CurrentWorldTime, bool bUpdateLastRenderTimeOnScreen)
 {
 	SceneData->SetLastRenderTime(CurrentWorldTime, bUpdateLastRenderTimeOnScreen);
+#if UE_WITH_PSO_PRECACHING
+	Proxy->BoostPrecachedPSORequestsOnDraw();
+#endif
 }
 
 void FPrimitiveOctreeSemantics::SetOctreeNodeIndex(const FPrimitiveSceneInfoCompact& Element, FOctreeElementId2 Id)
diff --git a/Engine/Source/Runtime/Renderer/Public/PrimitiveSceneInfo.h b/Engine/Source/Runtime/Renderer/Public/PrimitiveSceneInfo.h
--- a/Engine/Source/Runtime/Renderer/Public/PrimitiveSceneInfo.h
+++ b/Engine/Source/Runtime/Renderer/Public/PrimitiveSceneInfo.h
@@ -535,7 +535,7 @@ public:
 	void RemoveCachedReflectionCaptures();
 
 	/** Helper function for writing out to the last render times to the game thread */
-	void UpdateComponentLastRenderTime(float CurrentWorldTime, bool bUpdateLastRenderTimeOnScreen) const;
+	void UpdateComponentLastRenderTime(float CurrentWorldTime, bool bUpdateLastRenderTimeOnScreen);
 
 	/** Updates static lighting uniform buffer, returns the number of entries needed for GPUScene */
 	RENDERER_API int32 UpdateStaticLightingBuffer();
