Simplifying new RHI_WANT_RESOURCE_INFO logic to make sure we include all configurations that do need it.

diff --git a/Engine/Source/Runtime/RHI/RHI.Build.cs b/Engine/Source/Runtime/RHI/RHI.Build.cs
--- a/Engine/Source/Runtime/RHI/RHI.Build.cs
+++ b/Engine/Source/Runtime/RHI/RHI.Build.cs
@@ -20,11 +20,8 @@ public class RHI : ModuleRules
 
 			if (Target.Type != TargetRules.TargetType.Server)   // Dedicated servers should skip loading everything but NullDrv
 			{
-				// Always disable for Shipping, always enable for Development and DebugGame,
-				// let the user decide to use for Test
-				if (Target.Configuration == UnrealTargetConfiguration.DebugGame ||
-					Target.Configuration == UnrealTargetConfiguration.Development ||
-					(Target.Configuration == UnrealTargetConfiguration.Test && Target.bTrackRHIResourceInfoForTest))
+				// Always disable for Shipping builds. Disable by default in Test builds but allow the target to force enable it.
+				if (Target.Configuration != UnrealTargetConfiguration.Shipping && (Target.Configuration != UnrealTargetConfiguration.Test || Target.bTrackRHIResourceInfoForTest))
 				{
 					PublicDefinitions.Add("RHI_WANT_RESOURCE_INFO=1");
 				}
