By default RHI Graphics access use the light weight pixel barrier. If you need the access a resource from another stage (ex: VS), use SRVGraphicsNonPixel.

Improve DynRes, on some console (0.3%).

diff --git a/Engine/Source/Runtime/RHI/Public/RHIAccess.h b/Engine/Source/Runtime/RHI/Public/RHIAccess.h
--- a/Engine/Source/Runtime/RHI/Public/RHIAccess.h
+++ b/Engine/Source/Runtime/RHI/Public/RHIAccess.h
@@ -50,8 +50,9 @@ enum class ERHIAccess : uint32
 	None = Unknown,
 	Mask = (Last << 1) - 1,
 
-	// Graphics is a combination of pixel and non-pixel
-	SRVGraphics = SRVGraphicsPixel | SRVGraphicsNonPixel,
+	// By default Graphics use the light weight pixel barrier. 
+	// If you need the access a resource from another stage (ex: VS), use SRVGraphicsNonPixel.
+	SRVGraphics = SRVGraphicsPixel,
 
 	// A mask of the two possible SRV states
 	SRVMask = SRVCompute | SRVGraphics,
