Revert CL34286796+CL34287621. Will need more testing and changes than what I expected.Will rework the change, since it a 0.3% dynres win on some platforms.

#okforgithug public

diff --git a/Engine/Source/Runtime/RHI/Public/RHIAccess.h b/Engine/Source/Runtime/RHI/Public/RHIAccess.h
--- a/Engine/Source/Runtime/RHI/Public/RHIAccess.h
+++ b/Engine/Source/Runtime/RHI/Public/RHIAccess.h
@@ -50,9 +50,8 @@ enum class ERHIAccess : uint32
 	None = Unknown,
 	Mask = (Last << 1) - 1,
 
-	// By default Graphics use the light weight pixel barrier. 
-	// If you need the access a resource from another stage (ex: VS), use SRVGraphicsNonPixel.
-	SRVGraphics = SRVGraphicsPixel,
+	// Graphics is a combination of pixel and non-pixel
+	SRVGraphics = SRVGraphicsPixel | SRVGraphicsNonPixel,
 
 	// A mask of the two possible SRV states
 	SRVMask = SRVCompute | SRVGraphics,
diff --git a/Engine/Source/Runtime/VulkanRHI/Private/VulkanBarriers.cpp b/Engine/Source/Runtime/VulkanRHI/Private/VulkanBarriers.cpp
--- a/Engine/Source/Runtime/VulkanRHI/Private/VulkanBarriers.cpp
+++ b/Engine/Source/Runtime/VulkanRHI/Private/VulkanBarriers.cpp
@@ -1976,6 +1976,7 @@ VkImageLayout FVulkanLayoutManager::GetDefaultLayout(FVulkanCmdBuffer* CmdBuffer
 	case ERHIAccess::SRVCompute:
 	case ERHIAccess::SRVGraphics:
 	case ERHIAccess::SRVGraphicsNonPixel:
+	case ERHIAccess::SRVGraphicsPixel:
 	case ERHIAccess::SRVMask:
 	{
 		if (VulkanTexture.IsDepthOrStencilAspect())
