Deprecating the RHICalcTextureXXXPlatformSize functions in favor of RHICalcTexturePlatformSize.
Deprecating the Texture members in FRHIResourceCreateInfo.

diff --git a/Engine/Source/Runtime/Engine/Classes/Engine/Texture2DArray.cpp b/Engine/Source/Runtime/Engine/Classes/Engine/Texture2DArray.cpp
--- a/Engine/Source/Runtime/Engine/Classes/Engine/Texture2DArray.cpp
+++ b/Engine/Source/Runtime/Engine/Classes/Engine/Texture2DArray.cpp
@@ -241,8 +241,14 @@ uint32 UTexture2DArray::CalcTextureMemorySize(int32 MipCount) const
 
 			// Must be consistent with the logic in FTexture2DResource::InitRHI
 			const FIntPoint MipExtents = CalcMipMapExtent(GetSizeX(), GetSizeY(), Format, FirstMip);
-			uint32 TextureAlign = 0;
-			Size = (uint32)RHICalcTexture2DArrayPlatformSize(MipExtents.X, MipExtents.Y, GetArraySize(), Format, FMath::Max(1, MipCount), 1, Flags, FRHIResourceCreateInfo(GetPlatformData()->GetExtData()), TextureAlign);
+
+			const FRHITextureDesc Desc =
+				FRHITextureCreateDesc::Create2DArray(TEXT("Temp"), MipExtents, GetArraySize(), Format)
+				.SetNumMips(FMath::Max(1, MipCount))
+				.SetFlags(Flags)
+				.SetExtData(GetPlatformData()->GetExtData());
+
+			Size = RHICalcTexturePlatformSize(Desc).Size;
 		}
 	}
 	return Size;
diff --git a/Engine/Source/Runtime/Engine/Private/Rendering/Texture2DArrayResource.cpp b/Engine/Source/Runtime/Engine/Private/Rendering/Texture2DArrayResource.cpp
--- a/Engine/Source/Runtime/Engine/Private/Rendering/Texture2DArrayResource.cpp
+++ b/Engine/Source/Runtime/Engine/Private/Rendering/Texture2DArrayResource.cpp
@@ -93,13 +93,17 @@ uint64 FTexture2DArrayResource::GetPlatformMipsSize(uint32 NumMips) const
 	if (PlatformData && NumMips > 0)
 	{
 		const FIntPoint MipExtents = CalcMipMapExtent(SizeX, SizeY, PixelFormat, State.LODCountToFirstLODIdx(NumMips));
-		uint32 TextureAlign = 0;
-		return RHICalcTexture2DArrayPlatformSize(MipExtents.X, MipExtents.Y, SizeZ, PixelFormat, NumMips, 1, CreationFlags, FRHIResourceCreateInfo(PlatformData->GetExtData()), TextureAlign);
-	}
-	else
-	{
-		return 0;
+
+		const FRHITextureDesc Desc =
+			FRHITextureCreateDesc::Create2DArray(TEXT("Temp"), MipExtents, SizeZ, PixelFormat)
+			.SetNumMips(NumMips)
+			.SetFlags(CreationFlags)
+			.SetExtData(PlatformData->GetExtData());
+
+		return RHICalcTexturePlatformSize(Desc).Size;
 	}
+
+	return 0;
 }
 
 void FTexture2DArrayResource::InitRHI(FRHICommandListBase& RHICmdList)
diff --git a/Engine/Source/Runtime/Engine/Private/Rendering/Texture2DResource.cpp b/Engine/Source/Runtime/Engine/Private/Rendering/Texture2DResource.cpp
--- a/Engine/Source/Runtime/Engine/Private/Rendering/Texture2DResource.cpp
+++ b/Engine/Source/Runtime/Engine/Private/Rendering/Texture2DResource.cpp
@@ -205,22 +205,36 @@ uint64 FTexture2DResource::GetPlatformMipsSize(uint32 NumMips) const
 		static TConsoleVariableData<int32>* CVarReducedMode = IConsoleManager::Get().FindTConsoleVariableDataInt(TEXT("r.VirtualTextureReducedMemory"));
 		check(CVarReducedMode);
 
+		FRHICalcTextureSizeResult Result{};
+
 		uint32 TextureAlign = 0;
 		// Must be consistent with the logic in FTexture2DResource::InitRHI
 		if (bUsePartiallyResidentMips && (!CVarReducedMode->GetValueOnRenderThread() || NumMips > State.NumNonStreamingLODs))
 		{
-			return RHICalcVMTexture2DPlatformSize(SizeX, SizeY, PixelFormat, State.MaxNumLODs, State.LODCountToFirstLODIdx(NumMips), 1, CreationFlags | TexCreate_Virtual, TextureAlign);
+			const FRHITextureDesc Desc =
+				FRHITextureCreateDesc::Create2D(TEXT("Temp"), SizeX, SizeY, PixelFormat)
+				.SetNumMips(State.MaxNumLODs)
+				.SetFlags(CreationFlags | TexCreate_Virtual);
+
+			Result = RHICalcTexturePlatformSize(Desc, State.LODCountToFirstLODIdx(NumMips));
 		}
 		else
 		{
 			const FIntPoint MipExtents = CalcMipMapExtent(SizeX, SizeY, PixelFormat, State.RequestedFirstLODIdx());
-			return RHICalcTexture2DPlatformSize(MipExtents.X, MipExtents.Y, PixelFormat, NumMips, 1, CreationFlags, FRHIResourceCreateInfo(PlatformData->GetExtData()), TextureAlign);
+
+			const FRHITextureDesc Desc =
+				FRHITextureCreateDesc::Create2D(TEXT("Temp"), MipExtents, PixelFormat)
+				.SetNumMips(NumMips)
+				.SetFlags(CreationFlags)
+				.SetExtData(PlatformData->GetExtData());
+
+			Result = RHICalcTexturePlatformSize(Desc);
 		}
+
+		return Result.Size;
 	}
-	else
-	{
-		return 0;
-	}
+
+	return 0;
 }
 
 uint32 FTexture2DResource::CalculateTightPackedMipSize(int32 MipSizeX,int32 MipSizeY,EPixelFormat PixelFormat,
diff --git a/Engine/Source/Runtime/Engine/Private/Rendering/Texture3DResource.cpp b/Engine/Source/Runtime/Engine/Private/Rendering/Texture3DResource.cpp
--- a/Engine/Source/Runtime/Engine/Private/Rendering/Texture3DResource.cpp
+++ b/Engine/Source/Runtime/Engine/Private/Rendering/Texture3DResource.cpp
@@ -84,10 +84,14 @@ FTexture3DResource::FTexture3DResource(UVolumeTexture* InOwner, const FStreamabl
 				uint32 MipExtentZ = 0;
 				CalcMipMapExtent3D(SizeX, SizeY, SizeZ, PixelFormat, State.RequestedFirstLODIdx(), MipExtentX, MipExtentY, MipExtentZ);
 
+				const FRHITextureDesc Desc =
+					FRHITextureCreateDesc::Create3D(TEXT("Temp"), MipExtentX, MipExtentY, MipExtentZ, PixelFormat)
+					.SetNumMips(State.NumRequestedLODs)
+					.SetFlags(CreationFlags)
+					.SetExtData(PlatformData->GetExtData());
+
+				const uint64 PlatformMipSize = RHICalcTexturePlatformSize(Desc).Size;
 
-				uint32 TextureAlign = 0;
-				uint64 PlatformMipSize = RHICalcTexture3DPlatformSize(MipExtentX, MipExtentY, MipExtentZ, (EPixelFormat)PixelFormat, State.NumRequestedLODs, CreationFlags, FRHIResourceCreateInfo(PlatformData->GetExtData()), TextureAlign);
-				
 				UE_LOG(LogTexture,Verbose,TEXT("FTexture3DResource::FTexture3DResource %d : %dx%dx%d : MipBytes=%d"),
 					MipIndex,MipExtentX,MipExtentY,MipExtentZ,
 					(int)PlatformMipSize);
@@ -193,13 +197,16 @@ uint64 FTexture3DResource::GetPlatformMipsSize(uint32 NumMips) const
 		uint32 MipExtentZ = 0;
 		CalcMipMapExtent3D(SizeX, SizeY, SizeZ, PixelFormat, State.LODCountToFirstLODIdx(NumMips), MipExtentX, MipExtentY, MipExtentZ);
 
-		uint32 TextureAlign = 0;
-		return RHICalcTexture3DPlatformSize(MipExtentX, MipExtentY, MipExtentZ, PixelFormat, NumMips, CreationFlags, FRHIResourceCreateInfo(PlatformData->GetExtData()), TextureAlign);
-	}
-	else
-	{
-		return 0;
+		const FRHITextureDesc Desc =
+			FRHITextureCreateDesc::Create3D(TEXT("Temp"), MipExtentX, MipExtentY, MipExtentZ, PixelFormat)
+			.SetNumMips(NumMips)
+			.SetFlags(CreationFlags)
+			.SetExtData(PlatformData->GetExtData());
+
+		return RHICalcTexturePlatformSize(Desc).Size;
 	}
+
+	return 0;
 }
 
 void FTexture3DResource::InitRHI(FRHICommandListBase& RHICmdList)
diff --git a/Engine/Source/Runtime/Engine/Private/Streaming/VolumeTextureStreaming.cpp b/Engine/Source/Runtime/Engine/Private/Streaming/VolumeTextureStreaming.cpp
--- a/Engine/Source/Runtime/Engine/Private/Streaming/VolumeTextureStreaming.cpp
+++ b/Engine/Source/Runtime/Engine/Private/Streaming/VolumeTextureStreaming.cpp
@@ -38,8 +38,13 @@ bool FVolumeTextureMipAllocator_Reallocate::AllocateMips(const FTextureUpdateCon
 		MipInfo.SizeY = OwnerMip.SizeY;
 		MipInfo.SizeZ = OwnerMip.SizeZ;
 
-		uint32 TextureAlign = 0;
-		MipInfo.DataSize = RHICalcTexture3DPlatformSize(MipInfo.SizeX, MipInfo.SizeY, MipInfo.SizeZ, MipInfo.Format, 1, Context.Resource->GetCreationFlags(), FRHIResourceCreateInfo(Context.Resource->GetExtData()), TextureAlign);
+		const FRHITextureDesc Desc =
+			FRHITextureCreateDesc::Create3D(TEXT("Temp"), MipInfo.SizeX, MipInfo.SizeY, MipInfo.SizeZ, MipInfo.Format)
+			.SetFlags(Context.Resource->GetCreationFlags())
+			.SetExtData(Context.Resource->GetExtData());
+
+		MipInfo.DataSize = RHICalcTexturePlatformSize(Desc).Size;
+
 		StreamedInMipData.GetMipSize()[MipIdx] = MipInfo.DataSize;
 
 		// When initializing the texture with it's bulk data, we proceed with a single allocation.
diff --git a/Engine/Source/Runtime/Engine/Private/Texture2D.cpp b/Engine/Source/Runtime/Engine/Private/Texture2D.cpp
--- a/Engine/Source/Runtime/Engine/Private/Texture2D.cpp
+++ b/Engine/Source/Runtime/Engine/Private/Texture2D.cpp
@@ -843,18 +843,28 @@ int32 UTexture2D::CalcTextureMemorySize( int32 MipCount ) const
 		const int32 NumMips = GetNumMips();
 		const int32 FirstMip = FMath::Max(0, NumMips - MipCount);
 		const EPixelFormat Format = GetPixelFormat();
-		uint32 TextureAlign;
 
 		// Must be consistent with the logic in FTexture2DResource::InitRHI
 		if (IsStreamable() && bCanUsePartiallyResidentMips && (!CVarReducedMode->GetValueOnAnyThread() || MipCount > UTexture2D::GetMinTextureResidentMipCount()))
 		{
-			TexCreateFlags |= TexCreate_Virtual;
-			Size = (int32)RHICalcVMTexture2DPlatformSize(SizeX, SizeY, Format, NumMips, FirstMip, 1, TexCreateFlags, TextureAlign);
+			const FRHITextureDesc Desc =
+				FRHITextureCreateDesc::Create2D(TEXT("Temp"), SizeX, SizeY, Format)
+				.SetNumMips(NumMips)
+				.SetFlags(TexCreateFlags | TexCreate_Virtual);
+
+			Size = RHICalcTexturePlatformSize(Desc, FirstMip).Size;
 		}
 		else
 		{
 			const FIntPoint MipExtents = CalcMipMapExtent(SizeX, SizeY, Format, FirstMip);
-			Size = (int32)RHICalcTexture2DPlatformSize(MipExtents.X, MipExtents.Y, Format, FMath::Max(1, MipCount), 1, TexCreateFlags, FRHIResourceCreateInfo(GetPlatformData()->GetExtData()), TextureAlign);
+
+			const FRHITextureDesc Desc =
+				FRHITextureCreateDesc::Create2D(TEXT("Temp"), MipExtents, Format)
+				.SetNumMips(FMath::Max(1, MipCount))
+				.SetFlags(TexCreateFlags)
+				.SetExtData(GetPlatformData()->GetExtData());
+
+			Size = RHICalcTexturePlatformSize(Desc).Size;
 		}
 	}
 	return Size;
diff --git a/Engine/Source/Runtime/Engine/Private/TextureCube.cpp b/Engine/Source/Runtime/Engine/Private/TextureCube.cpp
--- a/Engine/Source/Runtime/Engine/Private/TextureCube.cpp
+++ b/Engine/Source/Runtime/Engine/Private/TextureCube.cpp
@@ -249,10 +249,13 @@ uint32 UTextureCube::CalcTextureMemorySize( int32 MipCount ) const
 		// Figure out what the first mip to use is.
 		int32 FirstMip	= FMath::Max( 0, NumMips - MipCount );		
 		FIntPoint MipExtents = CalcMipMapExtent(SizeX, SizeY, Format, FirstMip);
-		
-		uint32 TextureAlign = 0;
-		uint64 TextureSize = RHICalcTextureCubePlatformSize(MipExtents.X, Format, FMath::Max( 1, MipCount ), TexCreate_None, FRHIResourceCreateInfo(GetPlatformData()->GetExtData()), TextureAlign);
-		Size = (uint32)TextureSize;
+
+		const FRHITextureDesc Desc =
+			FRHITextureCreateDesc::CreateCube(TEXT("Temp"), MipExtents.X, Format)
+			.SetNumMips(FMath::Max(1, MipCount))
+			.SetExtData(GetPlatformData()->GetExtData());
+
+		Size = RHICalcTexturePlatformSize(Desc).Size;
 	}
 	return Size;
 }
diff --git a/Engine/Source/Runtime/Engine/Private/TextureCubeArray.cpp b/Engine/Source/Runtime/Engine/Private/TextureCubeArray.cpp
--- a/Engine/Source/Runtime/Engine/Private/TextureCubeArray.cpp
+++ b/Engine/Source/Runtime/Engine/Private/TextureCubeArray.cpp
@@ -504,10 +504,12 @@ uint32 UTextureCubeArray::CalcTextureMemorySize(int32 MipCount) const
 		int32 FirstMip = FMath::Max(0, NumMips - MipCount);
 		FIntPoint MipExtents = CalcMipMapExtent(SizeX, SizeY, Format, FirstMip);
 
-		// TODO add RHICalcTextureCubeArrayPlatformSize
-		uint32 TextureAlign = 0;
-		uint64 TextureSize = RHICalcTextureCubePlatformSize(MipExtents.X, Format, FMath::Max(1, MipCount), TexCreate_None, FRHIResourceCreateInfo(GetPlatformData()->GetExtData()), TextureAlign) * ArraySize;
-		Size = (uint32)TextureSize;
+		const FRHITextureDesc Desc =
+			FRHITextureCreateDesc::CreateCubeArray(TEXT("Temp"), MipExtents.X, ArraySize, Format)
+			.SetNumMips(FMath::Max(1, MipCount))
+			.SetExtData(GetPlatformData()->GetExtData());
+
+		Size = RHICalcTexturePlatformSize(Desc).Size;
 	}
 	return Size;
 }
diff --git a/Engine/Source/Runtime/Engine/Private/VolumeTexture.cpp b/Engine/Source/Runtime/Engine/Private/VolumeTexture.cpp
--- a/Engine/Source/Runtime/Engine/Private/VolumeTexture.cpp
+++ b/Engine/Source/Runtime/Engine/Private/VolumeTexture.cpp
@@ -412,8 +412,13 @@ uint32 UVolumeTexture::CalcTextureMemorySize(int32 MipCount) const
 			uint32 SizeZ = 0;
 			CalcMipMapExtent3D(GetSizeX(), GetSizeY(), GetSizeZ(), Format, FMath::Max<int32>(0, GetNumMips() - MipCount), SizeX, SizeY, SizeZ);
 
-			uint32 TextureAlign = 0;
-			Size = (uint32)RHICalcTexture3DPlatformSize(SizeX, SizeY, SizeZ, Format, FMath::Max(1, MipCount), Flags, FRHIResourceCreateInfo(GetPlatformData()->GetExtData()), TextureAlign);
+			const FRHITextureDesc Desc =
+				FRHITextureCreateDesc::Create3D(TEXT("Temp"), SizeX, SizeY, SizeZ, Format)
+				.SetNumMips(FMath::Max(1, MipCount))
+				.SetFlags(Flags)
+				.SetExtData(GetPlatformData()->GetExtData());
+
+			Size = RHICalcTexturePlatformSize(Desc).Size;
 		}
 	}
 	return Size;
diff --git a/Engine/Source/Runtime/RHI/Public/DynamicRHI.h b/Engine/Source/Runtime/RHI/Public/DynamicRHI.h
--- a/Engine/Source/Runtime/RHI/Public/DynamicRHI.h
+++ b/Engine/Source/Runtime/RHI/Public/DynamicRHI.h
@@ -107,10 +107,20 @@ private:
 	TArray<TRefCountPtr<FRHIRenderQuery>> Queries;
 };
 
+struct FRHICalcTextureSizeResult
+{
+	// The total size of the texture, in bytes.
+	uint64 Size;
+
+	// The required address alignment for the texture.
+	uint32 Align;
+};
+
 /** The interface which is implemented by the dynamically bound RHI. */
 class FDynamicRHI
 {
 public:
+	using FRHICalcTextureSizeResult = ::FRHICalcTextureSizeResult;
 
 	/** Declare a virtual destructor, so the dynamic RHI can be deleted without knowing its type. */
 	virtual ~FDynamicRHI() {}
@@ -361,14 +371,6 @@ public:
 
 	RHI_API virtual void RHIUpdateTextureReference(FRHICommandListBase& RHICmdList, FRHITextureReference* TextureRef, FRHITexture* NewTexture);
 
-	struct FRHICalcTextureSizeResult
-	{
-		// The total size of the texture, in bytes.
-		uint64 Size;
-
-		// The required address alignment for the texture.
-		uint32 Align;
-	};
 
 	/**
 	* Computes the total GPU memory a texture resource with the specified parameters will occupy on the current RHI platform.
@@ -1159,24 +1161,25 @@ FORCEINLINE FUniformBufferRHIRef RHICreateUniformBuffer(const void* Contents, co
 	return GDynamicRHI->RHICreateUniformBuffer(Contents, Layout, Usage, Validation);
 }
 
-FORCEINLINE FDynamicRHI::FRHICalcTextureSizeResult RHICalcTexturePlatformSize(FRHITextureDesc const& Desc, uint32 FirstMipIndex = 0)
+FORCEINLINE FRHICalcTextureSizeResult RHICalcTexturePlatformSize(FRHITextureDesc const& Desc, uint32 FirstMipIndex = 0)
 {
 	if ( ! Desc.IsValid() )
 	{
 		// Invalid texture desc; return zero to indicate failure
-		FDynamicRHI::FRHICalcTextureSizeResult ZeroResult = { 0 };
+		FRHICalcTextureSizeResult ZeroResult{};
 		return ZeroResult;
 	}
 
 	return GDynamicRHI->RHICalcTexturePlatformSize(Desc, FirstMipIndex);
 }
 
-//UE_DEPRECATED(5.1, "The separate RHICalcTexture... functions for each texture type are deprecated. Use RHICalcTexturePlatformSize instead.")
+UE_DEPRECATED(5.5, "The separate RHICalcTexture... functions for each texture type are deprecated. Use RHICalcTexturePlatformSize instead.")
 FORCEINLINE uint64 RHICalcTexture2DPlatformSize(uint32 SizeX, uint32 SizeY, uint8 Format, uint32 NumMips, uint32 NumSamples, ETextureCreateFlags Flags, const FRHIResourceCreateInfo& CreateInfo, uint32& OutAlign)
 {
 	const uint16 Depth = 1;
 	const uint16 ArraySize = 1;
 
+PRAGMA_DISABLE_DEPRECATION_WARNINGS
 	FRHITextureDesc Desc(
 		ETextureDimension::Texture2D,
 		Flags,
@@ -1189,17 +1192,19 @@ FORCEINLINE uint64 RHICalcTexture2DPlatformSize(uint32 SizeX, uint32 SizeY, uint
 		(uint8)NumSamples,
 		CreateInfo.ExtData
 	);
+PRAGMA_ENABLE_DEPRECATION_WARNINGS
 
-	auto Result = RHICalcTexturePlatformSize(Desc, 0);
+	FRHICalcTextureSizeResult Result = RHICalcTexturePlatformSize(Desc, 0);
 	OutAlign = Result.Align;
 	return Result.Size;
 }
 
-//UE_DEPRECATED(5.1, "The separate RHICalcTexture... functions for each texture type are deprecated. Use RHICalcTexturePlatformSize instead.")
+UE_DEPRECATED(5.5, "The separate RHICalcTexture... functions for each texture type are deprecated. Use RHICalcTexturePlatformSize instead.")
 FORCEINLINE uint64 RHICalcTexture2DArrayPlatformSize(uint32 SizeX, uint32 SizeY, uint32 ArraySize, uint8 Format, uint32 NumMips, uint32 NumSamples, ETextureCreateFlags Flags, const FRHIResourceCreateInfo& CreateInfo, uint32& OutAlign)
 {
 	const uint16 Depth = 1;
 
+PRAGMA_DISABLE_DEPRECATION_WARNINGS
 	FRHITextureDesc Desc(
 		ETextureDimension::Texture2DArray,
 		Flags,
@@ -1212,13 +1217,14 @@ FORCEINLINE uint64 RHICalcTexture2DArrayPlatformSize(uint32 SizeX, uint32 SizeY,
 		(uint8)NumSamples,
 		CreateInfo.ExtData
 	);
+PRAGMA_ENABLE_DEPRECATION_WARNINGS
 
-	auto Result = RHICalcTexturePlatformSize(Desc, 0);
+	FRHICalcTextureSizeResult Result = RHICalcTexturePlatformSize(Desc, 0);
 	OutAlign = Result.Align;
 	return Result.Size;
 }
 
-//UE_DEPRECATED(5.1, "The separate RHICalcTexture... functions for each texture type are deprecated. Use RHICalcTexturePlatformSize instead.")
+UE_DEPRECATED(5.5, "The separate RHICalcTexture... functions for each texture type are deprecated. Use RHICalcTexturePlatformSize instead.")
 FORCEINLINE uint64 RHICalcVMTexture2DPlatformSize(uint32 Mip0Width, uint32 Mip0Height, uint8 Format, uint32 NumMips, uint32 FirstMipIdx, uint32 NumSamples, ETextureCreateFlags Flags, uint32& OutAlign)
 {
 	const uint16 Depth     = 1;
@@ -1238,17 +1244,18 @@ FORCEINLINE uint64 RHICalcVMTexture2DPlatformSize(uint32 Mip0Width, uint32 Mip0H
 		ExtData
 	);
 
-	auto Result = RHICalcTexturePlatformSize(Desc, FirstMipIdx);
+	FRHICalcTextureSizeResult Result = RHICalcTexturePlatformSize(Desc, FirstMipIdx);
 	OutAlign = Result.Align;
 	return Result.Size;
 }
 
-//UE_DEPRECATED(5.1, "The separate RHICalcTexture... functions for each texture type are deprecated. Use RHICalcTexturePlatformSize instead.")
+UE_DEPRECATED(5.5, "The separate RHICalcTexture... functions for each texture type are deprecated. Use RHICalcTexturePlatformSize instead.")
 FORCEINLINE uint64 RHICalcTexture3DPlatformSize(uint32 SizeX, uint32 SizeY, uint32 SizeZ, uint8 Format, uint32 NumMips, ETextureCreateFlags Flags, const FRHIResourceCreateInfo& CreateInfo, uint32& OutAlign)
 {
 	const uint16 ArraySize  = 1;
 	const uint8 NumSamples = 1;
 
+PRAGMA_DISABLE_DEPRECATION_WARNINGS
 	FRHITextureDesc Desc(
 		ETextureDimension::Texture3D,
 		Flags,
@@ -1261,19 +1268,21 @@ FORCEINLINE uint64 RHICalcTexture3DPlatformSize(uint32 SizeX, uint32 SizeY, uint
 		NumSamples,
 		CreateInfo.ExtData
 	);
+PRAGMA_ENABLE_DEPRECATION_WARNINGS
 
-	auto Result = RHICalcTexturePlatformSize(Desc, 0);
+	FRHICalcTextureSizeResult Result = RHICalcTexturePlatformSize(Desc, 0);
 	OutAlign = Result.Align;
 	return Result.Size;
 }
 
-//UE_DEPRECATED(5.1, "The separate RHICalcTexture... functions for each texture type are deprecated. Use RHICalcTexturePlatformSize instead.")
+UE_DEPRECATED(5.5, "The separate RHICalcTexture... functions for each texture type are deprecated. Use RHICalcTexturePlatformSize instead.")
 FORCEINLINE uint64 RHICalcTextureCubePlatformSize(uint32 Size, uint8 Format, uint32 NumMips, ETextureCreateFlags Flags, const FRHIResourceCreateInfo& CreateInfo, uint32& OutAlign)
 {
 	const uint16 Depth      = 1;
 	const uint16 ArraySize  = 1;
 	const uint8 NumSamples = 1;
 
+PRAGMA_DISABLE_DEPRECATION_WARNINGS
 	FRHITextureDesc Desc(
 		ETextureDimension::TextureCube,
 		Flags,
@@ -1286,8 +1295,9 @@ FORCEINLINE uint64 RHICalcTextureCubePlatformSize(uint32 Size, uint8 Format, uin
 		NumSamples,
 		CreateInfo.ExtData
 	);
+PRAGMA_ENABLE_DEPRECATION_WARNINGS
 
-	auto Result = RHICalcTexturePlatformSize(Desc, 0);
+	FRHICalcTextureSizeResult Result = RHICalcTexturePlatformSize(Desc, 0);
 	OutAlign = Result.Align;
 	return Result.Size;
 }
diff --git a/Engine/Source/Runtime/RHI/Public/RHICommandList.h b/Engine/Source/Runtime/RHI/Public/RHICommandList.h
--- a/Engine/Source/Runtime/RHI/Public/RHICommandList.h
+++ b/Engine/Source/Runtime/RHI/Public/RHICommandList.h
@@ -746,8 +746,7 @@ public:
 
 	FORCEINLINE FBufferRHIRef CreateVertexBuffer(uint32 Size, EBufferUsageFlags Usage, FRHIResourceCreateInfo& CreateInfo)
 	{
-		bool bHasInitialData = CreateInfo.BulkData != nullptr;
-		ERHIAccess ResourceState = RHIGetDefaultResourceState(Usage | EBufferUsageFlags::VertexBuffer, bHasInitialData);
+		ERHIAccess ResourceState = RHIGetDefaultResourceState(Usage | EBufferUsageFlags::VertexBuffer, false);
 		return CreateVertexBuffer(Size, Usage, ResourceState, CreateInfo);
 	}
 
@@ -758,8 +757,7 @@ public:
 
 	FORCEINLINE FBufferRHIRef CreateStructuredBuffer(uint32 Stride, uint32 Size, EBufferUsageFlags Usage, FRHIResourceCreateInfo& CreateInfo)
 	{
-		bool bHasInitialData = CreateInfo.BulkData != nullptr;
-		ERHIAccess ResourceState = RHIGetDefaultResourceState(Usage | EBufferUsageFlags::StructuredBuffer, bHasInitialData);
+		ERHIAccess ResourceState = RHIGetDefaultResourceState(Usage | EBufferUsageFlags::StructuredBuffer, false);
 		return CreateStructuredBuffer(Stride, Size, Usage, ResourceState, CreateInfo);
 	}
 
@@ -770,8 +768,7 @@ public:
 
 	FORCEINLINE FBufferRHIRef CreateIndexBuffer(uint32 Stride, uint32 Size, EBufferUsageFlags Usage, FRHIResourceCreateInfo& CreateInfo)
 	{
-		bool bHasInitialData = CreateInfo.BulkData != nullptr;
-		ERHIAccess ResourceState = RHIGetDefaultResourceState(Usage | EBufferUsageFlags::IndexBuffer, bHasInitialData);
+		ERHIAccess ResourceState = RHIGetDefaultResourceState(Usage | EBufferUsageFlags::IndexBuffer, false);
 		return CreateIndexBuffer(Stride, Size, Usage, ResourceState, CreateInfo);
 	}
 
diff --git a/Engine/Source/Runtime/RHI/Public/RHIResources.h b/Engine/Source/Runtime/RHI/Public/RHIResources.h
--- a/Engine/Source/Runtime/RHI/Public/RHIResources.h
+++ b/Engine/Source/Runtime/RHI/Public/RHIResources.h
@@ -354,22 +354,19 @@ struct FClearValueBinding
 struct FRHIResourceCreateInfo
 {
 	FRHIResourceCreateInfo(const TCHAR* InDebugName)
-		: BulkData(nullptr)
-		, ResourceArray(nullptr)
-		, ClearValueBinding(FLinearColor::Transparent)
-		, GPUMask(FRHIGPUMask::All())
-		, bWithoutNativeResource(false)
-		, DebugName(InDebugName)
-		, ExtData(0)
+		: DebugName(InDebugName)
 	{
 		check(InDebugName);
 	}
 
 	// for CreateTexture calls
+	UE_DEPRECATED(5.5, "Please use FRHITextureCreateDesc for creating Textures with Bulk Data")
 	FRHIResourceCreateInfo(const TCHAR* InDebugName, FResourceBulkDataInterface* InBulkData)
 		: FRHIResourceCreateInfo(InDebugName)
 	{
+PRAGMA_DISABLE_DEPRECATION_WARNINGS
 		BulkData = InBulkData;
+PRAGMA_ENABLE_DEPRECATION_WARNINGS
 	}
 
 	// for CreateBuffer calls
@@ -379,38 +376,48 @@ struct FRHIResourceCreateInfo
 		ResourceArray = InResourceArray;
 	}
 
+	UE_DEPRECATED(5.5, "Please use FRHITextureCreateDesc for creating Textures with a Clear Value Binding")
 	FRHIResourceCreateInfo(const TCHAR* InDebugName, const FClearValueBinding& InClearValueBinding)
 		: FRHIResourceCreateInfo(InDebugName)
 	{
+PRAGMA_DISABLE_DEPRECATION_WARNINGS
 		ClearValueBinding = InClearValueBinding;
+PRAGMA_ENABLE_DEPRECATION_WARNINGS
 	}
 
+	UE_DEPRECATED(5.5, "Please use FRHITextureCreateDesc for creating Textures with Ext Data")
 	FRHIResourceCreateInfo(uint32 InExtData)
 		: FRHIResourceCreateInfo(TEXT(""))
 	{
+PRAGMA_DISABLE_DEPRECATION_WARNINGS
 		ExtData = InExtData;
+PRAGMA_ENABLE_DEPRECATION_WARNINGS
 	}
 
 	FName GetTraceClassName() const									{ const static FLazyName FRHIBufferName(TEXT("FRHIBuffer")); return (ClassName == NAME_None) ? FRHIBufferName : ClassName; }
 
 	// for CreateTexture calls
-	FResourceBulkDataInterface* BulkData;
+	UE_DEPRECATED(5.5, "Please use FRHITextureCreateDesc for creating Textures with Bulk Data")
+	FResourceBulkDataInterface* BulkData = nullptr;
 
 	// for CreateBuffer calls
-	FResourceArrayUploadInterface* ResourceArray;
+	FResourceArrayUploadInterface* ResourceArray = nullptr;
 
 	// for binding clear colors to render targets.
-	FClearValueBinding ClearValueBinding;
+	UE_DEPRECATED(5.5, "Please use FRHITextureCreateDesc for creating Textures with a Clear Value Binding")
+	FClearValueBinding ClearValueBinding = FClearValueBinding(FLinearColor::Transparent);
 
 	// set of GPUs on which to create the resource
-	FRHIGPUMask GPUMask;
+	FRHIGPUMask GPUMask = FRHIGPUMask::All();
 
 	// whether to create an RHI object with no underlying resource
-	bool bWithoutNativeResource;
-	const TCHAR* DebugName;
+	bool bWithoutNativeResource = false;
 
 	// optional data that would have come from an offline cooker or whatever - general purpose
-	uint32 ExtData;
+	UE_DEPRECATED(5.5, "Please use FRHITextureCreateDesc for creating Textures with Ext Data")
+	uint32 ExtData = 0;
+
+	const TCHAR* DebugName;
 
 	FName ClassName = NAME_None;	// The owner class of FRHIBuffer used for Insight asset metadata tracing
 	FName OwnerName = NAME_None;	// The owner name used for Insight asset metadata tracing
