Use a different category for logging PSO hitching and raise verbosity for individual hitching instances.

For normal logging verbosity, just log the cumulative amount of PSO hitches every 50 hitches.

diff --git a/Engine/Source/Runtime/RHI/Private/PipelineStateCache.cpp b/Engine/Source/Runtime/RHI/Private/PipelineStateCache.cpp
--- a/Engine/Source/Runtime/RHI/Private/PipelineStateCache.cpp
+++ b/Engine/Source/Runtime/RHI/Private/PipelineStateCache.cpp
@@ -41,6 +41,8 @@ namespace PipelineStateCache
 
 CSV_DECLARE_CATEGORY_EXTERN(PSO);
 
+DEFINE_LOG_CATEGORY_STATIC(LogPSOHitching, Log, All);
+
 DECLARE_DWORD_ACCUMULATOR_STAT(TEXT("Runtime Graphics PSO Hitch Count"), STAT_RuntimeGraphicsPSOHitchCount, STATGROUP_PipelineStateCache);
 DECLARE_DWORD_ACCUMULATOR_STAT(TEXT("Runtime Compute PSO Hitch Count"), STAT_RuntimeComputePSOHitchCount, STATGROUP_PipelineStateCache);
 
@@ -197,6 +199,7 @@ static FAutoConsoleCommand DumpPipelineCmd(
 
 static int32 GraphicsPSOCreationHitchCount = 0;
 static int32 ComputePSOCreationHitchCount = 0;
+static int32 PrecachedPSOCreationHitchCount = 0;
 
 struct FPSOCompilationDebugData
 {
@@ -217,9 +220,9 @@ static inline void CheckAndUpdateHitchCountStat(FPSOPrecacheRequestID::EType PSO
 			if (PSOType == FPSOPrecacheRequestID::EType::Graphics)
 			{
 #if WITH_RHI_BREADCRUMBS
-				UE_LOG(LogRHI, Log, TEXT("Runtime graphics PSO creation hitch (%.2f msec) for %s (precache status: %s) - Breadcrumbs: %s"), PSOCreationTimeMs, *PSOCompilationDebugData.PSOCompilationEventName, LexToString(PSOPrecacheResult), PSOCompilationDebugData.BreadcrumbNode ? *PSOCompilationDebugData.BreadcrumbNode->GetFullPath() : TEXT("Unknown"));
+				UE_LOG(LogPSOHitching, Verbose, TEXT("Runtime graphics PSO creation hitch (%.2f msec) for %s (precache status: %s) - Breadcrumbs: %s"), PSOCreationTimeMs, *PSOCompilationDebugData.PSOCompilationEventName, LexToString(PSOPrecacheResult), PSOCompilationDebugData.BreadcrumbNode ? *PSOCompilationDebugData.BreadcrumbNode->GetFullPath() : TEXT("Unknown"));
 #else
-				UE_LOG(LogRHI, Log, TEXT("Runtime graphics PSO creation hitch (%.2f msec) for %s (precache status: %s)"), PSOCreationTimeMs, *PSOCompilationDebugData.PSOCompilationEventName, LexToString(PSOPrecacheResult));
+				UE_LOG(LogPSOHitching, Verbose, TEXT("Runtime graphics PSO creation hitch (%.2f msec) for %s (precache status: %s)"), PSOCreationTimeMs, *PSOCompilationDebugData.PSOCompilationEventName, LexToString(PSOPrecacheResult));
 #endif // WITH_RHI_BREADCRUMBS
 				INC_DWORD_STAT(STAT_RuntimeGraphicsPSOHitchCount);
 				CSV_CUSTOM_STAT(PSO, GraphicsPSOHitch, 1, ECsvCustomStatOp::Accumulate);
@@ -228,14 +231,25 @@ static inline void CheckAndUpdateHitchCountStat(FPSOPrecacheRequestID::EType PSO
 			else if (PSOType == FPSOPrecacheRequestID::EType::Compute)
 			{
 #if WITH_RHI_BREADCRUMBS
-				UE_LOG(LogRHI, Log, TEXT("Runtime compute PSO creation hitch (%.2f msec) for %s (precache status: %s) - Breadcrumbs: %s"), PSOCreationTimeMs, *PSOCompilationDebugData.PSOCompilationEventName, LexToString(PSOPrecacheResult), PSOCompilationDebugData.BreadcrumbNode ? *PSOCompilationDebugData.BreadcrumbNode->GetFullPath() : TEXT("Unknown"));
+				UE_LOG(LogPSOHitching, Verbose, TEXT("Runtime compute PSO creation hitch (%.2f msec) for %s (precache status: %s) - Breadcrumbs: %s"), PSOCreationTimeMs, *PSOCompilationDebugData.PSOCompilationEventName, LexToString(PSOPrecacheResult), PSOCompilationDebugData.BreadcrumbNode ? *PSOCompilationDebugData.BreadcrumbNode->GetFullPath() : TEXT("Unknown"));
 #else
-				UE_LOG(LogRHI, Log, TEXT("Runtime compute PSO creation hitch (%.2f msec) for %s (precache status: %s)"), PSOCreationTimeMs, *PSOCompilationDebugData.PSOCompilationEventName, LexToString(PSOPrecacheResult));
+				UE_LOG(LogPSOHitching, Verbose, TEXT("Runtime compute PSO creation hitch (%.2f msec) for %s (precache status: %s)"), PSOCreationTimeMs, *PSOCompilationDebugData.PSOCompilationEventName, LexToString(PSOPrecacheResult));
 #endif // WITH_RHI_BREADCRUMBS
 				INC_DWORD_STAT(STAT_RuntimeComputePSOHitchCount);
 				CSV_CUSTOM_STAT(PSO, ComputePSOHitch, 1, ECsvCustomStatOp::Accumulate);
 				ComputePSOCreationHitchCount++;
 			}
+
+			if (PSOPrecacheResult == EPSOPrecacheResult::Complete)
+			{
+				PrecachedPSOCreationHitchCount++;
+			}
+
+			int32 TotalHitches = GraphicsPSOCreationHitchCount + ComputePSOCreationHitchCount;
+			if (TotalHitches > 0 && TotalHitches % 50 == 0)
+			{
+				UE_LOG(LogPSOHitching, Log, TEXT("Encountered %d PSO creation hitches so far (%d graphics, %d compute). %d of them were precached."), TotalHitches, GraphicsPSOCreationHitchCount, ComputePSOCreationHitchCount, PrecachedPSOCreationHitchCount);
+			}
 		}
 	}
 }
@@ -3461,6 +3475,8 @@ void PipelineStateCache::ResetPSOHitchTrackingStats()
 
 	SET_DWORD_STAT(STAT_RuntimeComputePSOHitchCount, 0);
 	ComputePSOCreationHitchCount = 0;
+
+	PrecachedPSOCreationHitchCount = 0;
 }
 
 FRHIGraphicsPipelineState* ExecuteSetGraphicsPipelineState(FGraphicsPipelineState* GraphicsPipelineState)
