Fix assert when using buffer SRVs with offsets on Vulkan
 - The check for OffsetInBytes being a multiple of stride should only apply to structured buffers. Typed buffers only need to align to what RHIGetMinimumAlignmentForBufferBackedSRV returns, which we were already checking for.
 - Minor improvements to assert strings and formatting.

diff --git a/Engine/Source/Runtime/RHI/Private/RHIResources.cpp b/Engine/Source/Runtime/RHI/Private/RHIResources.cpp
--- a/Engine/Source/Runtime/RHI/Private/RHIResources.cpp
+++ b/Engine/Source/Runtime/RHI/Private/RHIResources.cpp
@@ -195,7 +195,12 @@ FRHIViewDesc::FBuffer::FViewInfo FRHIViewDesc::FBuffer::GetViewInfo(FRHIBuffer*
 		checkf(!EnumHasAnyFlags(Desc.Usage, BUF_StructuredBuffer | BUF_AccelerationStructure), TEXT("Cannot create typed views of structured buffers, or ray tracing acceleration structures."));
 		checkf(Format != PF_Unknown, TEXT("Format cannot be unknown for typed buffers."));
 		checkf(Stride == 0, TEXT("Do not specify a stride for typed buffer views."));
-		checkf((OffsetInBytes % RHIGetMinimumAlignmentForBufferBackedSRV(Info.Format)) == 0, TEXT("Buffer offset must be a multiple of the minimum alignment supported by the RHI."));
+		checkf((OffsetInBytes % RHIGetMinimumAlignmentForBufferBackedSRV(Info.Format)) == 0, TEXT("Buffer OffsetInBytes (%d) must be a multiple of the minimum alignment (%d) supported by the RHI for format (%d: %s)."),
+			OffsetInBytes,
+			RHIGetMinimumAlignmentForBufferBackedSRV(Info.Format),
+			Info.Format,
+			GPixelFormats[Info.Format].Name
+		);
 
 		// Stride is determined by the format
 		Info.StrideInBytes = GPixelFormats[Info.Format].BlockBytes;
@@ -208,10 +213,11 @@ FRHIViewDesc::FBuffer::FViewInfo FRHIViewDesc::FBuffer::GetViewInfo(FRHIBuffer*
 		// Stride is taken from the view, or the underlying buffer if not provided.
 		Info.StrideInBytes = Stride == 0 ? Desc.Stride : Stride;
 		checkf(Info.StrideInBytes > 0, TEXT("Stride for structured buffers must be set by the view, or on the underlying buffer resource."));
+		checkf((OffsetInBytes % Info.StrideInBytes) == 0, TEXT("OffsetInBytes (%d) must be a multiple of element stride (%d)."), OffsetInBytes, Info.StrideInBytes);
 		break;
 
 	case EBufferType::AccelerationStructure:
-		checkf(EnumHasAnyFlags(Desc.Usage, BUF_AccelerationStructure), TEXT("The buffer descriptor does not a ray tracing acceleration structure, so is incompatible with this view type."));
+		checkf(EnumHasAnyFlags(Desc.Usage, BUF_AccelerationStructure), TEXT("The buffer descriptor is not a ray tracing acceleration structure, so is incompatible with this view type."));
 		checkf(Format == PF_Unknown, TEXT("Acceleration structure views should not specify a format."));
 		checkf(RayTracingScene != nullptr, TEXT("RayTracingScene must be specified when creating view of ray tracing acceleration structures."));
 
@@ -232,7 +238,6 @@ FRHIViewDesc::FBuffer::FViewInfo FRHIViewDesc::FBuffer::GetViewInfo(FRHIBuffer*
 	}
 
 	checkf(OffsetInBytes < Desc.Size, TEXT("Buffer byte offset (%d) is out of bounds (size: %d)."), OffsetInBytes, Desc.Size);
-	checkf((OffsetInBytes % Info.StrideInBytes) == 0, TEXT("Offset in bytes must be a multiple of element stride."));
 	Info.OffsetInBytes = OffsetInBytes;
 
 	// OffsetInBytes == 0 && NumElements == 0 is a special case to mean "whole resource". If offset is non-zero, we need the caller to pass the required number of elements, except for acceleration structures.
@@ -243,7 +248,7 @@ FRHIViewDesc::FBuffer::FViewInfo FRHIViewDesc::FBuffer::GetViewInfo(FRHIBuffer*
 	Info.SizeInBytes = Info.NumElements * Info.StrideInBytes;
 
 	checkf(Info.OffsetInBytes + Info.SizeInBytes <= Desc.Size,
-		TEXT("The bounds of the view (offset: %d, size in bytes: %d, stride: %d, num elements %d) exceeds the size of the underlying buffer (%d bytes)."),
+		TEXT("The bounds of the view (offset: %d, size in bytes: %d, stride: %d, num elements: %d) exceeds the size of the underlying buffer (%d bytes)."),
 		Info.OffsetInBytes,
 		Info.SizeInBytes,
 		Info.StrideInBytes,
