Fix broken Insights marker nesting when running with `-trace=default,RHICommands"
 - The RHICommands channel adds markers around each RHI command execution, however these markers interfere with RHIBeginBreadcrumbCPU / RHIEndBreadcrumbCPU which adds their own markers explicitly.
 - Added EnqueueLambda_NoMarker to skip the automatic markers surrounding the breadcrumb lambdas.

diff --git a/Engine/Source/Runtime/RHI/Public/RHICommandList.h b/Engine/Source/Runtime/RHI/Public/RHICommandList.h
--- a/Engine/Source/Runtime/RHI/Public/RHICommandList.h
+++ b/Engine/Source/Runtime/RHI/Public/RHICommandList.h
@@ -376,6 +376,22 @@ struct TRHILambdaCommand final : public FRHICommandBase
 	}
 };
 
+template <typename RHICmdListType, typename LAMBDA>
+struct TRHILambdaCommand_NoMarker final : public FRHICommandBase
+{
+	LAMBDA Lambda;
+
+	TRHILambdaCommand_NoMarker(LAMBDA&& InLambda)
+		: Lambda(Forward<LAMBDA>(InLambda))
+	{}
+
+	void ExecuteAndDestruct(FRHICommandListBase& CmdList) override final
+	{
+		Lambda(*static_cast<RHICmdListType*>(&CmdList));
+		Lambda.~LAMBDA();
+	}
+};
+
 class FRHIContextArray : public TRHIPipelineArray<IRHIComputeContext*>
 {
 	using Base = TRHIPipelineArray<IRHIComputeContext*>;
@@ -2600,6 +2616,20 @@ public:
 		}
 	}
 
+	// Same as EnqueueLambda, but skips the Insights marker surrounding the lambda. Used by the RHI breadcrumb system.
+	template <typename LAMBDA>
+	FORCEINLINE_DEBUGGABLE void EnqueueLambda_NoMarker(LAMBDA&& Lambda)
+	{
+		if (IsBottomOfPipe())
+		{
+			Lambda(*this);
+		}
+		else
+		{
+			ALLOC_COMMAND(TRHILambdaCommand_NoMarker<FRHIComputeCommandList, LAMBDA>)(Forward<LAMBDA>(Lambda));
+		}
+	}
+
 	template <typename LAMBDA>
 	FORCEINLINE_DEBUGGABLE void EnqueueLambda(LAMBDA&& Lambda)
 	{
@@ -3065,7 +3095,7 @@ public:
 			}
 		}
 
-		EnqueueLambda(TEXT("BeginBreadcrumbCPU"), [Breadcrumb, bLink](FRHICommandListBase& ExecutingCmdList)
+		EnqueueLambda_NoMarker([Breadcrumb, bLink](FRHICommandListBase& ExecutingCmdList)
 		{
 			// Translating thread
 			ExecutingCmdList.PersistentState.LocalBreadcrumb = Breadcrumb;
@@ -3095,7 +3125,7 @@ public:
 			}
 		}
 
-		EnqueueLambda(TEXT("EndBreadcrumbCPU"), [Breadcrumb, bLink](FRHICommandListBase& ExecutingCmdList)
+		EnqueueLambda_NoMarker([Breadcrumb, bLink](FRHICommandListBase& ExecutingCmdList)
 		{
 			// Translating thread
 			ExecutingCmdList.PersistentState.LocalBreadcrumb = Breadcrumb->GetParent();
