Add RHIThreadFence(true) to SetBindingsOnShaderBindingTable to fix crash in ray tracing
 - D3D12's implementation of this function is not thread safe. Running it in parallel with DispatchRays causes a TArray modify-while-iterating assert, and later a memory stomp / crash.

diff --git a/Engine/Source/Runtime/RHI/Public/RHICommandList.h b/Engine/Source/Runtime/RHI/Public/RHICommandList.h
--- a/Engine/Source/Runtime/RHI/Public/RHICommandList.h
+++ b/Engine/Source/Runtime/RHI/Public/RHICommandList.h
@@ -4189,6 +4189,8 @@ public:
 			{
 				ALLOC_COMMAND(FRHICommandSetBindingsOnShaderBindingTable)(SBT, Pipeline, NumBindings, Bindings, BindingType);
 			}
+
+			RHIThreadFence(true);
 		}
 	}
 
