RT: Initialize memory for FRayTracingLocalShaderBindings (allows common code paths in RHI that read LooseParameterData/LooseParameterDataSize)

#rnx

diff --git a/Engine/Source/Runtime/RHI/Public/RHICommandList.h b/Engine/Source/Runtime/RHI/Public/RHICommandList.h
--- a/Engine/Source/Runtime/RHI/Public/RHICommandList.h
+++ b/Engine/Source/Runtime/RHI/Public/RHICommandList.h
@@ -4262,7 +4262,7 @@ public:
 		check(NumUniformBuffers <= UINT16_MAX);
 		check(LooseParameterDataSize <= UINT16_MAX);
 
-		FRayTracingLocalShaderBindings* InlineBindings = Alloc<FRayTracingLocalShaderBindings>();
+		FRayTracingLocalShaderBindings* InlineBindings = new(Alloc<FRayTracingLocalShaderBindings>()) FRayTracingLocalShaderBindings();
 		InlineBindings->Geometry = Scene->GetInitializer().PerInstanceGeometries[InstanceIndex];
 		InlineBindings->SegmentIndex = SegmentIndex;
 		PRAGMA_DISABLE_DEPRECATION_WARNINGS
@@ -4303,7 +4303,7 @@ public:
 		check(NumUniformBuffers <= UINT16_MAX);
 		check(LooseParameterDataSize <= UINT16_MAX);
 
-		FRayTracingLocalShaderBindings* InlineBindings = Alloc<FRayTracingLocalShaderBindings>();
+		FRayTracingLocalShaderBindings* InlineBindings = new(Alloc<FRayTracingLocalShaderBindings>()) FRayTracingLocalShaderBindings();
 		InlineBindings->RecordIndex = RecordIndex;
 		InlineBindings->Geometry = Geometry;
 		InlineBindings->SegmentIndex = GeometrySegmentIndex;
@@ -4333,7 +4333,7 @@ public:
 	FORCEINLINE_DEBUGGABLE void SetDefaultRayTracingHitGroup(
 		FRHIShaderBindingTable* SBT, FRayTracingPipelineState* Pipeline, uint32 HitGroupIndex)
 	{
-		FRayTracingLocalShaderBindings* InlineBindings = Alloc<FRayTracingLocalShaderBindings>();
+		FRayTracingLocalShaderBindings* InlineBindings = new(Alloc<FRayTracingLocalShaderBindings>()) FRayTracingLocalShaderBindings();
 		InlineBindings->ShaderIndexInPipeline = HitGroupIndex;
 		InlineBindings->RecordIndex = 0; //< Default hit group always stored at index 0
 
@@ -4347,7 +4347,7 @@ public:
 		uint32 NumUniformBuffers, FRHIUniformBuffer* const* UniformBuffers,
 		uint32 UserData)
 	{
-		FRayTracingLocalShaderBindings* InlineBindings = Alloc<FRayTracingLocalShaderBindings>();
+		FRayTracingLocalShaderBindings* InlineBindings = new(Alloc<FRayTracingLocalShaderBindings>()) FRayTracingLocalShaderBindings();
 		InlineBindings->RecordIndex = ShaderSlotInScene;
 		InlineBindings->ShaderIndexInPipeline = ShaderIndexInPipeline;
 		InlineBindings->UserData = UserData;
@@ -4373,7 +4373,7 @@ public:
 		uint32 NumUniformBuffers, FRHIUniformBuffer* const* UniformBuffers,
 		uint32 UserData)
 	{
-		FRayTracingLocalShaderBindings* InlineBindings = Alloc<FRayTracingLocalShaderBindings>();
+		FRayTracingLocalShaderBindings* InlineBindings = new(Alloc<FRayTracingLocalShaderBindings>()) FRayTracingLocalShaderBindings();
 		InlineBindings->RecordIndex = RecordIndex;
 		InlineBindings->ShaderIndexInPipeline = ShaderIndexInPipeline;
 		InlineBindings->UserData = UserData;
@@ -4398,7 +4398,7 @@ public:
 		uint32 NumUniformBuffers, FRHIUniformBuffer* const* UniformBuffers,
 		uint32 UserData)
 	{
-		FRayTracingLocalShaderBindings* InlineBindings = Alloc<FRayTracingLocalShaderBindings>();
+		FRayTracingLocalShaderBindings* InlineBindings = new(Alloc<FRayTracingLocalShaderBindings>()) FRayTracingLocalShaderBindings();
 		InlineBindings->RecordIndex = ShaderSlotInScene;
 		InlineBindings->ShaderIndexInPipeline = ShaderIndexInPipeline;
 		InlineBindings->UserData = UserData;
@@ -4424,7 +4424,7 @@ public:
 		uint32 NumUniformBuffers, FRHIUniformBuffer* const* UniformBuffers,
 		uint32 UserData)
 	{
-		FRayTracingLocalShaderBindings* InlineBindings = Alloc<FRayTracingLocalShaderBindings>();
+		FRayTracingLocalShaderBindings* InlineBindings = new(Alloc<FRayTracingLocalShaderBindings>()) FRayTracingLocalShaderBindings();
 		InlineBindings->RecordIndex = RecordIndex;
 		InlineBindings->ShaderIndexInPipeline = ShaderIndexInPipeline;
 		InlineBindings->UserData = UserData;
