By default RHI Graphics access use the light weight pixel barrier. If you need the access a resource from another stage (ex: VS), use SRVGraphicsNonPixel.

Improve DynRes, on some console (0.2%).

#tests 66735272056eaa2cb6c9e896

diff --git a/Engine/Source/Runtime/RHI/Public/RHIAccess.h b/Engine/Source/Runtime/RHI/Public/RHIAccess.h
--- a/Engine/Source/Runtime/RHI/Public/RHIAccess.h
+++ b/Engine/Source/Runtime/RHI/Public/RHIAccess.h
@@ -50,8 +50,9 @@ enum class ERHIAccess : uint32
 	None = Unknown,
 	Mask = (Last << 1) - 1,
 
-	// Graphics is a combination of pixel and non-pixel
-	SRVGraphics = SRVGraphicsPixel | SRVGraphicsNonPixel,
+	// By default Graphics use the light weight pixel barrier. 
+	// If you need the access a resource from another stage (ex: VS), use SRVGraphicsNonPixel.
+	SRVGraphics = SRVGraphicsPixel,
 
 	// A mask of the two possible SRV states
 	SRVMask = SRVCompute | SRVGraphics,
@@ -60,7 +61,7 @@ enum class ERHIAccess : uint32
 	UAVMask = UAVCompute | UAVGraphics,
 
 	// A mask of all bits representing read-only states which cannot be combined with other write states.
-	ReadOnlyExclusiveMask = CPURead | Present | IndirectArgs | VertexOrIndexBuffer | SRVGraphics | SRVCompute | CopySrc | ResolveSrc | BVHRead | ShadingRateSource,
+	ReadOnlyExclusiveMask = CPURead | Present | IndirectArgs | VertexOrIndexBuffer | SRVGraphics | SRVGraphicsNonPixel | SRVCompute | CopySrc | ResolveSrc | BVHRead | ShadingRateSource,
 
 	// A mask of all bits representing read-only states on the compute pipe which cannot be combined with other write states.
 	ReadOnlyExclusiveComputeMask = CPURead | IndirectArgs | SRVCompute | CopySrc | BVHRead,
diff --git a/Engine/Source/Runtime/RHICore/Public/RHICoreShader.h b/Engine/Source/Runtime/RHICore/Public/RHICoreShader.h
--- a/Engine/Source/Runtime/RHICore/Public/RHICoreShader.h
+++ b/Engine/Source/Runtime/RHICore/Public/RHICoreShader.h
@@ -241,9 +241,7 @@ void SetResourcesFromTables(TBinder&& Binder, FRHIShader const& Shader, TBitMask
 				{
 					ERHIAccess Access = IsComputeShaderFrequency(Shader.GetFrequency())
 						? ERHIAccess::SRVCompute
-						: Shader.GetFrequency() == SF_Pixel
-							? ERHIAccess::SRVGraphicsPixel
-							: ERHIAccess::SRVGraphicsNonPixel;
+						: ERHIAccess::SRVGraphics;
 
 					Tracker->Assert(SRV->GetViewIdentity(), Access);
 				}
@@ -280,6 +278,7 @@ void SetResourcesFromTables(TBinder&& Binder, FRHIShader const& Shader, TBitMask
 #if ENABLE_RHI_VALIDATION
 				if (Tracker)
 				{
+					// Buffers currently do not require use of the the NonPixelSRV state
 					ERHIAccess Access = IsComputeShaderFrequency(Shader.GetFrequency())
 						? ERHIAccess::UAVCompute
 						: ERHIAccess::UAVGraphics;
diff --git a/Engine/Source/Runtime/VulkanRHI/Private/VulkanBarriers.cpp b/Engine/Source/Runtime/VulkanRHI/Private/VulkanBarriers.cpp
--- a/Engine/Source/Runtime/VulkanRHI/Private/VulkanBarriers.cpp
+++ b/Engine/Source/Runtime/VulkanRHI/Private/VulkanBarriers.cpp
@@ -1976,7 +1976,6 @@ VkImageLayout FVulkanLayoutManager::GetDefaultLayout(FVulkanCmdBuffer* CmdBuffer
 	case ERHIAccess::SRVCompute:
 	case ERHIAccess::SRVGraphics:
 	case ERHIAccess::SRVGraphicsNonPixel:
-	case ERHIAccess::SRVGraphicsPixel:
 	case ERHIAccess::SRVMask:
 	{
 		if (VulkanTexture.IsDepthOrStencilAspect())
