Revert changes to SRVGraphics (CL34570735), still causes to much rhi validations warning. To fix the perf regression on some console, check if we only use a SRVGraphicsNonPixel rhi access to apply the ORDER_PRIMITIVES barrier.

#tests EngineTest wwith rhgivalidation (locally)

diff --git a/Engine/Source/Runtime/RHI/Public/RHIAccess.h b/Engine/Source/Runtime/RHI/Public/RHIAccess.h
--- a/Engine/Source/Runtime/RHI/Public/RHIAccess.h
+++ b/Engine/Source/Runtime/RHI/Public/RHIAccess.h
@@ -50,9 +50,8 @@ enum class ERHIAccess : uint32
 	None = Unknown,
 	Mask = (Last << 1) - 1,
 
-	// By default Graphics use the light weight pixel barrier. 
-	// If you need the access a resource from another stage (ex: VS), use SRVGraphicsNonPixel.
-	SRVGraphics = SRVGraphicsPixel,
+	// Graphics is a combination of pixel and non-pixel
+	SRVGraphics = SRVGraphicsPixel | SRVGraphicsNonPixel,
 
 	// A mask of the two possible SRV states
 	SRVMask = SRVCompute | SRVGraphics,
@@ -61,7 +60,7 @@ enum class ERHIAccess : uint32
 	UAVMask = UAVCompute | UAVGraphics,
 
 	// A mask of all bits representing read-only states which cannot be combined with other write states.
-	ReadOnlyExclusiveMask = CPURead | Present | IndirectArgs | VertexOrIndexBuffer | SRVGraphics | SRVGraphicsNonPixel | SRVCompute | CopySrc | ResolveSrc | BVHRead | ShadingRateSource,
+	ReadOnlyExclusiveMask = CPURead | Present | IndirectArgs | VertexOrIndexBuffer | SRVGraphics | SRVCompute | CopySrc | ResolveSrc | BVHRead | ShadingRateSource,
 
 	// A mask of all bits representing read-only states on the compute pipe which cannot be combined with other write states.
 	ReadOnlyExclusiveComputeMask = CPURead | IndirectArgs | SRVCompute | CopySrc | BVHRead,
diff --git a/Engine/Source/Runtime/RHICore/Public/RHICoreShader.h b/Engine/Source/Runtime/RHICore/Public/RHICoreShader.h
--- a/Engine/Source/Runtime/RHICore/Public/RHICoreShader.h
+++ b/Engine/Source/Runtime/RHICore/Public/RHICoreShader.h
@@ -243,7 +243,9 @@ void SetResourcesFromTables(TBinder&& Binder, FRHIShader const& Shader, TBitMask
 				{
 					ERHIAccess Access = IsComputeShaderFrequency(Shader.GetFrequency())
 						? ERHIAccess::SRVCompute
-						: ERHIAccess::SRVGraphics;
+						: Shader.GetFrequency() == SF_Pixel
+							? ERHIAccess::SRVGraphicsPixel
+							: ERHIAccess::SRVGraphicsNonPixel;
 
 					Tracker->Assert(SRV->GetViewIdentity(), Access);
 				}
@@ -280,7 +282,6 @@ void SetResourcesFromTables(TBinder&& Binder, FRHIShader const& Shader, TBitMask
 #if ENABLE_RHI_VALIDATION
 				if (Tracker)
 				{
-					// Buffers currently do not require use of the the NonPixelSRV state
 					ERHIAccess Access = IsComputeShaderFrequency(Shader.GetFrequency())
 						? ERHIAccess::UAVCompute
 						: ERHIAccess::UAVGraphics;
diff --git a/Engine/Source/Runtime/VulkanRHI/Private/VulkanBarriers.cpp b/Engine/Source/Runtime/VulkanRHI/Private/VulkanBarriers.cpp
--- a/Engine/Source/Runtime/VulkanRHI/Private/VulkanBarriers.cpp
+++ b/Engine/Source/Runtime/VulkanRHI/Private/VulkanBarriers.cpp
@@ -1976,6 +1976,7 @@ VkImageLayout FVulkanLayoutManager::GetDefaultLayout(FVulkanCmdBuffer* CmdBuffer
 	case ERHIAccess::SRVCompute:
 	case ERHIAccess::SRVGraphics:
 	case ERHIAccess::SRVGraphicsNonPixel:
+	case ERHIAccess::SRVGraphicsPixel:
 	case ERHIAccess::SRVMask:
 	{
 		if (VulkanTexture.IsDepthOrStencilAspect())
