Optimize Insights markers for RHI breadcrumbs
 - Register the breadcrumb name string once, and emit the begin/end events using only the returned ID.
 - Avoids having to sprintf the name string multiple times, and reduces the overhead in the Insights stream.

diff --git a/Engine/Source/Runtime/RHI/Public/RHIBreadcrumbs.h b/Engine/Source/Runtime/RHI/Public/RHIBreadcrumbs.h
--- a/Engine/Source/Runtime/RHI/Public/RHIBreadcrumbs.h
+++ b/Engine/Source/Runtime/RHI/Public/RHIBreadcrumbs.h
@@ -444,6 +444,21 @@
 	public:
 		uint32 const ID;
 
+#if RHI_BREADCRUMBS_EMIT_CPU
+		uint32 const CPUTraceMarkerID = 0;
+#endif
+
+		void CreateTraceMarkers()
+		{
+#if RHI_BREADCRUMBS_EMIT_CPU
+			if (TRACE_CPUPROFILER_EVENT_MANUAL_IS_ENABLED())
+			{
+				FBuffer Buffer;
+				const_cast<uint32&>(CPUTraceMarkerID) = FCpuProfilerTrace::OutputEventType(GetTCHAR(Buffer));
+			}
+#endif
+		}
+
 		struct FBuffer
 		{
 			TCHAR Data[128];
@@ -466,11 +481,9 @@
 	inline void FRHIBreadcrumbNode::BeginCPU() const
 	{
 #if RHI_BREADCRUMBS_EMIT_CPU
-		if (TRACE_CPUPROFILER_EVENT_MANUAL_IS_ENABLED())
+		if (Name.CPUTraceMarkerID)
 		{
-			FRHIBreadcrumb::FBuffer Buffer;
-			TCHAR const* Str = Name.GetTCHAR(Buffer);
-			FCpuProfilerTrace::OutputBeginDynamicEvent(Str);
+			FCpuProfilerTrace::OutputBeginEvent(Name.CPUTraceMarkerID);
 		}
 #endif
 	}
@@ -478,7 +491,7 @@
 	inline void FRHIBreadcrumbNode::EndCPU() const
 	{
 #if RHI_BREADCRUMBS_EMIT_CPU
-		if (TRACE_CPUPROFILER_EVENT_MANUAL_IS_ENABLED())
+		if (Name.CPUTraceMarkerID)
 		{
 			FCpuProfilerTrace::OutputEndEvent();
 		}
@@ -639,16 +652,15 @@
 			);
 		}
 
-		TRHIBreadcrumb(TRHIBreadcrumb const& Other)
-			: FormatString(Other.FormatString)
-			, Values(Other.Values)
-		{}
+		TRHIBreadcrumb(TRHIBreadcrumb const& Other) = delete;
 
 	public:
 		TRHIBreadcrumb(TCHAR const* FormatString, TArgs const&... Args)
 			: FormatString(FormatString)
 			, Values(Args...)
-		{}
+		{
+			CreateTraceMarkers();
+		}
 
 		virtual TCHAR const* GetTCHAR(FBuffer& Buffer) const override
 		{
@@ -670,14 +682,14 @@
 
 		TCHAR const* StringLiteral;
 
-		TRHIBreadcrumb(TRHIBreadcrumb const& Other)
-			: StringLiteral(Other.StringLiteral)
-		{}
+		TRHIBreadcrumb(TRHIBreadcrumb const& Other) = delete;
 
 	public:
 		TRHIBreadcrumb(TCHAR const* StringLiteral)
 			: StringLiteral(StringLiteral)
-		{}
+		{
+			CreateTraceMarkers();
+		}
 
 		virtual TCHAR const* GetTCHAR(FBuffer&) const override
 		{
