Fix crash when building large RT acceleration (e.g., a lot of groom instance).

The scratch buffer used for building the BLAS is limited to 2Gb which might not be enough. This CL split the acceleration building into pieces which accomodate this limitation.

diff --git a/Engine/Source/Runtime/RHI/Private/RHICommandList.cpp b/Engine/Source/Runtime/RHI/Private/RHICommandList.cpp
--- a/Engine/Source/Runtime/RHI/Private/RHICommandList.cpp
+++ b/Engine/Source/Runtime/RHI/Private/RHICommandList.cpp
@@ -1759,18 +1759,46 @@ void FRHIComputeCommandList::BuildAccelerationStructure(FRHIRayTracingGeometry*
 
 void FRHIComputeCommandList::BuildAccelerationStructures(const TArrayView<const FRayTracingGeometryBuildParams> Params)
 {
-	uint64 TotalRequiredScratchMemorySize = 0;
-	for (const FRayTracingGeometryBuildParams& P : Params)
+	// Buffer size is limited to 2Gb, thus split acceleration structure building into pieces to accomodate this limitation
+	const uint64 MaxScratchMemorySize = 2147483647u;
+	const uint32 ParamTotalCount = Params.Num();
+
+	uint32 ParamIt = 0;
+	while (ParamIt < ParamTotalCount)
 	{
-		uint64 ScratchBufferRequiredSize = P.BuildMode == EAccelerationStructureBuildMode::Update ? P.Geometry->GetSizeInfo().UpdateScratchSize : P.Geometry->GetSizeInfo().BuildScratchSize;
-		TotalRequiredScratchMemorySize += ScratchBufferRequiredSize;
-	}
+		uint32 ParamBegin = ParamIt;
+		uint32 ParamCount = 0;
 
-	FRHIResourceCreateInfo ScratchBufferCreateInfo(TEXT("RHIScratchBuffer"));
-	FRHIBufferRange ScratchBufferRange{};	
-	ScratchBufferRange.Buffer = CreateBuffer(TotalRequiredScratchMemorySize, BUF_StructuredBuffer | BUF_RayTracingScratch, 0, ERHIAccess::UAVCompute, ScratchBufferCreateInfo);
+		// Select a sub-range of input params which fits into MaxScratchMemorySize
+		uint64 TotalRequiredScratchMemorySize = 0;
+		for ( ;ParamIt < ParamTotalCount; ++ParamIt)
+		{
+			const FRayTracingGeometryBuildParams& P = Params[ParamIt];
+			uint64 ScratchBufferRequiredSize = P.BuildMode == EAccelerationStructureBuildMode::Update ? P.Geometry->GetSizeInfo().UpdateScratchSize : P.Geometry->GetSizeInfo().BuildScratchSize;
+			if (ScratchBufferRequiredSize + TotalRequiredScratchMemorySize < MaxScratchMemorySize)
+			{
+				TotalRequiredScratchMemorySize += ScratchBufferRequiredSize;
+				ParamCount++;
+			}
+			else
+			{
+				break;
+			}
+		}
+
+		// Allocate scratch buffer and build the acceleration structure for the selected sub-range
+		if (ParamCount > 0)
+		{
+			const TArrayView<const FRayTracingGeometryBuildParams> EffectiveParams = MakeArrayView<const FRayTracingGeometryBuildParams>(&Params[ParamBegin], ParamCount);
+
+			FRHIResourceCreateInfo ScratchBufferCreateInfo(TEXT("RHIScratchBuffer"));
+			FRHIBufferRange ScratchBufferRange{};
+			check(uint32(TotalRequiredScratchMemorySize) == TotalRequiredScratchMemorySize)
+			ScratchBufferRange.Buffer = CreateBuffer(TotalRequiredScratchMemorySize, BUF_StructuredBuffer | BUF_RayTracingScratch, 0, ERHIAccess::UAVCompute, ScratchBufferCreateInfo);
 
-	BuildAccelerationStructures(Params, ScratchBufferRange);
+			BuildAccelerationStructures(EffectiveParams, ScratchBufferRange);
+		}
+	}
 }
 
 static FLockTracker GLockTracker;
