Fix draw/primitive count stats in shipping
 - When HAS_GPU_STATS is 0, we should still accumulate the total counts.

diff --git a/Engine/Source/Runtime/RHI/Public/RHICommandList.h b/Engine/Source/Runtime/RHI/Public/RHICommandList.h
--- a/Engine/Source/Runtime/RHI/Public/RHICommandList.h
+++ b/Engine/Source/Runtime/RHI/Public/RHICommandList.h
@@ -1361,17 +1361,23 @@ public:
 	RHI_API void AttachBreadcrumbSubTree(FRHIBreadcrumbAllocator& Allocator, FRHIBreadcrumbList& Nodes);
 #endif
 
-#if HAS_GPU_STATS
 	void Stats_AddDraw()
 	{
+#if HAS_GPU_STATS
 		DrawStats.AddDraw(PersistentState.CurrentGPUMask, PersistentState.CurrentDrawStatsCategory.GetValue());
+#else
+		DrawStats.AddDraw(PersistentState.CurrentGPUMask, nullptr);
+#endif
 	}
 
 	void Stats_AddDrawAndPrimitives(EPrimitiveType PrimitiveType, uint32 NumPrimitives)
 	{
+#if HAS_GPU_STATS
 		DrawStats.AddDrawAndPrimitives(PersistentState.CurrentGPUMask, PersistentState.CurrentDrawStatsCategory.GetValue(), PrimitiveType, NumPrimitives);
-	}
+#else
+		DrawStats.AddDrawAndPrimitives(PersistentState.CurrentGPUMask, nullptr, PrimitiveType, NumPrimitives);
 #endif
+	}
 
 	TStaticArray<void*, MAX_NUM_GPUS>& GetQueryBatchData(ERenderQueryType QueryType)
 	{
diff --git a/Engine/Source/Runtime/RHI/Public/RHIStats.h b/Engine/Source/Runtime/RHI/Public/RHIStats.h
--- a/Engine/Source/Runtime/RHI/Public/RHIStats.h
+++ b/Engine/Source/Runtime/RHI/Public/RHIStats.h
@@ -103,13 +103,8 @@ struct FRHIDrawStatsCategory
 
 // Macros for use inside RHI context Draw functions.
 // Updates the Stats structure on the executing RHI command list
-#if HAS_GPU_STATS
-	#define RHI_DRAW_CALL_INC()                              do { GetExecutingCommandList().Stats_AddDraw();                                          } while (false)
-	#define RHI_DRAW_CALL_STATS(PrimitiveType,NumPrimitives) do { GetExecutingCommandList().Stats_AddDrawAndPrimitives(PrimitiveType, NumPrimitives); } while (false)
-#else
-	#define RHI_DRAW_CALL_INC()
-	#define RHI_DRAW_CALL_STATS(PrimitiveType,NumPrimitives)
-#endif
+#define RHI_DRAW_CALL_INC()                              do { GetExecutingCommandList().Stats_AddDraw();                                          } while (false)
+#define RHI_DRAW_CALL_STATS(PrimitiveType,NumPrimitives) do { GetExecutingCommandList().Stats_AddDrawAndPrimitives(PrimitiveType, NumPrimitives); } while (false)
 
 struct FRHIDrawStats
 {
