Input latency stat computation enabled for DX11/DX12 using IDXGISwapChain::GetFrameStatistics and correlate the input reading timestamp to when the frame is handed to the display

Note:
- not enabled in shipping, currently only used to compute the Input latency debug stat as in other RHIs
- computed stats are currently different from Nvidia Reflex PC latency timing reports
- needs extra work to support g-sync/free sync monitors

diff --git a/Engine/Source/Runtime/Core/Private/ProfilingDebugging/MiscTrace.cpp b/Engine/Source/Runtime/Core/Private/ProfilingDebugging/MiscTrace.cpp
--- a/Engine/Source/Runtime/Core/Private/ProfilingDebugging/MiscTrace.cpp
+++ b/Engine/Source/Runtime/Core/Private/ProfilingDebugging/MiscTrace.cpp
@@ -98,6 +98,14 @@ void FMiscTrace::OutputBookmarkInternal(const void* BookmarkPoint, uint16 Encode
 		<< Bookmark.FormatArgs(EncodedFormatArgs, EncodedFormatArgsSize);
 }
 
+void FMiscTrace::OutputBookmarkInternalCycles(uint64 Cycles, const void* BookmarkPoint, uint16 EncodedFormatArgsSize, uint8* EncodedFormatArgs)
+{
+	UE_TRACE_LOG(Misc, Bookmark, BookmarkChannel)
+		<< Bookmark.Cycle(Cycles)
+		<< Bookmark.BookmarkPoint(BookmarkPoint)
+		<< Bookmark.FormatArgs(EncodedFormatArgs, EncodedFormatArgsSize);
+}
+
 void FMiscTrace::OutputBeginFrame(ETraceFrameType FrameType)
 {
 	if (!UE_TRACE_CHANNELEXPR_IS_ENABLED(FrameChannel))
diff --git a/Engine/Source/Runtime/Core/Public/ProfilingDebugging/MiscTrace.h b/Engine/Source/Runtime/Core/Public/ProfilingDebugging/MiscTrace.h
--- a/Engine/Source/Runtime/Core/Public/ProfilingDebugging/MiscTrace.h
+++ b/Engine/Source/Runtime/Core/Public/ProfilingDebugging/MiscTrace.h
@@ -62,6 +62,17 @@ struct FMiscTrace
 		}
 	}
 
+	template <typename... Types>
+	static void OutputBookmarkCycles(uint64 Cycles, const void* BookmarkPoint, Types... FormatArgs)
+	{
+		uint8 FormatArgsBuffer[4096];
+		uint16 FormatArgsSize = FFormatArgsTrace::EncodeArguments(FormatArgsBuffer, FormatArgs...);
+		if (FormatArgsSize)
+		{
+			OutputBookmarkInternalCycles(Cycles, BookmarkPoint, FormatArgsSize, FormatArgsBuffer);
+		}
+	}
+
 	CORE_API static void OutputBeginRegion(const TCHAR* RegionName);
 	CORE_API static void OutputEndRegion(const TCHAR* RegionName);
 
@@ -74,6 +85,7 @@ struct FMiscTrace
 
 private:
 	CORE_API static void OutputBookmarkInternal(const void* BookmarkPoint, uint16 EncodedFormatArgsSize, uint8* EncodedFormatArgs);
+	CORE_API static void OutputBookmarkInternalCycles(uint64 Cycles, const void* BookmarkPoint, uint16 EncodedFormatArgsSize, uint8* EncodedFormatArgs);
 };
 
 #define TRACE_BOOKMARK(Format, ...) \
@@ -88,6 +100,18 @@ private:
 	} \
 	FMiscTrace::OutputBookmark(&PREPROCESSOR_JOIN(__BookmarkPoint, __LINE__), ##__VA_ARGS__);
 
+#define TRACE_BOOKMARK_CYCLES(Cycles, Format, ...) \
+	static bool PREPROCESSOR_JOIN(__BookmarkPoint, __LINE__); \
+	if (!PREPROCESSOR_JOIN(__BookmarkPoint, __LINE__)) \
+	{ \
+		static_assert(std::is_const_v<std::remove_reference_t<decltype(Format)>>, "Formatting string must be a const TCHAR array."); \
+		static_assert(TIsArrayOrRefOfTypeByPredicate<decltype(Format), TIsCharEncodingCompatibleWithTCHAR>::Value, "Formatting string must be a TCHAR array."); \
+		UE_VALIDATE_FORMAT_STRING(Format, ##__VA_ARGS__); \
+		FMiscTrace::OutputBookmarkSpec(&PREPROCESSOR_JOIN(__BookmarkPoint, __LINE__), __FILE__, __LINE__, Format); \
+		PREPROCESSOR_JOIN(__BookmarkPoint, __LINE__) = true; \
+	} \
+	FMiscTrace::OutputBookmarkCycles(Cycles, &PREPROCESSOR_JOIN(__BookmarkPoint, __LINE__), ##__VA_ARGS__);
+
 #define TRACE_BEGIN_REGION(RegionName) \
 	FMiscTrace::OutputBeginRegion(RegionName);
 
@@ -112,6 +136,7 @@ private:
 #else
 
 #define TRACE_BOOKMARK(...)
+#define TRACE_BOOKMARK_CYCLES(...)
 #define TRACE_BEGIN_REGION(...)
 #define TRACE_END_REGION(...)
 #define TRACE_BEGIN_FRAME(...)
diff --git a/Engine/Source/Runtime/Core/Public/Windows/WindowsPlatform.h b/Engine/Source/Runtime/Core/Public/Windows/WindowsPlatform.h
--- a/Engine/Source/Runtime/Core/Public/Windows/WindowsPlatform.h
+++ b/Engine/Source/Runtime/Core/Public/Windows/WindowsPlatform.h
@@ -64,6 +64,7 @@ typedef FWindowsPlatformTypes FPlatformTypes;
 
 #define PLATFORM_LITTLE_ENDIAN								1
 #define PLATFORM_SUPPORTS_UNALIGNED_LOADS					1
+#define PLATFORM_SUPPORTS_FLIP_TRACKING						1
 
 #define PLATFORM_SUPPORTS_PRAGMA_PACK						1
 #if (defined(__arm__) || defined(_M_ARM) || defined(__aarch64__) || defined(_M_ARM64) || defined(_M_ARM64EC))
diff --git a/Engine/Source/Runtime/D3D12RHI/Private/Windows/WindowsD3D12Viewport.cpp b/Engine/Source/Runtime/D3D12RHI/Private/Windows/WindowsD3D12Viewport.cpp
--- a/Engine/Source/Runtime/D3D12RHI/Private/Windows/WindowsD3D12Viewport.cpp
+++ b/Engine/Source/Runtime/D3D12RHI/Private/Windows/WindowsD3D12Viewport.cpp
@@ -7,6 +7,7 @@
 #include "D3D12RHIPrivate.h"
 #include "Features/IModularFeatures.h"
 #include "HDRHelper.h"
+#include "RHIUtilities.h"
 #include "HAL/ThreadHeartBeat.h"
 #include "Windows/IDXGISwapchainProvider.h"
 
@@ -343,7 +344,37 @@ HRESULT FD3D12Viewport::PresentInternal(int32 SyncInterval)
 	{
 		// Ignore time spent waiting in Present. This function blocks based on GPU progress and space in the swap chain.
 		FRenderThreadIdleScope Scope(ERenderThreadIdleTypes::WaitingForGPUPresent);
-		return SwapChain1->Present(SyncInterval, Flags);
+#if !UE_BUILD_SHIPPING && PLATFORM_SUPPORTS_FLIP_TRACKING
+		static FRHIFlipDetails LastFlipFrame;
+		static DXGI_FRAME_STATISTICS LastStats = { 0 };
+		HRESULT GetStatHR;
+		DXGI_FRAME_STATISTICS stats = { 0 };
+		while (SUCCEEDED(GetStatHR = SwapChain1->GetFrameStatistics(&stats)) && (stats.PresentCount > LastFlipFrame.PresentIndex))
+		{
+			FRHIFlipDetails NewFlipFrame;
+			NewFlipFrame.PresentIndex = stats.PresentCount;
+			NewFlipFrame.VBlankTimeInCycles = stats.SyncQPCTime.QuadPart;
+
+			RHISetVsyncDebugInfo(NewFlipFrame);
+
+			LastFlipFrame = NewFlipFrame;
+			LastStats = stats;
+		}
+#endif
+
+		HRESULT PresentHR = SwapChain1->Present(SyncInterval, Flags);
+
+		UINT PresentID;
+		if (SUCCEEDED(SwapChain1->GetLastPresentCount(&PresentID)))
+		{
+			GRHIPresentCounter = PresentID;
+		}
+		else
+		{
+			GRHIPresentCounter++;
+		}
+
+		return PresentHR;
 	}
 
 	return S_OK;
diff --git a/Engine/Source/Runtime/Engine/Private/UnrealEngine.cpp b/Engine/Source/Runtime/Engine/Private/UnrealEngine.cpp
--- a/Engine/Source/Runtime/Engine/Private/UnrealEngine.cpp
+++ b/Engine/Source/Runtime/Engine/Private/UnrealEngine.cpp
@@ -13119,7 +13119,7 @@ void FFrameEndSync::Sync( bool bAllowOneFrameThreadLag )
 			[CurrentFrameCounter, CurrentInputTime](FRHICommandListImmediate&)
 		{
 			// Set the FrameCount and InputTime for input latency stats and flip debugging.
-			RHISetFrameDebugInfo(GRHIPresentCounter - 1, CurrentFrameCounter - 1, CurrentInputTime);
+			RHISetFrameDebugInfo(GRHIPresentCounter - 1, CurrentFrameCounter, CurrentInputTime);
 		});
 	});
 #endif
diff --git a/Engine/Source/Runtime/RHI/Private/RHIUtilities.cpp b/Engine/Source/Runtime/RHI/Private/RHIUtilities.cpp
--- a/Engine/Source/Runtime/RHI/Private/RHIUtilities.cpp
+++ b/Engine/Source/Runtime/RHI/Private/RHIUtilities.cpp
@@ -17,9 +17,18 @@ RHIUtilities.cpp:
 #include "RHIFwd.h"
 #include "RHIStrings.h"
 #include "Tasks/Task.h"
+#include "ProfilingDebugging/MiscTrace.h"
 
 #define USE_FRAME_OFFSET_THREAD 1
 
+#if !UE_BUILD_SHIPPING && PLATFORM_SUPPORTS_FLIP_TRACKING
+bool GVsyncInformationInsights = false;
+static FAutoConsoleVariableRef CVarVsyncInformationInsights(
+	TEXT("r.VsyncInformationInsights"),
+	GVsyncInformationInsights,
+	TEXT("Whether to show Vsync and Input events in UnrealInsights"));
+#endif
+
 TAutoConsoleVariable<FString> FDumpTransitionsHelper::CVarDumpTransitionsForResource(
 	TEXT("r.DumpTransitionsForResource"),
 	TEXT(""),
@@ -225,8 +234,8 @@ struct FRHIFrameOffsetThread : public FRunnable
 			double TargetFrameTimeInSeconds = double(SyncInterval) / double(FPlatformMisc::GetMaxRefreshRate());
 			double SlackInSeconds = FMath::Min(RHIGetSyncSlackMS() / 1000.0, TargetFrameTimeInSeconds);			// Clamp slack sync time to at most one full frame interval
 			double TargetFlipTime = (NewFlipFrame.VBlankTimeInSeconds + TargetFrameTimeInSeconds) - SlackInSeconds;
-
-			double Timeout = FMath::Max(0.0, TargetFlipTime - FPlatformTime::Seconds());
+			
+			double Timeout = FMath::Max(0.0, TargetFlipTime - FPlatformTime::ToSeconds64(FPlatformTime::Cycles64()));
 
 			FPlatformProcess::Sleep(Timeout);
 
@@ -243,7 +252,14 @@ struct FRHIFrameOffsetThread : public FRunnable
 					auto const& DebugInfo = FrameDebugInfos[DebugInfoIndex];
 					if (NewFlipFrame.PresentIndex == DebugInfo.PresentIndex)
 					{
-						GInputLatencyTime = (NewFlipFrame.VBlankTimeInSeconds / FPlatformTime::GetSecondsPerCycle64()) - DebugInfo.InputTime;
+						uint64 VBlankTimeInCycle = (uint64)(NewFlipFrame.VBlankTimeInSeconds / FPlatformTime::GetSecondsPerCycle64());
+
+						GInputLatencyTime = (VBlankTimeInCycle) - DebugInfo.InputTime;
+
+						if (GVsyncInformationInsights)
+						{
+							TRACE_BOOKMARK_CYCLES(VBlankTimeInCycle, TEXT("Vsync P:%d F:%d Lat:%.2f"), DebugInfo.PresentIndex, DebugInfo.FrameIndex, FPlatformTime::ToMilliseconds64(GInputLatencyTime))
+						}
 					}
 
 					if (DebugInfo.PresentIndex <= NewFlipFrame.PresentIndex)
@@ -334,14 +350,17 @@ public:
 	{
 #if !UE_BUILD_SHIPPING && PLATFORM_SUPPORTS_FLIP_TRACKING
 		FScopeLock Lock(&Singleton.CS);
-		if (Thread)
-		{
-			FFrameDebugInfo DebugInfo;
-			DebugInfo.PresentIndex = PresentIndex;
-			DebugInfo.FrameIndex = FrameIndex;
-			DebugInfo.InputTime = InputTime;
-			Singleton.FrameDebugInfos.Add(DebugInfo);
-		}
+
+		FFrameDebugInfo DebugInfo;
+		DebugInfo.PresentIndex = PresentIndex;
+		DebugInfo.FrameIndex = FrameIndex;
+		DebugInfo.InputTime = InputTime;
+		Singleton.FrameDebugInfos.Add(DebugInfo);
+
+ 		if (GVsyncInformationInsights)
+ 		{
+ 			TRACE_BOOKMARK_CYCLES(DebugInfo.InputTime, TEXT("Input P:%d F:%d"), DebugInfo.PresentIndex, DebugInfo.FrameIndex)
+ 		}
 #endif
 	}
 
@@ -552,6 +571,41 @@ void RHISetFrameDebugInfo(uint64 PresentIndex, uint64 FrameIndex, uint64 InputTi
 #endif
 }
 
+void RHISetVsyncDebugInfo(FRHIFlipDetails& NewFlipFrame)
+{
+#if !UE_BUILD_SHIPPING && USE_FRAME_OFFSET_THREAD && PLATFORM_SUPPORTS_FLIP_TRACKING
+	check(FRHIFrameOffsetThread::Thread == nullptr);
+
+	FScopeLock Lock(&FRHIFrameOffsetThread::Singleton.CS);
+	FRHIFrameOffsetThread::Singleton.LastFlipFrame = NewFlipFrame;
+
+	for (int32 DebugInfoIndex = FRHIFrameOffsetThread::Singleton.FrameDebugInfos.Num() - 1; DebugInfoIndex >= 0; --DebugInfoIndex)
+	{
+		auto const& DebugInfo = FRHIFrameOffsetThread::Singleton.FrameDebugInfos[DebugInfoIndex];
+		//Note that we are using PresentIndex-1 since RHISetFrameDebugInfo is setting GRHIPresentCounter-1 due to other platforms setup
+		if ((NewFlipFrame.PresentIndex-1) == DebugInfo.PresentIndex)
+		{
+			uint64 VBlankTimeInCycle = NewFlipFrame.VBlankTimeInCycles;
+
+			if (VBlankTimeInCycle > DebugInfo.InputTime)
+			{
+				GInputLatencyTime = VBlankTimeInCycle - DebugInfo.InputTime;
+			}
+			
+			if (GVsyncInformationInsights)
+			{
+				TRACE_BOOKMARK_CYCLES(VBlankTimeInCycle, TEXT("Vsync P:%d F:%d Lat:%.2f"), DebugInfo.PresentIndex, DebugInfo.FrameIndex, FPlatformTime::ToMilliseconds64(GInputLatencyTime))
+			}
+		}
+
+		if (DebugInfo.PresentIndex <= NewFlipFrame.PresentIndex)
+		{
+			FRHIFrameOffsetThread::Singleton.FrameDebugInfos.RemoveAtSwap(DebugInfoIndex);
+		}
+	}
+#endif
+}
+
 void RHIInitializeFlipTracking()
 {
 #if USE_FRAME_OFFSET_THREAD
diff --git a/Engine/Source/Runtime/RHI/Public/DynamicRHI.h b/Engine/Source/Runtime/RHI/Public/DynamicRHI.h
--- a/Engine/Source/Runtime/RHI/Public/DynamicRHI.h
+++ b/Engine/Source/Runtime/RHI/Public/DynamicRHI.h
@@ -46,17 +46,20 @@ struct FRHIFlipDetails
 	uint64 PresentIndex;
 	double FlipTimeInSeconds;
 	double VBlankTimeInSeconds;
+	uint64 VBlankTimeInCycles;
 
 	FRHIFlipDetails()
 		: PresentIndex(0)
 		, FlipTimeInSeconds(0)
 		, VBlankTimeInSeconds(0)
+		, VBlankTimeInCycles(0)
 	{}
 
-	FRHIFlipDetails(uint64 InPresentIndex, double InFlipTimeInSeconds, double InVBlankTimeInSeconds)
+	FRHIFlipDetails(uint64 InPresentIndex, double InFlipTimeInSeconds, double InVBlankTimeInSeconds, uint64 InVBlankTimeInCycles)
 		: PresentIndex(InPresentIndex)
 		, FlipTimeInSeconds(InFlipTimeInSeconds)
 		, VBlankTimeInSeconds(InVBlankTimeInSeconds)
+		, VBlankTimeInCycles(InVBlankTimeInCycles)
 	{}
 };
 
diff --git a/Engine/Source/Runtime/RHI/Public/RHIUtilities.h b/Engine/Source/Runtime/RHI/Public/RHIUtilities.h
--- a/Engine/Source/Runtime/RHI/Public/RHIUtilities.h
+++ b/Engine/Source/Runtime/RHI/Public/RHIUtilities.h
@@ -580,6 +580,9 @@ extern RHI_API void RHITriggerTaskEventOnFlip(uint64 PresentIndex, const UE::Tas
 /** Sets the FrameIndex and InputTime for the current frame. */
 extern RHI_API void RHISetFrameDebugInfo(uint64 PresentIndex, uint64 FrameIndex, uint64 InputTime);
 
+/** Sets the Vsync information for a new frame */
+extern RHI_API void RHISetVsyncDebugInfo(FRHIFlipDetails& NewFlipFrame);
+
 extern RHI_API void RHIInitializeFlipTracking();
 extern RHI_API void RHIShutdownFlipTracking();
 
diff --git a/Engine/Source/Runtime/Renderer/Private/DeferredShadingRenderer.cpp b/Engine/Source/Runtime/Renderer/Private/DeferredShadingRenderer.cpp
--- a/Engine/Source/Runtime/Renderer/Private/DeferredShadingRenderer.cpp
+++ b/Engine/Source/Runtime/Renderer/Private/DeferredShadingRenderer.cpp
@@ -1566,7 +1566,8 @@ void FDeferredShadingSceneRenderer::Render(FRDGBuilder& GraphBuilder)
 	// By default, limit our GPU usage to only GPUs specified in the view masks.
 	RDG_GPU_MASK_SCOPE(GraphBuilder, ViewFamily.EngineShowFlags.PathTracing ? FRHIGPUMask::All() : AllViewsGPUMask);
 	RDG_EVENT_SCOPE(GraphBuilder, "Scene");
-	RDG_GPU_STAT_SCOPE_VERBOSE(GraphBuilder, Unaccounted, *ViewFamily.ProfileDescription);
+	FString FrameNumDescription = FString::Printf(TEXT("%s Frame: %d"), *ViewFamily.ProfileDescription, GFrameCounterRenderThread);
+	RDG_GPU_STAT_SCOPE_VERBOSE(GraphBuilder, Unaccounted, *FrameNumDescription);
 	
 	if (RendererOutput == ERendererOutput::FinalSceneColor)
 	{
diff --git a/Engine/Source/Runtime/Windows/D3D11RHI/Private/D3D11Viewport.cpp b/Engine/Source/Runtime/Windows/D3D11RHI/Private/D3D11Viewport.cpp
--- a/Engine/Source/Runtime/Windows/D3D11RHI/Private/D3D11Viewport.cpp
+++ b/Engine/Source/Runtime/Windows/D3D11RHI/Private/D3D11Viewport.cpp
@@ -366,6 +366,27 @@ bool FD3D11Viewport::PresentChecked(IRHICommandContext& RHICmdContext, int32 Syn
 	HRESULT Result = S_OK;
 	bool bNeedNativePresent = true;
 
+#if !UE_BUILD_SHIPPING && PLATFORM_SUPPORTS_FLIP_TRACKING
+	if (SwapChain.IsValid())
+	{
+		static FRHIFlipDetails LastFlipFrame;
+		static DXGI_FRAME_STATISTICS LastStats = { 0 };
+		HRESULT GetStatHR;
+		DXGI_FRAME_STATISTICS stats = { 0 };
+		while (SUCCEEDED(GetStatHR = SwapChain->GetFrameStatistics(&stats)) && (stats.PresentCount > LastFlipFrame.PresentIndex))
+		{
+			FRHIFlipDetails NewFlipFrame;
+			NewFlipFrame.PresentIndex = stats.PresentCount;
+			NewFlipFrame.VBlankTimeInCycles = stats.SyncQPCTime.QuadPart;
+
+			RHISetVsyncDebugInfo(NewFlipFrame);
+
+			LastFlipFrame = NewFlipFrame;
+			LastStats = stats;
+		}
+	}
+#endif
+
 	if (IsValidRef(CustomPresent))
 	{
 		SCOPE_CYCLE_COUNTER(STAT_D3D11CustomPresentTime);
@@ -468,6 +489,16 @@ bool FD3D11Viewport::PresentChecked(IRHICommandContext& RHICmdContext, int32 Syn
 
 	D3DRHI->GetDeviceContext()->OMSetRenderTargets(0,0,0);
 
+	UINT PresentID;
+	if (SUCCEEDED(SwapChain->GetLastPresentCount(&PresentID)))
+	{
+		GRHIPresentCounter = PresentID;
+	}
+	else
+	{
+		GRHIPresentCounter++;
+	}
+
 	return bNeedNativePresent;
 }
 
