Add Linear to sRGB conversion support when converting 16bit Unorm format to FColor.

diff --git a/Engine/Source/Runtime/Apple/MetalRHI/Private/MetalRenderTarget.cpp b/Engine/Source/Runtime/Apple/MetalRHI/Private/MetalRenderTarget.cpp
--- a/Engine/Source/Runtime/Apple/MetalRHI/Private/MetalRenderTarget.cpp
+++ b/Engine/Source/Runtime/Apple/MetalRHI/Private/MetalRenderTarget.cpp
@@ -108,7 +108,7 @@ static void ConvertSurfaceDataToFColor(EPixelFormat Format, uint32 Width, uint32
 	}
 	else if (Format == PF_A16B16G16R16)
 	{
-		ConvertRawR16G16B16A16DataToFColor(Width, Height, In, SrcPitch, Out);
+		ConvertRawR16G16B16A16DataToFColor(Width, Height, In, SrcPitch, Out, bLinearToGamma);
 	}
 	else if (Format == PF_G16R16)
 	{
diff --git a/Engine/Source/Runtime/RHI/Public/RHISurfaceDataConversion.h b/Engine/Source/Runtime/RHI/Public/RHISurfaceDataConversion.h
--- a/Engine/Source/Runtime/RHI/Public/RHISurfaceDataConversion.h
+++ b/Engine/Source/Runtime/RHI/Public/RHISurfaceDataConversion.h
@@ -345,7 +345,7 @@ inline void ConvertRawDepthStencil64DataToFColor(uint32 Width, uint32 Height, ui
 	}
 }
 
-inline void ConvertRawR16G16B16A16DataToFColor(uint32 Width, uint32 Height, uint8 *In, uint32 SrcPitch, FColor* Out)
+inline void ConvertRawR16G16B16A16DataToFColor(uint32 Width, uint32 Height, uint8 *In, uint32 SrcPitch, FColor* Out, bool bLinearToGamma = false)
 {
 	for (uint32 Y = 0; Y < Height; Y++)
 	{
@@ -353,12 +353,12 @@ inline void ConvertRawR16G16B16A16DataToFColor(uint32 Width, uint32 Height, uint
 		FColor* DestPtr = Out + Y * Width;
 		for (uint32 X = 0; X < Width; X++)
 		{
-			*DestPtr = FColor(
-				FColor::Requantize16to8(SrcPtr->R),
-				FColor::Requantize16to8(SrcPtr->G),
-				FColor::Requantize16to8(SrcPtr->B),
-				FColor::Requantize16to8(SrcPtr->A)
-			);
+			*DestPtr = FLinearColor(
+				FColor::DequantizeUNorm16ToFloat(SrcPtr->R),
+				FColor::DequantizeUNorm16ToFloat(SrcPtr->G),
+				FColor::DequantizeUNorm16ToFloat(SrcPtr->B),
+				FColor::DequantizeUNorm16ToFloat(SrcPtr->A)
+			).ToFColor(bLinearToGamma);
 			++SrcPtr;
 			++DestPtr;
 		}
@@ -968,7 +968,7 @@ static bool ConvertDXGIToFColor(DXGI_FORMAT Format, uint32 Width, uint32 Height,
 		ConvertRawDepthStencil64DataToFColor(Width, Height, In, SrcPitch, Out, InFlags);
 		return true;
 	case DXGI_FORMAT_R16G16B16A16_UNORM:
-		ConvertRawR16G16B16A16DataToFColor(Width, Height, In, SrcPitch, Out);
+		ConvertRawR16G16B16A16DataToFColor(Width, Height, In, SrcPitch, Out, bLinearToGamma);
 		return true;
 	case DXGI_FORMAT_R16G16_UNORM:
 		ConvertRawR16G16DataToFColor(Width, Height, In, SrcPitch, Out);
diff --git a/Engine/Source/Runtime/VulkanRHI/Private/VulkanRenderTarget.cpp b/Engine/Source/Runtime/VulkanRHI/Private/VulkanRenderTarget.cpp
--- a/Engine/Source/Runtime/VulkanRHI/Private/VulkanRenderTarget.cpp
+++ b/Engine/Source/Runtime/VulkanRHI/Private/VulkanRenderTarget.cpp
@@ -100,7 +100,7 @@ static void ConvertRawDataToFColor(VkFormat VulkanFormat, uint32 DestWidth, uint
 		break;
 
 	case VK_FORMAT_R16G16B16A16_UNORM:
-		ConvertRawR16G16B16A16DataToFColor(DestWidth, DestHeight, In, SrcPitch, Dest);
+		ConvertRawR16G16B16A16DataToFColor(DestWidth, DestHeight, In, SrcPitch, Dest, bLinearToGamma);
 		break;
 
 	case VK_FORMAT_B8G8R8A8_UNORM:
