Add fatal log during FRHIShader destruction to make sure it's not being used by PSO compilation anymore (temporary helper in trying to collect information about random rare crashes during PSO compilation - current assumption is that the shader is being destroyed while the PSO compilation is still in flight, and thus something is going wrong with the refcounting)

diff --git a/Engine/Source/Runtime/RHI/Private/PipelineFileCache.cpp b/Engine/Source/Runtime/RHI/Private/PipelineFileCache.cpp
--- a/Engine/Source/Runtime/RHI/Private/PipelineFileCache.cpp
+++ b/Engine/Source/Runtime/RHI/Private/PipelineFileCache.cpp
@@ -207,6 +207,31 @@ static inline bool IsReferenceMaskSet(uint64 ReferenceMask, uint64 PSOMask)
 	return (ReferenceMask & PSOMask) == ReferenceMask;
 }
 
+#if PLATFORM_WINDOWS
+
+FRHIShader::~FRHIShader()
+{
+	if (InUseByPSOCompilation > 0)
+	{
+		UE_LOG(LogRHI, Fatal, TEXT("FRHIShader with hash: %s and Frequency: %d still in use by PSO compilation when being destroyed"), *Hash.ToString(), Frequency);
+	}
+}
+
+void FRHIShader::SetInUseByPSOCompilation(bool bInUse)
+{
+	if (bInUse)
+	{
+		FPlatformAtomics::InterlockedIncrement(&InUseByPSOCompilation);
+	}
+	else
+	{
+		check(InUseByPSOCompilation > 0);
+		FPlatformAtomics::InterlockedDecrement(&InUseByPSOCompilation);
+	}
+}
+
+#endif // PLATFORM_WINDOWS
+
 void FRHIComputeShader::UpdateStats()
 {
 	FPipelineStateStats::UpdateStats(Stats);
diff --git a/Engine/Source/Runtime/RHI/Private/PipelineStateCache.cpp b/Engine/Source/Runtime/RHI/Private/PipelineStateCache.cpp
--- a/Engine/Source/Runtime/RHI/Private/PipelineStateCache.cpp
+++ b/Engine/Source/Runtime/RHI/Private/PipelineStateCache.cpp
@@ -2170,9 +2170,28 @@ public:
 		{
 			GPipelinePrecompileTasksInFlight++;
 		}
+				
+#if PLATFORM_WINDOWS
+		auto MarkInUseByPSOCompilation = [](FRHIShader* Shader)
+		{
+			if (Shader)
+			{
+				Shader->SetInUseByPSOCompilation(true);
+			}
+		};
+#endif // PLATFORM_WINDOWS
 
 		if (!Pipeline->IsCompute())
 		{
+#if PLATFORM_WINDOWS
+			MarkInUseByPSOCompilation(Initializer.BoundShaderState.GetMeshShader());
+			MarkInUseByPSOCompilation(Initializer.BoundShaderState.GetAmplificationShader());
+			MarkInUseByPSOCompilation(Initializer.BoundShaderState.VertexShaderRHI);
+			MarkInUseByPSOCompilation(Initializer.BoundShaderState.PixelShaderRHI);
+			MarkInUseByPSOCompilation(Initializer.BoundShaderState.GetGeometryShader());
+			MarkInUseByPSOCompilation(Initializer.BoundShaderState.GetMeshShader());
+#endif // PLATFORM_WINDOWS
+
 			if (Initializer.BoundShaderState.GetMeshShader())
 			{
 				Initializer.BoundShaderState.GetMeshShader()->AddRef();
@@ -2222,7 +2241,14 @@ public:
 			{
 				Initializer.DepthStencilState->AddRef();
 			}
+		}		
+#if PLATFORM_WINDOWS
+		else
+		{
+			FComputePipelineState* ComputePipeline = static_cast<FComputePipelineState*>(Pipeline);
+			MarkInUseByPSOCompilation(ComputePipeline->ComputeShader);
 		}
+#endif // PLATFORM_WINDOWS
 	}
 
 	~FCompilePipelineStateTask()
@@ -2263,6 +2289,17 @@ public:
 			case EPSOPrecacheResult::NotSupported: PSOPrecacheResultScopeString = TEXT("PSOPrecache: Not Supported"); break;
 			case EPSOPrecacheResult::Untracked:    PSOPrecacheResultScopeString = TEXT("PSOPrecache: Untracked"); break;
 		}
+
+#if PLATFORM_WINDOWS
+		auto MarkUnusedByPSOCompilation = [](FRHIShader* Shader)
+		{
+			if (Shader)
+			{
+				Shader->SetInUseByPSOCompilation(false);
+			}
+		};
+#endif // PLATFORM_WINDOWS
+
 		TRACE_CPUPROFILER_EVENT_SCOPE_TEXT(PSOPrecacheResultScopeString);
 		{
 			TRACE_CPUPROFILER_EVENT_SCOPE_TEXT_CONDITIONAL(*PSOCompilationDebugData.PSOCompilationEventName, !PSOCompilationDebugData.PSOCompilationEventName.IsEmpty())
@@ -2285,6 +2322,10 @@ public:
 					bool bCSValid = ComputePipeline->RHIPipeline != nullptr && ComputePipeline->RHIPipeline->IsValid();
 					GPrecacheComputePipelineCache.PrecacheFinished(FPrecacheComputeInitializer(ComputePipeline->ComputeShader, ComputePipeline->Name), bCSValid);
 				}
+				
+#if PLATFORM_WINDOWS
+				MarkUnusedByPSOCompilation(ComputePipeline->ComputeShader);
+#endif // PLATFORM_WINDOWS
 			}
 			else
 			{
@@ -2344,6 +2385,15 @@ public:
 				{
 					GPrecacheGraphicsPipelineCache.PrecacheFinished(Initializer, GfxPipeline->RHIPipeline != nullptr);
 				}
+				
+#if PLATFORM_WINDOWS
+				MarkUnusedByPSOCompilation(Initializer.BoundShaderState.GetMeshShader());
+				MarkUnusedByPSOCompilation(Initializer.BoundShaderState.GetAmplificationShader());
+				MarkUnusedByPSOCompilation(Initializer.BoundShaderState.VertexShaderRHI);
+				MarkUnusedByPSOCompilation(Initializer.BoundShaderState.PixelShaderRHI);
+				MarkUnusedByPSOCompilation(Initializer.BoundShaderState.GetGeometryShader());
+				MarkUnusedByPSOCompilation(Initializer.BoundShaderState.GetMeshShader());
+#endif // PLATFORM_WINDOWS
 
 				if (Initializer.BoundShaderState.GetMeshShader())
 				{
diff --git a/Engine/Source/Runtime/RHI/Public/RHIResources.h b/Engine/Source/Runtime/RHI/Public/RHIResources.h
--- a/Engine/Source/Runtime/RHI/Public/RHIResources.h
+++ b/Engine/Source/Runtime/RHI/Public/RHIResources.h
@@ -918,6 +918,11 @@ public:
 	{
 	}
 
+#if PLATFORM_WINDOWS
+	RHI_API virtual ~FRHIShader();	
+	RHI_API void SetInUseByPSOCompilation(bool bInUse);
+#endif // PLATFORM_WINDOWS
+
 	inline EShaderFrequency GetFrequency() const
 	{
 		return Frequency;
@@ -948,6 +953,9 @@ private:
 	EShaderFrequency Frequency;
 	uint8 bNoDerivativeOps : 1;
 	uint8 bHasShaderBundleUsage : 1;
+#if PLATFORM_WINDOWS
+	volatile int16 InUseByPSOCompilation = 0;
+#endif // PLATFORM_WINDOWS
 };
 
 class FRHIGraphicsShader : public FRHIShader
