Do not store full FGraphicsPipelineStateInitializer in pre-cache PipelineCache.
Use PSO hash to identify pre-cached PSOs instead of doing full FGraphicsPipelineStateInitializer matching. This reduces number of cached entries by keeping only relevant ones

diff --git a/Engine/Source/Runtime/Apple/MetalRHI/Private/MetalDynamicRHI.cpp b/Engine/Source/Runtime/Apple/MetalRHI/Private/MetalDynamicRHI.cpp
--- a/Engine/Source/Runtime/Apple/MetalRHI/Private/MetalDynamicRHI.cpp
+++ b/Engine/Source/Runtime/Apple/MetalRHI/Private/MetalDynamicRHI.cpp
@@ -198,39 +198,33 @@ bool FMetalDynamicRHI::RHIGetRenderQueryResult(FRHIRenderQuery* QueryRHI, uint64
 uint64 FMetalDynamicRHI::RHIComputePrecachePSOHash(const FGraphicsPipelineStateInitializer& Initializer)
 {
 	// When compute precache PSO hash we assume a valid state precache PSO hash is already provided
-	checkf(Initializer.StatePrecachePSOHash != 0, TEXT("Initializer should have a valid state precache PSO hash set when computing the full initializer PSO hash"));
+	uint64 StatePrecachePSOHash = Initializer.StatePrecachePSOHash;
+	if (StatePrecachePSOHash == 0)
+	{
+		StatePrecachePSOHash = RHIComputeStatePrecachePSOHash(Initializer);
+	}
 
 	// All members which are not part of the state objects and influence the PSO on Metal
 	struct FNonStateHashKey
 	{
 		uint64							StatePrecachePSOHash;
 
-		EPrimitiveType					PrimitiveType;
 		uint32							RenderTargetsEnabled;
 		FGraphicsPipelineStateInitializer::TRenderTargetFormats RenderTargetFormats;
 		EPixelFormat					DepthStencilTargetFormat;
 		uint16							NumSamples;
 		EConservativeRasterization		ConservativeRasterization;
-		bool							bDepthBounds;
-		uint8							MultiViewCount;
-		bool							bHasFragmentDensityAttachment;
-		EVRSShadingRate					ShadingRate;
 	} HashKey;
 
 	FMemory::Memzero(&HashKey, sizeof(FNonStateHashKey));
 
-	HashKey.StatePrecachePSOHash			= Initializer.StatePrecachePSOHash;
+	HashKey.StatePrecachePSOHash			= StatePrecachePSOHash;
 
-	HashKey.PrimitiveType					= Initializer.PrimitiveType;
 	HashKey.RenderTargetsEnabled			= Initializer.RenderTargetsEnabled;
 	HashKey.RenderTargetFormats				= Initializer.RenderTargetFormats;
 	HashKey.DepthStencilTargetFormat		= Initializer.DepthStencilTargetFormat;
 	HashKey.NumSamples						= Initializer.NumSamples;
 	HashKey.ConservativeRasterization		= Initializer.ConservativeRasterization;
-	HashKey.bDepthBounds					= Initializer.bDepthBounds;
-	HashKey.MultiViewCount					= Initializer.MultiViewCount;
-	HashKey.bHasFragmentDensityAttachment	= Initializer.bHasFragmentDensityAttachment;
-	HashKey.ShadingRate						= Initializer.ShadingRate;
 
 	return CityHash64((const char*)&HashKey, sizeof(FNonStateHashKey));
 }
diff --git a/Engine/Source/Runtime/Apple/MetalRHI/Private/MetalVertexDeclaration.h b/Engine/Source/Runtime/Apple/MetalRHI/Private/MetalVertexDeclaration.h
--- a/Engine/Source/Runtime/Apple/MetalRHI/Private/MetalVertexDeclaration.h
+++ b/Engine/Source/Runtime/Apple/MetalRHI/Private/MetalVertexDeclaration.h
@@ -41,6 +41,7 @@ public:
 	uint32 BaseHash;
 
 	virtual bool GetInitializer(FVertexDeclarationElementList& Init) override final;
+	virtual uint32 GetPrecachePSOHash() const override final { return BaseHash; }
 
 protected:
 	void GenerateLayout(const FVertexDeclarationElementList& Elements);
diff --git a/Engine/Source/Runtime/RHI/Private/PipelineStateCache.cpp b/Engine/Source/Runtime/RHI/Private/PipelineStateCache.cpp
--- a/Engine/Source/Runtime/RHI/Private/PipelineStateCache.cpp
+++ b/Engine/Source/Runtime/RHI/Private/PipelineStateCache.cpp
@@ -1438,7 +1438,7 @@ protected:
 	FPSOPrecacheRequestResult TryAddNewState(const TPrecachedPSOInitializer& Initializer, const FString& PSOCompilationEventName, bool bDoAsyncCompile)
 	{
 		FPSOPrecacheRequestResult Result;
-		uint32 InitializerHash = TPrecachePipelineCacheDerived::PipelineStateInitializerHash(Initializer);
+		uint64 InitializerHash = TPrecachePipelineCacheDerived::PipelineStateInitializerHash(Initializer);
 
 		// Fast check first with read lock
 		{
@@ -1460,7 +1460,7 @@ protected:
 
 			// Add to array to get the new RequestID
 			Result.RequestID.Type = (uint32)PSOType;
-			Result.RequestID.RequestID = PrecachedPSOInitializers.Add(Initializer);
+			Result.RequestID.RequestID = PrecachedPSOInitializers.Add(InitializerHash);
 
 			// create new graphics state
 			NewPipelineState = TPrecachePipelineCacheDerived::CreateNewPSO(Initializer);
@@ -1469,7 +1469,7 @@ protected:
 			FPrecacheTask PrecacheTask;
 			PrecacheTask.PipelineState = NewPipelineState;
 			PrecacheTask.RequestID = Result.RequestID;
-			PrecachedPSOInitializerData.AddByHash(InitializerHash, Initializer, PrecacheTask);
+			PrecachedPSOInitializerData.Add(InitializerHash, PrecacheTask);
 
 			if (bDoAsyncCompile)
 			{
@@ -1488,7 +1488,7 @@ protected:
 		// A boost request might have been issued while we were kicking the task, need to check it here
 		{
 			FRWScopeLock ReadLock(PrecachePSOsRWLock, SLT_ReadOnly);
-			FPrecacheTask* FindResult = PrecachedPSOInitializerData.FindByHash(InitializerHash, Initializer);
+			FPrecacheTask* FindResult = PrecachedPSOInitializerData.Find(InitializerHash);
 			check(FindResult != nullptr);
 			if (FindResult != nullptr)
 			{
@@ -1551,17 +1551,18 @@ public:
 	{
 		check(RequestID.GetType() == PSOType);
 
-		TPrecachedPSOInitializer Initializer;
+		uint64 InitializerHash = 0u;
 		{
 			FRWScopeLock ReadLock(PrecachePSOsRWLock, SLT_ReadOnly);
-			Initializer = PrecachedPSOInitializers[RequestID.RequestID];
+			InitializerHash = PrecachedPSOInitializers[RequestID.RequestID];
 		}
-		return GetPrecachingStateInternal(Initializer);
+		return GetPrecachingStateInternal(InitializerHash);
 	}
 
 	EPSOPrecacheResult GetPrecachingState(const TPrecachedPSOInitializer& Initializer)
 	{
-		return GetPrecachingStateInternal(Initializer);
+		uint64 InitializerHash = TPrecachePipelineCacheDerived::PipelineStateInitializerHash(Initializer);
+		return GetPrecachingStateInternal(InitializerHash);
 	}
 
 	bool IsPrecaching()
@@ -1576,8 +1577,8 @@ public:
 
 		// Won't modify anything in this cache so readlock should be enough?
 		FRWScopeLock ReadLock(PrecachePSOsRWLock, SLT_ReadOnly);
-		const TPrecachedPSOInitializer& Initializer = PrecachedPSOInitializers[RequestID.RequestID];
-		FPrecacheTask* FindResult = PrecachedPSOInitializerData.Find(Initializer);
+		uint64 InitializerHash = PrecachedPSOInitializers[RequestID.RequestID];
+		FPrecacheTask* FindResult = PrecachedPSOInitializerData.Find(InitializerHash);
 		check(FindResult);
 
 		EPSOPrecacheStateMask NewMask = EPSOPrecacheStateMask::Boosted | (PSOPrecachePriority == EPSOPrecachePriority::Highest ? EPSOPrecacheStateMask::HighestPri : (EPSOPrecacheStateMask)0);
@@ -1607,21 +1608,21 @@ public:
 
 	void PrecacheFinished(const TPrecachedPSOInitializer& Initializer, bool bValid)
 	{
-		uint32 InitializerHash = TPrecachePipelineCacheDerived::PipelineStateInitializerHash(Initializer);
+		uint64 InitializerHash = TPrecachePipelineCacheDerived::PipelineStateInitializerHash(Initializer);
 
 		EPSOPrecacheStateMask PreviousStateMask = EPSOPrecacheStateMask::None;
 		{
 			FRWScopeLock WriteLock(PrecachePSOsRWLock, SLT_Write);
 
 			// Mark compiled (either succeeded or failed)
-			FPrecacheTask* FindResult = PrecachedPSOInitializerData.FindByHash(InitializerHash, Initializer);
+			FPrecacheTask* FindResult = PrecachedPSOInitializerData.Find(InitializerHash);
 			check(FindResult);
 			// We still add the 'compiling' bit here because if the task is fast enough, we can get here before the end of TryAddNewState
 			const EPSOPrecacheStateMask CompleteStateMask = bValid ? (EPSOPrecacheStateMask::Succeeded | EPSOPrecacheStateMask::Compiling) : (EPSOPrecacheStateMask::Failed | EPSOPrecacheStateMask::Compiling);
 			PreviousStateMask = FindResult->AddPSOPrecacheState(CompleteStateMask);
 
 			// Add to array of precached PSOs so it can be cleaned up
-			PrecachedPSOs.Add(Initializer);
+			PrecachedPSOs.Add(InitializerHash);
 		}
 
         // Need to ensure that the boost request was actually executed: if only it was asked by BoostPriority, but not requested (ie TryAddNewState has not set the Compiling bit
@@ -1654,10 +1655,9 @@ public:
 		FRWScopeLock WriteLock(PrecachePSOsRWLock, SLT_Write);
 		for (int32 Index = 0; Index < PrecachedPSOs.Num(); ++Index)		
 		{
-			const TPrecachedPSOInitializer& Initializer = PrecachedPSOs[Index];
-			uint32 InitializerHash = TPrecachePipelineCacheDerived::PipelineStateInitializerHash(Initializer);
+			uint64 InitializerHash = PrecachedPSOs[Index];
 
-			FPrecacheTask* FindResult = PrecachedPSOInitializerData.FindByHash(InitializerHash, Initializer);
+			FPrecacheTask* FindResult = PrecachedPSOInitializerData.Find(InitializerHash);
 			check(FindResult && IsCompilationDone((FindResult->ReadPSOPrecacheState())));
 			if (FindResult->PipelineState->IsComplete())
 			{
@@ -1674,9 +1674,9 @@ public:
 	}
 
 protected:
-	bool HasPSOBeenRequested(const TPrecachedPSOInitializer& Initializer, uint32 InitializerHash, FPSOPrecacheRequestResult& Result)
+	bool HasPSOBeenRequested(const TPrecachedPSOInitializer& Initializer, uint64 InitializerHash, FPSOPrecacheRequestResult& Result)
 	{
-		FPrecacheTask* FindResult = PrecachedPSOInitializerData.FindByHash(InitializerHash, Initializer);
+		FPrecacheTask* FindResult = PrecachedPSOInitializerData.Find(InitializerHash);
 		if (FindResult)
 		{			
 			// If not compiled yet, then return the request ID so the caller can check the state
@@ -1692,13 +1692,12 @@ protected:
 		return false;
 	}
 
-	EPSOPrecacheResult GetPrecachingStateInternal(const TPrecachedPSOInitializer& Initializer)
+	EPSOPrecacheResult GetPrecachingStateInternal(uint64 InitializerHash)
 	{
 		EPSOPrecacheStateMask CompilationState = EPSOPrecacheStateMask::None;
-		uint32 InitializerHash = TPrecachePipelineCacheDerived::PipelineStateInitializerHash(Initializer);	
 		{
 			FRWScopeLock ReadLock(PrecachePSOsRWLock, SLT_ReadOnly);
-			FPrecacheTask* FindResult = PrecachedPSOInitializerData.FindByHash(InitializerHash, Initializer);
+			FPrecacheTask* FindResult = PrecachedPSOInitializerData.Find(InitializerHash);
 			if (FindResult == nullptr)
 			{
 				return EPSOPrecacheResult::Missed;
@@ -1764,7 +1763,7 @@ protected:
 	FRWLock PrecachePSOsRWLock;
 
 	// Array containing all the precached PSO initializers thus far - the index in this array is used to uniquely identify the PSO requests
-	TArray<TPrecachedPSOInitializer> PrecachedPSOInitializers;
+	TArray<uint64> PrecachedPSOInitializers;
 
 	// Hash map used for fast retrieval of already precached PSOs
 	struct FPrecacheTask
@@ -1788,19 +1787,7 @@ protected:
 		volatile EPSOPrecacheStateMask StateMask = EPSOPrecacheStateMask::None;
 	};
 
-	using TMapKeyFuncsBaseClass = TDefaultMapKeyFuncs<TPrecachedPSOInitializer, FPrecacheTask, false>;
-	struct TMapKeyFuncs : public TMapKeyFuncsBaseClass
-	{
-		static FORCEINLINE bool Matches(typename TMapKeyFuncsBaseClass::KeyInitType A, typename TMapKeyFuncsBaseClass::KeyInitType B)
-		{
-			return TPrecachePipelineCacheDerived::PipelineStateInitializerMatch(A, B);
-		}
-		static FORCEINLINE uint32 GetKeyHash(typename TMapKeyFuncsBaseClass::KeyInitType Key)
-		{
-			return TPrecachePipelineCacheDerived::PipelineStateInitializerHash(Key);
-		}
-	};
-	TMap<TPrecachedPSOInitializer, FPrecacheTask, FDefaultSetAllocator, TMapKeyFuncs> PrecachedPSOInitializerData;
+	TMap<uint64, FPrecacheTask> PrecachedPSOInitializerData;
 
 	// Number of open active compiles
 	volatile int32 ActiveCompileCount = 0;
@@ -1811,7 +1798,7 @@ protected:
 	volatile int32 HighestPriorityCompileCount = 0;
 
 	// Finished Precached PSOs which can be garbage collected
-	TArray<TPrecachedPSOInitializer> PrecachedPSOs;
+	TArray<uint64> PrecachedPSOs;
 };
 
 struct FPrecacheComputeInitializer
@@ -1853,17 +1840,12 @@ public:
 	{
 		return GET_STATFNAME(STAT_HighestPriorityComputePSOPrecacheRequests);
 	}
-	static FORCEINLINE bool PipelineStateInitializerMatch(const FPrecacheComputeInitializer& ComputeShaderInitializerA, const FPrecacheComputeInitializer& ComputeShaderInitializerB)
-	{
-		// todo: would be good/more robust to have the CS equivalent of RHIMatchPrecachePSOInitializers instead of relying on pointers
-		// (for example if an FRHIComputeShader gets released and a new one is allocated with the same pointer value)
-		return ComputeShaderInitializerA.RHIComputeShaderAsU64 == ComputeShaderInitializerB.RHIComputeShaderAsU64;
-	}
 
-	static FORCEINLINE uint32 PipelineStateInitializerHash(const FPrecacheComputeInitializer& Key)
+	static FORCEINLINE uint64 PipelineStateInitializerHash(const FPrecacheComputeInitializer& Key)
 	{
-		return GetTypeHash(Key.RHIComputeShaderAsU64);
+		return Key.RHIComputeShaderAsU64;
 	}
+
 };
 
 class FPrecacheGraphicsPipelineCache : public TPrecachePipelineCacheBase<FPrecacheGraphicsPipelineCache, FGraphicsPipelineStateInitializer, FGraphicsPipelineState>
@@ -1875,12 +1857,7 @@ public:
 		return new FGraphicsPipelineState;
 	}
 
-	static FORCEINLINE bool PipelineStateInitializerMatch(const FGraphicsPipelineStateInitializer& A, const FGraphicsPipelineStateInitializer& B)
-	{
-		return RHIMatchPrecachePSOInitializers(A, B);
-	}
-
-	static FORCEINLINE uint32 PipelineStateInitializerHash(const FGraphicsPipelineStateInitializer& Key)
+	static FORCEINLINE uint64 PipelineStateInitializerHash(const FGraphicsPipelineStateInitializer& Key)
 	{
 		return RHIComputePrecachePSOHash(Key);
 	}
