Buffer Uploads
- Removing FRHIResourceCreateInfoUploadArray in favor of UE::RHIResourceUtils::Create[XXX]BufferFromArray
- Changing a number of buffer initialization points to use UE::RHIResourceUtils instead of temporary arrays.

diff --git a/Engine/Plugins/Experimental/RemoteSession/Source/RemoteSession/Private/Channels/RemoteSessionARCameraChannel.cpp b/Engine/Plugins/Experimental/RemoteSession/Source/RemoteSession/Private/Channels/RemoteSessionARCameraChannel.cpp
--- a/Engine/Plugins/Experimental/RemoteSession/Source/RemoteSession/Private/Channels/RemoteSessionARCameraChannel.cpp
+++ b/Engine/Plugins/Experimental/RemoteSession/Source/RemoteSession/Private/Channels/RemoteSessionARCameraChannel.cpp
@@ -30,6 +30,7 @@
 
 #include "PostProcess/DrawRectangle.h"
 #include "PostProcess/PostProcessMaterialInputs.h"
+#include "RHIResourceUtils.h"
 
 #define CAMERA_MESSAGE_ADDRESS TEXT("/ARCamera")
 
@@ -192,24 +193,15 @@ void FARCameraSceneViewExtension::PreRenderView_RenderThread(FRDGBuilder& GraphB
 
 	if (VertexBufferRHI == nullptr || !VertexBufferRHI.IsValid())
 	{
-		// Setup vertex buffer
-		TResourceArray<FFilterVertex, VERTEXBUFFER_ALIGNMENT> Vertices;
-		Vertices.SetNumUninitialized(4);
-
-		Vertices[0].Position = FVector4f(0.f, 0.f, 0.f, 1.f);
-		Vertices[0].UV = FVector2f(0.f, 0.f);
-
-		Vertices[1].Position = FVector4f(1.f, 0.f, 0.f, 1.f);
-		Vertices[1].UV = FVector2f(1.f, 0.f);
-
-		Vertices[2].Position = FVector4f(0.f, 1.f, 0.f, 1.f);
-		Vertices[2].UV = FVector2f(0.f, 1.f);
-
-		Vertices[3].Position = FVector4f(1.f, 1.f, 0.f, 1.f);
-		Vertices[3].UV = FVector2f(1.f, 1.f);
+		const FFilterVertex Vertices[] =
+		{
+			{ FVector4f(0.0f, 0.0f, 0.0f, 1.0f), FVector2f(0.0f, 0.0f) },
+			{ FVector4f(1.0f, 0.0f, 0.0f, 1.0f), FVector2f(1.0f, 0.0f) },
+			{ FVector4f(0.0f, 1.0f, 0.0f, 1.0f), FVector2f(0.0f, 1.0f) },
+			{ FVector4f(1.0f, 1.0f, 0.0f, 1.0f), FVector2f(1.0f, 1.0f) },
+		};
 
-		FRHIResourceCreateInfo CreateInfoVB(TEXT("FARCameraSceneViewExtension"), &Vertices);
-		VertexBufferRHI = RHICmdList.CreateVertexBuffer(Vertices.GetResourceDataSize(), BUF_Static, CreateInfoVB);
+		VertexBufferRHI = UE::RHIResourceUtils::CreateVertexBufferFromArray(RHICmdList, TEXT("FARCameraSceneViewExtension"), EBufferUsageFlags::Static, MakeConstArrayView(Vertices));
 	}
 
 	if (IndexBufferRHI == nullptr || !IndexBufferRHI.IsValid())
@@ -217,13 +209,7 @@ void FARCameraSceneViewExtension::PreRenderView_RenderThread(FRDGBuilder& GraphB
 		// Setup index buffer
 		const uint16 Indices[] = { 0, 1, 2, 2, 1, 3 };
 
-		TResourceArray<uint16, INDEXBUFFER_ALIGNMENT> IndexBuffer;
-		const uint32 NumIndices = UE_ARRAY_COUNT(Indices);
-		IndexBuffer.AddUninitialized(NumIndices);
-		FMemory::Memcpy(IndexBuffer.GetData(), Indices, NumIndices * sizeof(uint16));
-
-		FRHIResourceCreateInfo CreateInfoIB(TEXT("FARCameraSceneViewExtension"), &IndexBuffer);
-		IndexBufferRHI = RHICmdList.CreateIndexBuffer(sizeof(uint16), IndexBuffer.GetResourceDataSize(), BUF_Static, CreateInfoIB);
+		IndexBufferRHI = UE::RHIResourceUtils::CreateIndexBufferFromArray(RHICmdList, TEXT("FARCameraSceneViewExtension"), EBufferUsageFlags::Static, MakeConstArrayView(Indices));
 	}
 
 	PPMaterial = Channel.GetPostProcessMaterial();
diff --git a/Engine/Plugins/Experimental/TextureGraph/Source/TextureGraphEngine/FxMat/FxMaterial.cpp b/Engine/Plugins/Experimental/TextureGraph/Source/TextureGraphEngine/FxMat/FxMaterial.cpp
--- a/Engine/Plugins/Experimental/TextureGraph/Source/TextureGraphEngine/FxMat/FxMaterial.cpp
+++ b/Engine/Plugins/Experimental/TextureGraph/Source/TextureGraphEngine/FxMat/FxMaterial.cpp
@@ -3,7 +3,8 @@
 #include "ShaderParameters.h" 
 #include "Device/Device.h"
 #include "EngineModule.h"
-#include <TextureResource.h>
+#include "RHIResourceUtils.h"
+#include "TextureResource.h"
 
 IMPLEMENT_GLOBAL_SHADER(VSH_Simple, "/Plugin/TextureGraph/Simple.usf", "VSH_Simple", SF_Vertex);
 IMPLEMENT_GLOBAL_SHADER(FSH_Simple, "/Plugin/TextureGraph/Simple.usf", "FSH_Simple", SF_Pixel);
@@ -268,24 +269,15 @@ void QuadScreenBuffer::InitRHI(FRHICommandListBase& RHICmdList)
 	//Elements.Add(FVertexElement(0, STRUCT_OFFSET(FFilterVertex, UV), VET_Float2, 1, Stride));
 	//VertexDeclarationRHI = PipelineStateCache::GetOrCreateVertexDeclaration(Elements);
 
-	TResourceArray<FFilterVertex, VERTEXBUFFER_ALIGNMENT> Vertices;
-	Vertices.SetNumUninitialized(4);
-
-	Vertices[0].Position = FVector4f(-1, 1, 0, 1);
-	Vertices[0].UV = FVector2f(0, 0);
-
-	Vertices[1].Position = FVector4f(1, 1, 0, 1);
-	Vertices[1].UV = FVector2f(1, 0);
-
-	Vertices[2].Position = FVector4f(-1, -1, 0, 1);
-	Vertices[2].UV = FVector2f(0, 1);
-
-	Vertices[3].Position = FVector4f(1, -1, 0, 1);
-	Vertices[3].UV = FVector2f(1, 1);
+	const FFilterVertex Vertices[] =
+	{
+		FFilterVertex { FVector4f(-1, 1, 0, 1), FVector2f(0, 0) },
+		FFilterVertex { FVector4f( 1, 1, 0, 1), FVector2f(1, 0) },
+		FFilterVertex { FVector4f(-1,-1, 0, 1), FVector2f(0, 1) },
+		FFilterVertex { FVector4f( 1,-1, 0, 1), FVector2f(1, 1) },
+	};
 
-	// Create vertex buffer. Fill buffer with initial data upon creation
-	FRHIResourceCreateInfo CreateInfo(TEXT("FxMaterial_VB"), &Vertices);
-	VertexBufferRHI = RHICmdList.CreateVertexBuffer(Vertices.GetResourceDataSize(), BUF_Static, CreateInfo);
+	VertexBufferRHI = UE::RHIResourceUtils::CreateVertexBufferFromArray(RHICmdList, TEXT("FxMaterial_VB"), EBufferUsageFlags::Static, MakeConstArrayView(Vertices));
 }
 
 //////////////////////////////////////////////////////////////////////////
diff --git a/Engine/Plugins/Experimental/VirtualHeightfieldMesh/Source/VirtualHeightfieldMesh/Private/VirtualHeightfieldMeshVertexFactory.cpp b/Engine/Plugins/Experimental/VirtualHeightfieldMesh/Source/VirtualHeightfieldMesh/Private/VirtualHeightfieldMeshVertexFactory.cpp
--- a/Engine/Plugins/Experimental/VirtualHeightfieldMesh/Source/VirtualHeightfieldMesh/Private/VirtualHeightfieldMeshVertexFactory.cpp
+++ b/Engine/Plugins/Experimental/VirtualHeightfieldMesh/Source/VirtualHeightfieldMesh/Private/VirtualHeightfieldMeshVertexFactory.cpp
@@ -6,6 +6,7 @@
 #include "MeshDrawShaderBindings.h"
 #include "MeshMaterialShader.h"
 #include "DataDrivenShaderPlatformInfo.h"
+#include "RHIResourceUtils.h"
 
 IMPLEMENT_GLOBAL_SHADER_PARAMETER_STRUCT(FVirtualHeightfieldMeshVertexFactoryParameters, "VHM");
 
@@ -14,7 +15,7 @@ namespace
 	template< typename T >
 	FBufferRHIRef CreateIndexBuffer(FRHICommandListBase& RHICmdList, uint32 NumQuadsPerSide)
 	{
-		TResourceArray<T, INDEXBUFFER_ALIGNMENT> Indices;
+		TArray<T> Indices;
 
 		// Allocate room for indices
 		Indices.Reserve(NumQuadsPerSide * NumQuadsPerSide * 6);
@@ -63,12 +64,7 @@ namespace
 			}
 		}
 
-		const uint32 Size = Indices.GetResourceDataSize();
-		const uint32 Stride = sizeof(T);
-
-		// Create index buffer. Fill buffer with initial data upon creation
-		FRHIResourceCreateInfo CreateInfo(TEXT("FVirtualHeightfieldMeshIndexBuffer"), &Indices);
-		return RHICmdList.CreateIndexBuffer(Stride, Size, BUF_Static, CreateInfo);
+		return UE::RHIResourceUtils::CreateIndexBufferFromArray(RHICmdList, TEXT("FVirtualHeightfieldMeshIndexBuffer"), EBufferUsageFlags::Static, MakeConstArrayView(Indices));
 	}
 }
 
diff --git a/Engine/Plugins/Experimental/Water/Source/Runtime/Private/WaterViewExtension.cpp b/Engine/Plugins/Experimental/Water/Source/Runtime/Private/WaterViewExtension.cpp
--- a/Engine/Plugins/Experimental/Water/Source/Runtime/Private/WaterViewExtension.cpp
+++ b/Engine/Plugins/Experimental/Water/Source/Runtime/Private/WaterViewExtension.cpp
@@ -9,7 +9,7 @@
 #include "GerstnerWaterWaveSubsystem.h"
 #include "WaterBodyManager.h"
 #include "WaterSubsystem.h"
-#include "Containers/DynamicRHIResourceArray.h"
+#include "RHIResourceUtils.h"
 
 static TAutoConsoleVariable<bool> CVarLocalTessellationFreeze(
 	TEXT("r.Water.WaterMesh.LocalTessellation.Freeze"),
@@ -223,8 +223,8 @@ void FWaterViewExtension::UpdateGPUBuffers()
 			});
 		}
 
-		TResourceArray<FVector4f> WaterBodyDataBuffer;
-		TResourceArray<FVector4f> WaterAuxDataBuffer;
+		TArray<FVector4f> WaterBodyDataBuffer;
+		TArray<FVector4f> WaterAuxDataBuffer;
 
 		// The first element of the WaterDataBuffer contains the offsets to each of the sub-buffers.
 		// X = WaterZoneDataOffset
@@ -236,7 +236,7 @@ void FWaterViewExtension::UpdateGPUBuffers()
 		// Transform the individual arrays into the single buffer:
 		{
 			/** Copy a buffer of arbitrary PoD into a float4 resource array. Returns the starting offset of the source buffer in the dest buffer. */
-			auto AppendDataToFloat4Buffer = []<typename T>(TResourceArray<FVector4f>& Dest, const TArray<T>& Source)
+			auto AppendDataToFloat4Buffer = []<typename T>(TArray<FVector4f>& Dest, const TArray<T>& Source)
 			{
 				constexpr int32 NumFloat4PerElement = (sizeof(T) / sizeof(FVector4f));
 				const int32 StartOffset = Dest.Num();
@@ -271,12 +271,22 @@ void FWaterViewExtension::UpdateGPUBuffers()
 		(
 			[WaterGPUData=WaterGPUData, WaterAuxDataBuffer, WaterBodyDataBuffer](FRHICommandListImmediate& RHICmdList) mutable
 			{
-				FRHIResourceCreateInfo AuxDataCreateInfo(TEXT("WaterAuxDataBuffer"), &WaterAuxDataBuffer);
-				WaterGPUData->AuxDataBuffer = RHICmdList.CreateBuffer(WaterAuxDataBuffer.GetResourceDataSize(), BUF_VertexBuffer | BUF_ShaderResource | BUF_Static, sizeof(FVector4f), ERHIAccess::SRVMask, AuxDataCreateInfo);
+				WaterGPUData->AuxDataBuffer = UE::RHIResourceUtils::CreateBufferFromArray(
+					RHICmdList,
+					TEXT("WaterAuxDataBuffer"),
+					EBufferUsageFlags::VertexBuffer | EBufferUsageFlags::ShaderResource | EBufferUsageFlags::Static,
+					ERHIAccess::SRVMask,
+					MakeConstArrayView(WaterAuxDataBuffer)
+				);
 				WaterGPUData->AuxDataSRV = RHICmdList.CreateShaderResourceView(WaterGPUData->AuxDataBuffer, sizeof(FVector4f), PF_A32B32G32R32F);
-
-				FRHIResourceCreateInfo WaterBodyDataCreateInfo(TEXT("WaterBodyDataBuffer"), &WaterBodyDataBuffer);
-				WaterGPUData->WaterBodyDataBuffer = RHICmdList.CreateBuffer(WaterBodyDataBuffer.GetResourceDataSize(), BUF_VertexBuffer | BUF_ShaderResource | BUF_Static, sizeof(FVector4f), ERHIAccess::SRVMask, WaterBodyDataCreateInfo);
+				
+				WaterGPUData->WaterBodyDataBuffer = UE::RHIResourceUtils::CreateBufferFromArray(
+					RHICmdList,
+					TEXT("WaterBodyDataBuffer"),
+					EBufferUsageFlags::VertexBuffer | EBufferUsageFlags::ShaderResource | EBufferUsageFlags::Static,
+					ERHIAccess::SRVMask,
+					MakeConstArrayView(WaterBodyDataBuffer)
+				);
 				WaterGPUData->WaterBodyDataSRV = RHICmdList.CreateShaderResourceView(WaterGPUData->WaterBodyDataBuffer, sizeof(FVector4f), PF_A32B32G32R32F);
 			}
 		);
diff --git a/Engine/Plugins/Experimental/Water/Source/Runtime/Public/WaterVertexFactory.h b/Engine/Plugins/Experimental/Water/Source/Runtime/Public/WaterVertexFactory.h
--- a/Engine/Plugins/Experimental/Water/Source/Runtime/Public/WaterVertexFactory.h
+++ b/Engine/Plugins/Experimental/Water/Source/Runtime/Public/WaterVertexFactory.h
@@ -5,6 +5,7 @@
 #include "VertexFactory.h"
 #include "Containers/DynamicRHIResourceArray.h"
 #include "WaterInstanceDataBuffer.h"
+#include "RHIResourceUtils.h"
 
 class FShaderParameterMap;
 struct FShaderCompilerEnvironment;
@@ -59,7 +60,7 @@ private:
 	template <typename IndexType>
 	FBufferRHIRef CreateIndexBuffer(FRHICommandListBase& RHICmdList)
 	{
-		TResourceArray<IndexType, INDEXBUFFER_ALIGNMENT> Indices;
+		TArray<IndexType> Indices;
 
 		// Allocate room for indices
 		Indices.Reserve(NumQuadsPerSide * NumQuadsPerSide * 6);
@@ -95,12 +96,9 @@ private:
 		}
 
 		NumIndices = Indices.Num();
-		const uint32 Size = Indices.GetResourceDataSize();
-		const uint32 Stride = sizeof(IndexType);
 
 		// Create index buffer. Fill buffer with initial data upon creation
-		FRHIResourceCreateInfo CreateInfo(TEXT("FWaterMeshIndexBuffer"), &Indices);
-		return RHICmdList.CreateIndexBuffer(Stride, Size, BUF_Static, CreateInfo);
+		return UE::RHIResourceUtils::CreateIndexBufferFromArray(RHICmdList, TEXT("FWaterMeshIndexBuffer"), EBufferUsageFlags::Static, MakeConstArrayView(Indices));
 	}
 
 	int32 NumIndices = 0;
diff --git a/Engine/Plugins/Media/WebMMedia/Source/WebMMedia/Private/Player/WebMVideoDecoder.cpp b/Engine/Plugins/Media/WebMMedia/Source/WebMMedia/Private/Player/WebMVideoDecoder.cpp
--- a/Engine/Plugins/Media/WebMMedia/Source/WebMMedia/Private/Player/WebMVideoDecoder.cpp
+++ b/Engine/Plugins/Media/WebMMedia/Source/WebMMedia/Private/Player/WebMVideoDecoder.cpp
@@ -12,6 +12,7 @@
 #include "PipelineStateCache.h"
 #include "RHIStaticStates.h"
 #include "Containers/DynamicRHIResourceArray.h"
+#include "RHIResourceUtils.h"
 
 namespace
 {
@@ -36,18 +37,15 @@ namespace
 			Elements.Add(FVertexElement(0, STRUCT_OFFSET(FMediaElementVertex, TextureCoordinate), VET_Float2, 1, Stride));
 			VertexDeclarationRHI = PipelineStateCache::GetOrCreateVertexDeclaration(Elements);
 
-			TResourceArray<FMediaElementVertex> Vertices;
-			Vertices.AddUninitialized(4);
-			Vertices[0].Position.Set(-1.0f, 1.0f, 1.0f, 1.0f);
-			Vertices[0].TextureCoordinate.Set(0.0f, 0.0f);
-			Vertices[1].Position.Set(1.0f, 1.0f, 1.0f, 1.0f);
-			Vertices[1].TextureCoordinate.Set(1.0f, 0.0f);
-			Vertices[2].Position.Set(-1.0f, -1.0f, 1.0f, 1.0f);
-			Vertices[2].TextureCoordinate.Set(0.0f, 1.0f);
-			Vertices[3].Position.Set(1.0f, -1.0f, 1.0f, 1.0f);
-			Vertices[3].TextureCoordinate.Set(1.0f, 1.0f);
-			FRHIResourceCreateInfo CreateInfo(TEXT("FMoviePlaybackResources"), &Vertices);
-			VertexBufferRHI = RHICmdList.CreateVertexBuffer(sizeof(FMediaElementVertex) * 4, BUF_Static, CreateInfo);
+			const FMediaElementVertex Vertices[] =
+			{
+				FMediaElementVertex(FVector4f(-1.0f,  1.0f, 1.0f, 1.0f), FVector2f(0.0f, 0.0f)),
+				FMediaElementVertex(FVector4f( 1.0f,  1.0f, 1.0f, 1.0f), FVector2f(1.0f, 0.0f)),
+				FMediaElementVertex(FVector4f(-1.0f, -1.0f, 1.0f, 1.0f), FVector2f(0.0f, 1.0f)),
+				FMediaElementVertex(FVector4f( 1.0f, -1.0f, 1.0f, 1.0f), FVector2f(1.0f, 1.0f)),
+			};
+
+			VertexBufferRHI = UE::RHIResourceUtils::CreateVertexBufferFromArray(RHICmdList, TEXT("FMoviePlaybackResources"), EBufferUsageFlags::Static, MakeConstArrayView(Vertices));
 		}
 
 		virtual void ReleaseRHI() override
diff --git a/Engine/Plugins/Runtime/AR/AppleAR/AppleARKit/Source/AppleARKit/Private/AppleARKitVideoOverlay.cpp b/Engine/Plugins/Runtime/AR/AppleAR/AppleARKit/Source/AppleARKit/Private/AppleARKitVideoOverlay.cpp
--- a/Engine/Plugins/Runtime/AR/AppleAR/AppleARKit/Source/AppleARKit/Private/AppleARKitVideoOverlay.cpp
+++ b/Engine/Plugins/Runtime/AR/AppleAR/AppleARKit/Source/AppleARKit/Private/AppleARKitVideoOverlay.cpp
@@ -17,6 +17,7 @@
 #include "PipelineStateCache.h"
 #include "RHIUtilities.h"
 #include "DataDrivenShaderPlatformInfo.h"
+#include "RHIResourceUtils.h"
 #include "RHIStaticStates.h"
 #include "EngineModule.h"
 #include "SceneUtils.h"
@@ -302,26 +303,15 @@ void FAppleARKitVideoOverlay::RenderVideoOverlayWithMaterial(FRHICommandList& RH
 	if (!OverlayVertexBufferRHI)
 	{
 		// Setup vertex buffer
-		const FVector4f Positions[] =
+		const FFilterVertex Vertices[] =
 		{
-			FVector4f(0.0f, 1.0f, 0.0f, 1.0f),
-			FVector4f(0.0f, 0.0f, 0.0f, 1.0f),
-			FVector4f(1.0f, 1.0f, 0.0f, 1.0f),
-			FVector4f(1.0f, 0.0f, 0.0f, 1.0f)
+			{ FVector4f(0.0f, 1.0f, 0.0f, 1.0f), FVector2f(0.0f, 1.0f) },
+			{ FVector4f(0.0f, 0.0f, 0.0f, 1.0f), FVector2f(0.0f, 0.0f) },
+			{ FVector4f(1.0f, 1.0f, 0.0f, 1.0f), FVector2f(1.0f, 1.0f) },
+			{ FVector4f(1.0f, 0.0f, 0.0f, 1.0f), FVector2f(1.0f, 0.0f) },
 		};
-		
-		TResourceArray<FFilterVertex, VERTEXBUFFER_ALIGNMENT> Vertices;
-		Vertices.SetNumUninitialized(4);
-		
-		for (auto Index = 0; Index < UE_ARRAY_COUNT(Positions); ++Index)
-		{
-			const auto& Position = Positions[Index];
-			Vertices[Index].Position = Position;
-			Vertices[Index].UV = FVector2f(Position.X, Position.Y);
-		}
-		
-		FRHIResourceCreateInfo CreateInfoVB(TEXT("VideoOverlayVertexBuffer"), &Vertices);
-		OverlayVertexBufferRHI = RHICmdList.CreateVertexBuffer(Vertices.GetResourceDataSize(), BUF_Static, CreateInfoVB);
+
+		OverlayVertexBufferRHI = UE::RHIResourceUtils::CreateVertexBufferFromArray(RHICmdList, TEXT("VideoOverlayVertexBuffer"), EBufferUsageFlags::Static, MakeConstArrayView(Vertices));
 		
 		// Cache UVOffsets
 		const FVector2D ViewSize(InView.UnconstrainedViewRect.Max.X, InView.UnconstrainedViewRect.Max.Y);
@@ -350,13 +340,7 @@ void FAppleARKitVideoOverlay::RenderVideoOverlayWithMaterial(FRHICommandList& RH
 		// Setup index buffer
 		const uint16 Indices[] = { 0, 1, 2, 2, 1, 3 };
 
-		TResourceArray<uint16, INDEXBUFFER_ALIGNMENT> IndexBuffer;
-		const uint32 NumIndices = UE_ARRAY_COUNT(Indices);
-		IndexBuffer.AddUninitialized(NumIndices);
-		FMemory::Memcpy(IndexBuffer.GetData(), Indices, NumIndices * sizeof(uint16));
-
-		FRHIResourceCreateInfo CreateInfoIB(TEXT("VideoOverlayIndexBuffer"), &IndexBuffer);
-		IndexBufferRHI = RHICmdList.CreateIndexBuffer(sizeof(uint16), IndexBuffer.GetResourceDataSize(), BUF_Static, CreateInfoIB);
+		IndexBufferRHI = UE::RHIResourceUtils::CreateIndexBufferFromArray(RHICmdList, TEXT("VideoOverlayIndexBuffer"), EBufferUsageFlags::Static, MakeConstArrayView(Indices));
 	}
 
 	const auto FeatureLevel = InView.GetFeatureLevel();
diff --git a/Engine/Plugins/Runtime/AR/Google/GoogleARCore/Source/GoogleARCoreRendering/Private/GoogleARCorePassthroughCameraRenderer.cpp b/Engine/Plugins/Runtime/AR/Google/GoogleARCore/Source/GoogleARCoreRendering/Private/GoogleARCorePassthroughCameraRenderer.cpp
--- a/Engine/Plugins/Runtime/AR/Google/GoogleARCore/Source/GoogleARCoreRendering/Private/GoogleARCorePassthroughCameraRenderer.cpp
+++ b/Engine/Plugins/Runtime/AR/Google/GoogleARCore/Source/GoogleARCoreRendering/Private/GoogleARCorePassthroughCameraRenderer.cpp
@@ -22,6 +22,7 @@
 #include "HAL/IConsoleManager.h"
 #include "UObject/Package.h"
 #include "PostProcess/DrawRectangle.h"
+#include "RHIResourceUtils.h"
 #include "RHIStaticStates.h"
 
 #if !(UE_BUILD_SHIPPING || UE_BUILD_TEST)
@@ -65,37 +66,22 @@ void FGoogleARCorePassthroughCameraRenderer::InitializeRenderer_RenderThread(FSc
 		// Initialize Index buffer;
 		const uint16 Indices[] = { 0, 1, 2, 2, 1, 3};
 
-		TResourceArray<uint16, INDEXBUFFER_ALIGNMENT> IndexBuffer;
-		uint32 NumIndices = UE_ARRAY_COUNT(Indices);
-		IndexBuffer.AddUninitialized(NumIndices);
-		FMemory::Memcpy(IndexBuffer.GetData(), Indices, NumIndices * sizeof(uint16));
-
 		// Create index buffer. Fill buffer with initial data upon creation
-		FRHIResourceCreateInfo CreateInfo(TEXT("OverlayIndexBuffer"), &IndexBuffer);
-		OverlayIndexBufferRHI = RHICmdList.CreateIndexBuffer(sizeof(uint16), IndexBuffer.GetResourceDataSize(), BUF_Static, CreateInfo);
+		OverlayIndexBufferRHI = UE::RHIResourceUtils::CreateIndexBufferFromArray(RHICmdList, TEXT("OverlayIndexBuffer"), EBufferUsageFlags::Static, MakeConstArrayView(Indices));
 	}
 	
 	if (!OverlayVertexBufferRHI)
 	{
-		TResourceArray<FFilterVertex, VERTEXBUFFER_ALIGNMENT> Vertices;
-		Vertices.SetNumUninitialized(4);
-
 		// Unreal uses reversed z. 0 is the farthest.
-		Vertices[0].Position = FVector4f(0, 0, 0, 1);
-		Vertices[0].UV = { 0, 0 };
-
-		Vertices[1].Position = FVector4f(0, 1, 0, 1);
-		Vertices[1].UV = { 0, 1 };
-
-		Vertices[2].Position = FVector4f(1, 0, 0, 1);
-		Vertices[2].UV = { 1, 0 };
-
-		Vertices[3].Position = FVector4f(1, 1, 0, 1);
-		Vertices[3].UV = { 1, 1 };
+		const FFilterVertex Vertices[] =
+		{
+			FFilterVertex{ FVector4f(0, 0, 0, 1), FVector2f(0, 0) },
+			FFilterVertex{ FVector4f(0, 1, 0, 1), FVector2f(0, 1) },
+			FFilterVertex{ FVector4f(1, 0, 0, 1), FVector2f(1, 0) },
+			FFilterVertex{ FVector4f(1, 1, 0, 1), FVector2f(1, 1) },
+		};
 
-		// Create vertex buffer. Fill buffer with initial data upon creation
-		FRHIResourceCreateInfo CreateInfo(TEXT("OverlayVertexBuffer"), &Vertices);
-		OverlayVertexBufferRHI = RHICmdList.CreateVertexBuffer(Vertices.GetResourceDataSize(), BUF_Static, CreateInfo);
+		OverlayVertexBufferRHI = UE::RHIResourceUtils::CreateVertexBufferFromArray(RHICmdList, TEXT("OverlayVertexBuffer"), EBufferUsageFlags::Static, MakeConstArrayView(Vertices));
 	}
 }
 
diff --git a/Engine/Source/Runtime/Engine/Private/DynamicMeshBuilder.cpp b/Engine/Source/Runtime/Engine/Private/DynamicMeshBuilder.cpp
--- a/Engine/Source/Runtime/Engine/Private/DynamicMeshBuilder.cpp
+++ b/Engine/Source/Runtime/Engine/Private/DynamicMeshBuilder.cpp
@@ -199,14 +199,12 @@ TGlobalResource<FGlobalDynamicMeshVertexPool> GDynamicMeshVertexPool;
 
 void FDynamicMeshIndexBuffer32::InitRHI(FRHICommandListBase& RHICmdList)
 {
-	FRHIResourceCreateInfoUploadArray CreateInfo(TEXT("FDynamicMeshIndexBuffer32"), MakeArrayView(Indices));
-	IndexBufferRHI = RHICmdList.CreateIndexBuffer(sizeof(uint32), CreateInfo.GetResourceDataSize(), BUF_Static, CreateInfo);
+	IndexBufferRHI = UE::RHIResourceUtils::CreateIndexBufferFromArray(RHICmdList, TEXT("FDynamicMeshIndexBuffer32"), EBufferUsageFlags::Static, MakeConstArrayView(Indices));
 }
 
 void FDynamicMeshIndexBuffer16::InitRHI(FRHICommandListBase& RHICmdList)
 {
-	FRHIResourceCreateInfoUploadArray CreateInfo(TEXT("FDynamicMeshIndexBuffer16"), MakeArrayView(Indices));
-	IndexBufferRHI = RHICmdList.CreateIndexBuffer(sizeof(uint16), CreateInfo.GetResourceDataSize(), BUF_Static, CreateInfo);
+	IndexBufferRHI = UE::RHIResourceUtils::CreateIndexBufferFromArray(RHICmdList, TEXT("FDynamicMeshIndexBuffer16"), EBufferUsageFlags::Static, MakeConstArrayView(Indices));
 }
 
 /** FDynamicMeshBufferAllocator's base implementation. It always reallocates new buffers. */
diff --git a/Engine/Source/Runtime/Engine/Private/Particles/SubUVAnimation.cpp b/Engine/Source/Runtime/Engine/Private/Particles/SubUVAnimation.cpp
--- a/Engine/Source/Runtime/Engine/Private/Particles/SubUVAnimation.cpp
+++ b/Engine/Source/Runtime/Engine/Private/Particles/SubUVAnimation.cpp
@@ -62,8 +62,7 @@ void FSubUVBoundingGeometryBuffer::InitRHI(FRHICommandListBase& RHICmdList)
 {
 	if (Vertices->Num() > 0)
 	{
-		FRHIResourceCreateInfoUploadArray CreateInfo(TEXT("FSubUVBoundingGeometryBuffer"), *Vertices);
-		VertexBufferRHI = RHICmdList.CreateVertexBuffer(CreateInfo.GetResourceDataSize(), BUF_ShaderResource | BUF_Static, CreateInfo);
+		VertexBufferRHI = UE::RHIResourceUtils::CreateVertexBufferFromArray(RHICmdList, TEXT("FSubUVBoundingGeometryBuffer"), EBufferUsageFlags::ShaderResource | EBufferUsageFlags::Static, MakeConstArrayView(*Vertices));
 		ShaderResourceView = RHICmdList.CreateShaderResourceView(VertexBufferRHI, sizeof(FVector2f), PF_G32R32F);
 	}
 }
diff --git a/Engine/Source/Runtime/Landscape/Private/LandscapeEditLayers.cpp b/Engine/Source/Runtime/Landscape/Private/LandscapeEditLayers.cpp
--- a/Engine/Source/Runtime/Landscape/Private/LandscapeEditLayers.cpp
+++ b/Engine/Source/Runtime/Landscape/Private/LandscapeEditLayers.cpp
@@ -42,6 +42,7 @@ LandscapeEditLayers.cpp: Landscape editing layers mode
 #include "MaterialCachedData.h"
 #include "ContentStreaming.h"
 #include "Templates/TypeHash.h"
+#include "RHIResourceUtils.h"
 
 #if WITH_EDITOR
 #include "AssetCompilingManager.h"
@@ -579,7 +580,7 @@ private:
 	/** Initialize the RHI for this rendering resource */
 	void InitRHI(FRHICommandListBase& RHICmdList) override
 	{
-		TResourceArray<FLandscapeLayersVertex, VERTEXBUFFER_ALIGNMENT> Vertices;
+		TArray<FLandscapeLayersVertex> Vertices;
 		Vertices.SetNumUninitialized(TriangleList.Num() * 3);
 
 		for (int32 i = 0; i < TriangleList.Num(); ++i)
@@ -590,8 +591,7 @@ private:
 		}
 
 		// Create vertex buffer. Fill buffer with initial data upon creation
-		FRHIResourceCreateInfo CreateInfo(TEXT("FLandscapeLayersVertexBuffer"), &Vertices);
-		VertexBufferRHI = RHICmdList.CreateVertexBuffer(Vertices.GetResourceDataSize(), BUF_Static, CreateInfo);
+		VertexBufferRHI = UE::RHIResourceUtils::CreateVertexBufferFromArray(RHICmdList, TEXT("FLandscapeLayersVertexBuffer"), EBufferUsageFlags::Static, MakeConstArrayView(Vertices));
 	}
 
 	TArray<FLandscapeLayersTriangle> TriangleList;
diff --git a/Engine/Source/Runtime/RHI/Public/RHIResourceUtils.h b/Engine/Source/Runtime/RHI/Public/RHIResourceUtils.h
--- a/Engine/Source/Runtime/RHI/Public/RHIResourceUtils.h
+++ b/Engine/Source/Runtime/RHI/Public/RHIResourceUtils.h
@@ -6,54 +6,38 @@
 #include "RHIResources.h"
 #include "Containers/ResourceArray.h"
 
-struct FRHIResourceCreateInfoUploadArray : public FRHIResourceCreateInfo
+namespace UE::RHIResourceUtils
 {
-	FResourceArrayUploadArrayView UploadView;
-
-	FRHIResourceCreateInfoUploadArray(const TCHAR* InDebugName, const void* InData, uint32 InSizeInBytes)
-		: FRHIResourceCreateInfo(InDebugName, &UploadView)
-		, UploadView(InData, InSizeInBytes)
+	static FBufferRHIRef CreateBufferFromArray(FRHICommandListBase& RHICmdList, const TCHAR* Name, EBufferUsageFlags UsageFlags, uint32 InStride, const void* InData, uint32 InSizeInBytes)
 	{
-	}
+		const ERHIAccess InitialState = RHIGetDefaultResourceState(UsageFlags, false);
+		FResourceArrayUploadArrayView UploadView(InData, InSizeInBytes);
 
-	template<typename ElementType>
-	FRHIResourceCreateInfoUploadArray(const TCHAR* InDebugName, TConstArrayView<ElementType> InView)
-		: FRHIResourceCreateInfo(InDebugName, &UploadView)
-		, UploadView(InView)
-	{
-	}
-
-	template<typename ElementType>
-	FRHIResourceCreateInfoUploadArray(const TCHAR* InDebugName, TArrayView<ElementType> InView)
-	: FRHIResourceCreateInfoUploadArray(InDebugName, TConstArrayView<ElementType>(InView))
-	{
+		FRHIResourceCreateInfo CreateInfo(Name, &UploadView);
+		return RHICmdList.CreateBuffer(InSizeInBytes, UsageFlags, InStride, InitialState, CreateInfo);
 	}
 
-	template<typename ElementType, typename AllocatorType>
-	FRHIResourceCreateInfoUploadArray(const TCHAR* InDebugName, const TArray<ElementType, AllocatorType>& InArray)
-	: FRHIResourceCreateInfoUploadArray(InDebugName, TConstArrayView<ElementType>(InArray))
+	template<typename TElementType>
+	static FBufferRHIRef CreateBufferFromArray(FRHICommandListBase& RHICmdList, const TCHAR* Name, EBufferUsageFlags UsageFlags, uint32 InStride, ERHIAccess InitialState, TConstArrayView<TElementType> Array)
 	{
-	}
+		FResourceArrayUploadArrayView UploadView(Array);
 
-	uint32 GetResourceDataSize() const
-	{
-		return UploadView.GetResourceDataSize();
+		FRHIResourceCreateInfo CreateInfo(Name, &UploadView);
+		return RHICmdList.CreateBuffer(UploadView.GetResourceDataSize(), UsageFlags, InStride, InitialState, CreateInfo);
 	}
-};
 
-namespace UE::RHIResourceUtils
-{
-	template <typename TElementType>
-	static FBufferRHIRef CreateBufferFromArray(FRHICommandListBase& RHICmdList, const TCHAR* Name, EBufferUsageFlags UsageFlags, ERHIAccess ResourceState, TConstArrayView<TElementType> Array)
+	template<typename TElementType>
+	static FBufferRHIRef CreateBufferFromArray(FRHICommandListBase& RHICmdList, const TCHAR* Name, EBufferUsageFlags UsageFlags, ERHIAccess InitialState, TConstArrayView<TElementType> Array)
 	{
-		FRHIResourceCreateInfoUploadArray CreateInfo(Name, Array);
-		return RHICmdList.CreateBuffer(CreateInfo.GetResourceDataSize(), UsageFlags, Array.GetTypeSize(), ResourceState, CreateInfo);
+		return CreateBufferFromArray(RHICmdList, Name, UsageFlags, Array.GetTypeSize(), InitialState, Array);
 	}
 
 	template<typename TElementType>
 	static FBufferRHIRef CreateVertexBufferFromArray(FRHICommandListBase& RHICmdList, const TCHAR* Name, EBufferUsageFlags ExtraFlags, TConstArrayView<TElementType> Array)
 	{
-		return CreateBufferFromArray<TElementType>(RHICmdList, Name, EBufferUsageFlags::VertexBuffer | ExtraFlags, ERHIAccess::VertexOrIndexBuffer, Array);
+		const EBufferUsageFlags Usage = EBufferUsageFlags::VertexBuffer | ExtraFlags;
+		const ERHIAccess InitialState = RHIGetDefaultResourceState(Usage, false);
+		return CreateBufferFromArray<TElementType>(RHICmdList, Name, Usage, 0, InitialState, Array);
 	}
 
 	template<typename TElementType>
@@ -65,7 +49,9 @@ namespace UE::RHIResourceUtils
 	template<typename TElementType>
 	static FBufferRHIRef CreateIndexBufferFromArray(FRHICommandListBase& RHICmdList, const TCHAR* Name, EBufferUsageFlags ExtraFlags, TConstArrayView<TElementType> Array)
 	{
-		return CreateBufferFromArray<TElementType>(RHICmdList, Name, EBufferUsageFlags::IndexBuffer | ExtraFlags, ERHIAccess::VertexOrIndexBuffer, Array);
+		const EBufferUsageFlags Usage = EBufferUsageFlags::IndexBuffer | ExtraFlags;
+		const ERHIAccess InitialState = RHIGetDefaultResourceState(Usage, false);
+		return CreateBufferFromArray<TElementType>(RHICmdList, Name, Usage, InitialState, Array);
 	}
 
 	template<typename TElementType>
diff --git a/Engine/Source/Runtime/RenderCore/Private/CommonRenderResources.cpp b/Engine/Source/Runtime/RenderCore/Private/CommonRenderResources.cpp
--- a/Engine/Source/Runtime/RenderCore/Private/CommonRenderResources.cpp
+++ b/Engine/Source/Runtime/RenderCore/Private/CommonRenderResources.cpp
@@ -50,8 +50,12 @@ static const FFilterVertex GScreenRectangleVertexBufferData[] =
 void FScreenRectangleVertexBuffer::InitRHI(FRHICommandListBase& RHICmdList)
 {
 	// Create vertex buffer. Fill buffer with initial data upon creation
-	FRHIResourceCreateInfoUploadArray CreateInfo(TEXT("FScreenRectangleVertexBuffer"), MakeArrayView(GScreenRectangleVertexBufferData));
-	VertexBufferRHI = RHICmdList.CreateVertexBuffer(CreateInfo.GetResourceDataSize(), BUF_Static, CreateInfo);
+	VertexBufferRHI = UE::RHIResourceUtils::CreateVertexBufferFromArray(
+		RHICmdList,
+		TEXT("FScreenRectangleVertexBuffer"),
+		EBufferUsageFlags::Static,
+		MakeConstArrayView(GScreenRectangleVertexBufferData)
+	);
 }
 
 static const uint16 GScreenRectangleIndexBufferData[] =
@@ -64,6 +68,10 @@ static const uint16 GScreenRectangleIndexBufferData[] =
 void FScreenRectangleIndexBuffer::InitRHI(FRHICommandListBase& RHICmdList)
 {
 	// Create index buffer. Fill buffer with initial data upon creation
-	FRHIResourceCreateInfoUploadArray CreateInfo(TEXT("FScreenRectangleIndexBuffer"), MakeArrayView(GScreenRectangleIndexBufferData));
-	IndexBufferRHI = RHICmdList.CreateIndexBuffer(sizeof(uint16), CreateInfo.GetResourceDataSize(), BUF_Static, CreateInfo);
+	IndexBufferRHI = UE::RHIResourceUtils::CreateIndexBufferFromArray(
+		RHICmdList,
+		TEXT("FScreenRectangleIndexBuffer"),
+		EBufferUsageFlags::Static,
+		MakeConstArrayView(GScreenRectangleIndexBufferData)
+	);
 }
diff --git a/Engine/Source/Runtime/RenderCore/Private/RenderUtils.cpp b/Engine/Source/Runtime/RenderCore/Private/RenderUtils.cpp
--- a/Engine/Source/Runtime/RenderCore/Private/RenderUtils.cpp
+++ b/Engine/Source/Runtime/RenderCore/Private/RenderUtils.cpp
@@ -11,6 +11,7 @@
 #include "PipelineStateCache.h"
 #include "RenderResource.h"
 #include "RHI.h"
+#include "RHIResourceUtils.h"
 #include "Shader.h"
 #include "ShaderPlatformCachedIniValue.h"
 #include "DataDrivenShaderPlatformInfo.h"
@@ -1002,7 +1003,7 @@ public:
 	void InitRHI(FRHICommandListBase& RHICmdList) override
 	{
 		const int32 NumVerts = 8;
-		TResourceArray<FVector4f, VERTEXBUFFER_ALIGNMENT> Verts;
+		TArray<FVector4f> Verts;
 		Verts.SetNumUninitialized(NumVerts);
 
 		for (uint32 Z = 0; Z < 2; Z++)
@@ -1023,11 +1024,8 @@ public:
 			}
 		}
 
-		uint32 Size = Verts.GetResourceDataSize();
-
 		// Create vertex buffer. Fill buffer with initial data upon creation
-		FRHIResourceCreateInfo CreateInfo(TEXT("FUnitCubeVertexBuffer"), &Verts);
-		VertexBufferRHI = RHICmdList.CreateVertexBuffer(Size, BUF_Static, CreateInfo);
+		VertexBufferRHI = UE::RHIResourceUtils::CreateVertexBufferFromArray(RHICmdList, TEXT("FUnitCubeVertexBuffer"), EBufferUsageFlags::Static, MakeConstArrayView(Verts));
 	}
 };
 
@@ -1039,18 +1037,8 @@ public:
 	*/
 	void InitRHI(FRHICommandListBase& RHICmdList) override
 	{
-		TResourceArray<uint16, INDEXBUFFER_ALIGNMENT> Indices;
-		
-		int32 NumIndices = UE_ARRAY_COUNT(GCubeIndices);
-		Indices.AddUninitialized(NumIndices);
-		FMemory::Memcpy(Indices.GetData(), GCubeIndices, NumIndices * sizeof(uint16));
-
-		const uint32 Size = Indices.GetResourceDataSize();
-		const uint32 Stride = sizeof(uint16);
-
 		// Create index buffer. Fill buffer with initial data upon creation
-		FRHIResourceCreateInfo CreateInfo(TEXT("FUnitCubeIndexBuffer"), &Indices);
-		IndexBufferRHI = RHICmdList.CreateIndexBuffer(Stride, Size, BUF_Static, CreateInfo);
+		IndexBufferRHI = UE::RHIResourceUtils::CreateIndexBufferFromArray(RHICmdList, TEXT("FUnitCubeIndexBuffer"), EBufferUsageFlags::Static, MakeConstArrayView(GCubeIndices));
 	}
 };
 
@@ -1064,17 +1052,14 @@ public:
 	*/
 	void InitRHI(FRHICommandListBase& RHICmdList) override
 	{
-		const int32 NumVerts = 2;
-		TResourceArray<FVector3f, VERTEXBUFFER_ALIGNMENT> Verts;
-		Verts.SetNumUninitialized(NumVerts);
-		Verts[0] = FVector3f(-0.5f, -0.5f, -0.5f);
-		Verts[1] = FVector3f(0.5f, 0.5f, 0.5f);
-
-		uint32 Size = Verts.GetResourceDataSize();
+		const FVector3f Vertices[] =
+		{
+			FVector3f(-0.5f, -0.5f, -0.5f),
+			FVector3f( 0.5f,  0.5f,  0.5f),
+		};
 
 		// Create vertex buffer. Fill buffer with initial data upon creation
-		FRHIResourceCreateInfo CreateInfo(TEXT("FUnitCubeAABBVertexBuffer"), &Verts);
-		VertexBufferRHI = RHICmdList.CreateVertexBuffer(Size, BUF_Static, CreateInfo);
+		VertexBufferRHI = UE::RHIResourceUtils::CreateVertexBufferFromArray(RHICmdList, TEXT("FUnitCubeAABBVertexBuffer"), EBufferUsageFlags::Static, MakeConstArrayView(Vertices));
 	}
 };
 #endif // RHI_RAYTRACING
diff --git a/Engine/Source/Runtime/RenderCore/Private/UnifiedBuffer.cpp b/Engine/Source/Runtime/RenderCore/Private/UnifiedBuffer.cpp
--- a/Engine/Source/Runtime/RenderCore/Private/UnifiedBuffer.cpp
+++ b/Engine/Source/Runtime/RenderCore/Private/UnifiedBuffer.cpp
@@ -1702,16 +1702,14 @@ void FScatterUploadBuffer::ResourceUploadTo(FRHICommandList& RHICmdList, const R
 		UploadBuffer.NumBytes = UploadDataSize;
 
 		const uint32 TypeSize = bFloat4Buffer ? 16 : 4;
-		const EBufferUsageFlags Usage = bFloat4Buffer ? BUF_None : BUF_ByteAddressBuffer;
+		const EBufferUsageFlags Usage = EBufferUsageFlags::StructuredBuffer | EBufferUsageFlags::ShaderResource | EBufferUsageFlags::Volatile | (bFloat4Buffer ? EBufferUsageFlags::None : EBufferUsageFlags::ByteAddressBuffer);
 
 		{
-			FRHIResourceCreateInfoUploadArray CreateInfo(TEXT("ScatterResourceArray"), ScatterData, ScatterDataSize);
-			ScatterBuffer.Buffer = RHICmdList.CreateStructuredBuffer(sizeof(uint32), ScatterDataSize, BUF_ShaderResource | BUF_Volatile | Usage, CreateInfo);
+			ScatterBuffer.Buffer = UE::RHIResourceUtils::CreateBufferFromArray(RHICmdList, TEXT("ScatterResourceArray"), Usage, sizeof(uint32), ScatterData, ScatterDataSize);
 			ScatterBuffer.SRV = RHICmdList.CreateShaderResourceView(ScatterBuffer.Buffer);
 		}
 		{
-			FRHIResourceCreateInfoUploadArray CreateInfo(TEXT("ScatterUploadBuffer"), UploadData, UploadDataSize);
-			UploadBuffer.Buffer = RHICmdList.CreateStructuredBuffer(TypeSize, UploadDataSize, BUF_ShaderResource | BUF_Volatile | Usage, CreateInfo);
+			UploadBuffer.Buffer = UE::RHIResourceUtils::CreateBufferFromArray(RHICmdList, TEXT("ScatterUploadBuffer"), Usage, TypeSize, UploadData, UploadDataSize);
 			UploadBuffer.SRV = RHICmdList.CreateShaderResourceView(UploadBuffer.Buffer);
 		}
 	}
diff --git a/Engine/Source/Runtime/Renderer/Private/LightRendering.cpp b/Engine/Source/Runtime/Renderer/Private/LightRendering.cpp
--- a/Engine/Source/Runtime/Renderer/Private/LightRendering.cpp
+++ b/Engine/Source/Runtime/Renderer/Private/LightRendering.cpp
@@ -34,6 +34,7 @@
 #include "LightFunctionAtlas.h"
 #include "HeterogeneousVolumes/HeterogeneousVolumes.h"
 #include "Materials/MaterialRenderProxy.h"
+#include "RHIResourceUtils.h"
 
 using namespace LightFunctionAtlas;
 
@@ -291,7 +292,7 @@ public:
 
 	void InitRHI(FRHICommandListBase& RHICmdList) override
 	{
-		TResourceArray<uint16, INDEXBUFFER_ALIGNMENT> Indices;
+		TArray<uint16> Indices;
 
 		Indices.Empty((NumSlices - 1) * NumSides * 12);
 		// Generate triangles for the vertices of the cone shape
@@ -334,14 +335,10 @@ public:
 			}
 		}
 
-		const uint32 Size = Indices.GetResourceDataSize();
-		const uint32 Stride = sizeof(uint16);
-
 		NumIndices = Indices.Num();
 
 		// Create index buffer. Fill buffer with initial data upon creation
-		FRHIResourceCreateInfo CreateInfo(TEXT("FStencilConeIndexBuffer"), &Indices);
-		IndexBufferRHI = RHICmdList.CreateIndexBuffer(Stride, Size, BUF_Static, CreateInfo);
+		IndexBufferRHI = UE::RHIResourceUtils::CreateIndexBufferFromArray(RHICmdList, TEXT("FStencilConeIndexBuffer"), EBufferUsageFlags::Static, MakeConstArrayView(Indices));
 	}
 
 	int32 GetIndexCount() const { return NumIndices; }
@@ -366,18 +363,14 @@ public:
 	*/
 	void InitRHI(FRHICommandListBase& RHICmdList) override
 	{
-		TResourceArray<FVector4f, VERTEXBUFFER_ALIGNMENT> Verts;
+		TArray<FVector4f> Verts;
 		Verts.Empty(NumVerts);
 		for (int32 s = 0; s < NumVerts; s++)
 		{
-			Verts.Add(FVector4f(0, 0, 0, 0));
+			Verts.Emplace(0, 0, 0, 0);
 		}
 
-		uint32 Size = Verts.GetResourceDataSize();
-
-		// Create vertex buffer. Fill buffer with initial data upon creation
-		FRHIResourceCreateInfo CreateInfo(TEXT("FStencilConeVertexBuffer"), &Verts);
-		VertexBufferRHI = RHICmdList.CreateVertexBuffer(Size, BUF_Static, CreateInfo);
+		VertexBufferRHI = UE::RHIResourceUtils::CreateVertexBufferFromArray(RHICmdList, TEXT("FStencilConeVertexBuffer"), EBufferUsageFlags::Static, MakeConstArrayView(Verts));
 	}
 
 	int32 GetVertexCount() const { return NumVerts; }
diff --git a/Engine/Source/Runtime/Renderer/Private/LightRendering.h b/Engine/Source/Runtime/Renderer/Private/LightRendering.h
--- a/Engine/Source/Runtime/Renderer/Private/LightRendering.h
+++ b/Engine/Source/Runtime/Renderer/Private/LightRendering.h
@@ -8,6 +8,7 @@
 
 #include "CoreMinimal.h"
 #include "RHI.h"
+#include "RHIResourceUtils.h"
 #include "RenderResource.h"
 #include "UniformBuffer.h"
 #include "ShaderParameters.h"
@@ -136,7 +137,7 @@ namespace StencilingGeometry
 				ArcVerts.Add(FVector3f(0.0f, FMath::Sin(Angle), FMath::Cos(Angle)));
 			}
 
-			TResourceArray<VectorType, VERTEXBUFFER_ALIGNMENT> Verts;
+			TArray<VectorType> Verts;
 			Verts.Empty(NumVerts);
 			// Then rotate this arc NumSides + 1 times.
 			const FVector3f Center = FVector3f(0,0,0);
@@ -153,11 +154,9 @@ namespace StencilingGeometry
 			}
 
 			NumSphereVerts = Verts.Num();
-			uint32 Size = Verts.GetResourceDataSize();
 
 			// Create vertex buffer. Fill buffer with initial data upon creation
-			FRHIResourceCreateInfo CreateInfo(TEXT("TStencilSphereVertexBuffer"), &Verts);
-			VertexBufferRHI = RHICmdList.CreateVertexBuffer(Size,BUF_Static,CreateInfo);
+			VertexBufferRHI = UE::RHIResourceUtils::CreateVertexBufferFromArray(RHICmdList, TEXT("TStencilSphereVertexBuffer"), EBufferUsageFlags::Static, MakeConstArrayView(Verts));
 		}
 
 		int32 GetVertexCount() const { return NumSphereVerts; }
@@ -204,7 +203,8 @@ namespace StencilingGeometry
 		{
 			const int32 NumSides = NumSphereSides;
 			const int32 NumRings = NumSphereRings;
-			TResourceArray<uint16, INDEXBUFFER_ALIGNMENT> Indices;
+			TArray<uint16> Indices;
+			Indices.Reserve(NumSides * NumRings * 6);
 
 			// Add triangles for all the vertices generated
 			for (int32 s = 0; s < NumSides; s++)
@@ -224,12 +224,9 @@ namespace StencilingGeometry
 			}
 
 			NumIndices = Indices.Num();
-			const uint32 Size = Indices.GetResourceDataSize();
-			const uint32 Stride = sizeof(uint16);
 
 			// Create index buffer. Fill buffer with initial data upon creation
-			FRHIResourceCreateInfo CreateInfo(TEXT("TStencilSphereIndexBuffer"), &Indices);
-			IndexBufferRHI = RHICmdList.CreateIndexBuffer(Stride, Size, BUF_Static, CreateInfo);
+			IndexBufferRHI = UE::RHIResourceUtils::CreateIndexBufferFromArray(RHICmdList, TEXT("TStencilSphereIndexBuffer"), EBufferUsageFlags::Static, MakeConstArrayView(Indices));
 		}
 
 		int32 GetIndexCount() const { return NumIndices; }; 
diff --git a/Engine/Source/Runtime/Renderer/Private/PostProcess/SceneFilterRendering.cpp b/Engine/Source/Runtime/Renderer/Private/PostProcess/SceneFilterRendering.cpp
--- a/Engine/Source/Runtime/Renderer/Private/PostProcess/SceneFilterRendering.cpp
+++ b/Engine/Source/Runtime/Renderer/Private/PostProcess/SceneFilterRendering.cpp
@@ -14,10 +14,11 @@
 #include "CommonRenderResources.h"
 #include "PostProcess/DrawRectangle.h"
 #include "ScenePrivate.h"
+#include "RHIResourceUtils.h"
 
 void FTesselatedScreenRectangleIndexBuffer::InitRHI(FRHICommandListBase& RHICmdList)
 {
-	TResourceArray<uint16, INDEXBUFFER_ALIGNMENT> IndexBuffer;
+	TArray<uint16> IndexBuffer;
 
 	uint32 NumIndices = NumPrimitives() * 3;
 	IndexBuffer.AddUninitialized(NumIndices);
@@ -45,8 +46,7 @@ void FTesselatedScreenRectangleIndexBuffer::InitRHI(FRHICommandListBase& RHICmdL
 	}
 
 	// Create index buffer. Fill buffer with initial data upon creation
-	FRHIResourceCreateInfo CreateInfo(TEXT("FTesselatedScreenRectangleIndexBuffer"), &IndexBuffer);
-	IndexBufferRHI = RHICmdList.CreateIndexBuffer(sizeof(uint16), IndexBuffer.GetResourceDataSize(), BUF_Static, CreateInfo);
+	IndexBufferRHI = UE::RHIResourceUtils::CreateIndexBufferFromArray(RHICmdList, TEXT("FTesselatedScreenRectangleIndexBuffer"), EBufferUsageFlags::Static, MakeConstArrayView(IndexBuffer));
 }
 
 uint32 FTesselatedScreenRectangleIndexBuffer::NumVertices() const
diff --git a/Engine/Source/Runtime/Renderer/Private/RayTracing/RayTracingDebug.cpp b/Engine/Source/Runtime/Renderer/Private/RayTracing/RayTracingDebug.cpp
--- a/Engine/Source/Runtime/Renderer/Private/RayTracing/RayTracingDebug.cpp
+++ b/Engine/Source/Runtime/Renderer/Private/RayTracing/RayTracingDebug.cpp
@@ -18,6 +18,7 @@
 #include "RayTracing/RayTracing.h"
 #include "RayTracing/RaytracingOptions.h"
 #include "RayTracing/RayTracingTraversalStatistics.h"
+#include "RHIResourceUtils.h"
 #include "Nanite/NaniteRayTracing.h"
 #include "PixelShaderUtils.h"
 #include "SystemTextures.h"
@@ -509,8 +510,6 @@ public:
 	*/
 	void InitRHI(FRHICommandListBase& RHICmdList) override
 	{
-		TResourceArray<uint16, INDEXBUFFER_ALIGNMENT> Indices;
-
 		static const uint16 LineIndices[12 * 2] =
 		{
 			0, 1,
@@ -527,16 +526,8 @@ public:
 			4, 5
 		};
 
-		int32 NumIndices = UE_ARRAY_COUNT(LineIndices);
-		Indices.AddUninitialized(NumIndices);
-		FMemory::Memcpy(Indices.GetData(), LineIndices, NumIndices * sizeof(uint16));
-
-		const uint32 Size = Indices.GetResourceDataSize();
-		const uint32 Stride = sizeof(uint16);
-
 		// Create index buffer. Fill buffer with initial data upon creation
-		FRHIResourceCreateInfo CreateInfo(TEXT("FRayTracingDebugLineAABBIndexBuffer"), &Indices);
-		IndexBufferRHI = RHICmdList.CreateIndexBuffer(Stride, Size, BUF_Static, CreateInfo);
+		IndexBufferRHI = UE::RHIResourceUtils::CreateIndexBufferFromArray(RHICmdList, TEXT("FRayTracingDebugLineAABBIndexBuffer"), EBufferUsageFlags::Static, MakeConstArrayView(LineIndices));
 	}
 };
 
diff --git a/Engine/Source/Runtime/Renderer/Private/Tests/RayTracingTestbed.cpp b/Engine/Source/Runtime/Renderer/Private/Tests/RayTracingTestbed.cpp
--- a/Engine/Source/Runtime/Renderer/Private/Tests/RayTracingTestbed.cpp
+++ b/Engine/Source/Runtime/Renderer/Private/Tests/RayTracingTestbed.cpp
@@ -2,6 +2,7 @@
 
 #include "CoreMinimal.h"
 #include "RHI.h"
+#include "RHIResourceUtils.h"
 #include "Misc/AutomationTest.h"
 #include "Math/DoubleFloat.h"
 
@@ -65,31 +66,25 @@ bool RunRayTracingTestbed_RenderThread(const FString& Parameters)
 	FRHICommandListImmediate& RHICmdList = FRHICommandListImmediate::Get();
 
 	{
-		TResourceArray<FVector3f> PositionData;
-		PositionData.SetNumUninitialized(3);
-		PositionData[0] = FVector3f( 1, -1, 0);
-		PositionData[1] = FVector3f( 1,  1, 0);
-		PositionData[2] = FVector3f(-1, -1, 0);
-
-		FRHIResourceCreateInfo CreateInfo(TEXT("RayTracingTestbedVB"));
-		CreateInfo.ResourceArray = &PositionData;
+		const FVector3f PositionData[] =
+		{
+			FVector3f( 1, -1, 0),
+			FVector3f( 1,  1, 0),
+			FVector3f(-1, -1, 0),
+		};
 
-		VertexBuffer = RHICmdList.CreateVertexBuffer(PositionData.GetResourceDataSize(), BUF_Static, CreateInfo);
+		VertexBuffer = UE::RHIResourceUtils::CreateVertexBufferFromArray(RHICmdList, TEXT("RayTracingTestbedVB"), EBufferUsageFlags::Static, MakeConstArrayView(PositionData));
 	}
 
 	FBufferRHIRef IndexBuffer;
 
 	{
-		TResourceArray<uint16> IndexData;
-		IndexData.SetNumUninitialized(3);
-		IndexData[0] = 0;
-		IndexData[1] = 1;
-		IndexData[2] = 2;
-
-		FRHIResourceCreateInfo CreateInfo(TEXT("RayTracingTestbedIB"));
-		CreateInfo.ResourceArray = &IndexData;
+		const uint16 IndexData[] =
+		{
+			0, 1, 2
+		};
 
-		IndexBuffer = RHICmdList.CreateIndexBuffer(2, IndexData.GetResourceDataSize(), BUF_Static, CreateInfo);
+		IndexBuffer = UE::RHIResourceUtils::CreateIndexBufferFromArray(RHICmdList, TEXT("RayTracingTestbedIB"), EBufferUsageFlags::Static, MakeConstArrayView(IndexData));
 	}
 
 	static constexpr uint32 NumRays = 4;
@@ -98,22 +93,22 @@ bool RunRayTracingTestbed_RenderThread(const FString& Parameters)
 	FShaderResourceViewRHIRef RayBufferView;
 
 	{
-		TResourceArray<FBasicRayTracingRay> RayData;
-		RayData.SetNumUninitialized(NumRays);
-		RayData[0] = FBasicRayTracingRay{ { 0.75f, 0.0f, -1.0f}, 0xFFFFFFFF, {0.0f, 0.0f,  1.0f}, 100000.0f }; // expected to hit
-		RayData[1] = FBasicRayTracingRay{ { 0.75f, 0.0f, -1.0f}, 0xFFFFFFFF, {0.0f, 0.0f,  1.0f},      0.5f }; // expected to miss (short ray)
-		RayData[2] = FBasicRayTracingRay{ { 0.75f, 0.0f,  1.0f}, 0xFFFFFFFF, {0.0f, 0.0f, -1.0f}, 100000.0f }; // expected to hit  (should hit back face)
-		RayData[3] = FBasicRayTracingRay{ {-0.75f, 0.0f, -1.0f}, 0xFFFFFFFF, {0.0f, 0.0f,  1.0f}, 100000.0f }; // expected to miss (doesn't intersect)
-
-		FRHIResourceCreateInfo CreateInfo(TEXT("RayBuffer"));
-		CreateInfo.ResourceArray = &RayData;
-
-		RayBuffer = RHICmdList.CreateBuffer(RayData.GetResourceDataSize(), 
-			BUF_Static | BUF_ShaderResource | BUF_StructuredBuffer, 
-			sizeof(FBasicRayTracingRay),
+		const FBasicRayTracingRay RayData[] =
+		{
+			FBasicRayTracingRay{ { 0.75f, 0.0f, -1.0f}, 0xFFFFFFFF, {0.0f, 0.0f,  1.0f}, 100000.0f }, // expected to hit
+			FBasicRayTracingRay{ { 0.75f, 0.0f, -1.0f}, 0xFFFFFFFF, {0.0f, 0.0f,  1.0f},      0.5f }, // expected to miss (short ray)
+			FBasicRayTracingRay{ { 0.75f, 0.0f,  1.0f}, 0xFFFFFFFF, {0.0f, 0.0f, -1.0f}, 100000.0f }, // expected to hit  (should hit back face)
+			FBasicRayTracingRay{ {-0.75f, 0.0f, -1.0f}, 0xFFFFFFFF, {0.0f, 0.0f,  1.0f}, 100000.0f }, // expected to miss (doesn't intersect)
+		};
+
+		RayBuffer = UE::RHIResourceUtils::CreateBufferFromArray(
+			RHICmdList,
+			TEXT("RayBuffer"), 
+			BUF_Static | BUF_ShaderResource | BUF_StructuredBuffer,
 			ERHIAccess::SRVMask,
-			CreateInfo
+			MakeConstArrayView(RayData)
 		);
+
 		RayBufferView = RHICmdList.CreateShaderResourceView(RayBuffer, 
 			FRHIViewDesc::CreateBufferSRV()
 			.SetType(FRHIViewDesc::EBufferType::Structured)
diff --git a/Engine/Source/Runtime/Renderer/Private/VolumetricFog.cpp b/Engine/Source/Runtime/Renderer/Private/VolumetricFog.cpp
--- a/Engine/Source/Runtime/Renderer/Private/VolumetricFog.cpp
+++ b/Engine/Source/Runtime/Renderer/Private/VolumetricFog.cpp
@@ -30,6 +30,7 @@ VolumetricFog.cpp
 #include "LightFunctionAtlas.h"
 #include "Math/UnrealMathUtility.h"
 #include "RayTracing/RayTracing.h"
+#include "RHIResourceUtils.h"
 
 using namespace LightFunctionAtlas;
 
@@ -677,7 +678,7 @@ public:
 	{
 		const int32 NumTriangles = FCircleRasterizeVertexBuffer::NumVertices - 2;
 
-		TResourceArray<uint16, INDEXBUFFER_ALIGNMENT> Indices;
+		TArray<uint16> Indices;
 		Indices.Empty(NumTriangles * 3);
 
 		for (int32 TriangleIndex = 0; TriangleIndex < NumTriangles; TriangleIndex++)
@@ -688,12 +689,8 @@ public:
 			Indices.Add(LeadingVertexIndex);
 		}
 
-		const uint32 Size = Indices.GetResourceDataSize();
-		const uint32 Stride = sizeof(uint16);
-
 		// Create index buffer. Fill buffer with initial data upon creation
-		FRHIResourceCreateInfo CreateInfo(TEXT("FCircleRasterizeIndexBuffer"), &Indices);
-		IndexBufferRHI = RHICmdList.CreateIndexBuffer(Stride, Size, BUF_Static, CreateInfo);
+		IndexBufferRHI = UE::RHIResourceUtils::CreateIndexBufferFromArray(RHICmdList, TEXT("FCircleRasterizeIndexBuffer"), EBufferUsageFlags::Static, MakeConstArrayView(Indices));
 	}
 };
 
diff --git a/Engine/Source/Runtime/SlateRHIRenderer/Private/SlateElementVertexBuffer.h b/Engine/Source/Runtime/SlateRHIRenderer/Private/SlateElementVertexBuffer.h
--- a/Engine/Source/Runtime/SlateRHIRenderer/Private/SlateElementVertexBuffer.h
+++ b/Engine/Source/Runtime/SlateRHIRenderer/Private/SlateElementVertexBuffer.h
@@ -8,6 +8,7 @@
 #include "SlateGlobals.h"
 #include "RHI.h"
 #include "RHICommandList.h"
+#include "RHIResourceUtils.h"
 #include "RenderResource.h"
 #include "Containers/ResourceArray.h"
 
@@ -172,11 +173,9 @@ public:
 	{
 		if (!IsValidRef(VertexBufferRHI))
 		{
-			static FStencilBufferResourceArray ResourceArray;
+			const uint32 Verts[] = { 0, 1, 2, 3 };
 
-			FRHIResourceCreateInfo CreateInfo(TEXT("SlateStencilClipVertexBuffer"));
-			CreateInfo.ResourceArray = &ResourceArray;
-			VertexBufferRHI = RHICmdList.CreateVertexBuffer(ResourceArray.GetResourceDataSize(), BUF_Static, CreateInfo);
+			VertexBufferRHI = UE::RHIResourceUtils::CreateVertexBufferFromArray(RHICmdList, TEXT("SlateStencilClipVertexBuffer"), EBufferUsageFlags::Static, MakeConstArrayView(Verts));
 
 			// Ensure the vertex buffer could be created
 			check(IsValidRef(VertexBufferRHI));
@@ -191,24 +190,4 @@ public:
 
 	/** Returns a friendly name for this buffer. */
 	virtual FString GetFriendlyName() const { return TEXT("SlateElementVertices"); }
-
-private:
-	struct FStencilBufferResourceArray : FResourceArrayInterface
-	{
-		virtual const void* GetResourceData() const override
-		{
-			static uint32 Verts[] = { 0, 1, 2, 3};
-			return Verts;
-		}
-
-		virtual uint32 GetResourceDataSize() const override
-		{
-			return sizeof(uint32) * 4;
-		}
-
-		virtual void Discard() override {}
-		virtual bool IsStatic() const override { return true; }
-		virtual bool GetAllowCPUAccess() const override { return false; }
-		virtual void SetAllowCPUAccess(bool bInNeedsCPUAccess) override { }
-	};
 };
