Revert to dynamic string Insights API for RHI breadcrumbs due to accumulation of marker strings causing memory pressure at runtime, and performance issues when analysing captures.
 - Original implementation in CL 34739050. This change undoes that CL, except for the deleted copy constructors.

diff --git a/Engine/Source/Runtime/RHI/Public/RHIBreadcrumbs.h b/Engine/Source/Runtime/RHI/Public/RHIBreadcrumbs.h
--- a/Engine/Source/Runtime/RHI/Public/RHIBreadcrumbs.h
+++ b/Engine/Source/Runtime/RHI/Public/RHIBreadcrumbs.h
@@ -444,21 +444,6 @@
 	public:
 		uint32 const ID;
 
-#if RHI_BREADCRUMBS_EMIT_CPU
-		uint32 const CPUTraceMarkerID = 0;
-#endif
-
-		void CreateTraceMarkers()
-		{
-#if RHI_BREADCRUMBS_EMIT_CPU
-			if (TRACE_CPUPROFILER_EVENT_MANUAL_IS_ENABLED())
-			{
-				FBuffer Buffer;
-				const_cast<uint32&>(CPUTraceMarkerID) = FCpuProfilerTrace::OutputEventType(GetTCHAR(Buffer));
-			}
-#endif
-		}
-
 		struct FBuffer
 		{
 			TCHAR Data[128];
@@ -481,9 +466,11 @@
 	inline void FRHIBreadcrumbNode::BeginCPU() const
 	{
 #if RHI_BREADCRUMBS_EMIT_CPU
-		if (Name.CPUTraceMarkerID)
+		if (TRACE_CPUPROFILER_EVENT_MANUAL_IS_ENABLED())
 		{
-			FCpuProfilerTrace::OutputBeginEvent(Name.CPUTraceMarkerID);
+			FRHIBreadcrumb::FBuffer Buffer;
+			TCHAR const* Str = Name.GetTCHAR(Buffer);
+			FCpuProfilerTrace::OutputBeginDynamicEvent(Str);
 		}
 #endif
 	}
@@ -491,7 +478,7 @@
 	inline void FRHIBreadcrumbNode::EndCPU() const
 	{
 #if RHI_BREADCRUMBS_EMIT_CPU
-		if (Name.CPUTraceMarkerID)
+		if (TRACE_CPUPROFILER_EVENT_MANUAL_IS_ENABLED())
 		{
 			FCpuProfilerTrace::OutputEndEvent();
 		}
@@ -658,9 +645,7 @@
 		TRHIBreadcrumb(TCHAR const* FormatString, TArgs const&... Args)
 			: FormatString(FormatString)
 			, Values(Args...)
-		{
-			CreateTraceMarkers();
-		}
+		{}
 
 		virtual TCHAR const* GetTCHAR(FBuffer& Buffer) const override
 		{
@@ -687,9 +672,7 @@
 	public:
 		TRHIBreadcrumb(TCHAR const* StringLiteral)
 			: StringLiteral(StringLiteral)
-		{
-			CreateTraceMarkers();
-		}
+		{}
 
 		virtual TCHAR const* GetTCHAR(FBuffer&) const override
 		{
