Fixing unsupported target for linear textures error in API validation layer error in TextureArrayReadback RHI unit test on a specific platform using a Texture 2D staging texture when reading back to the CPU.

diff --git a/Engine/Plugins/Tests/RHITests/Source/RHITests/Private/RHIReadbackTests.cpp b/Engine/Plugins/Tests/RHITests/Source/RHITests/Private/RHIReadbackTests.cpp
--- a/Engine/Plugins/Tests/RHITests/Source/RHITests/Private/RHIReadbackTests.cpp
+++ b/Engine/Plugins/Tests/RHITests/Source/RHITests/Private/RHIReadbackTests.cpp
@@ -313,6 +313,13 @@ bool FRHIReadbackTests::Test_TextureReadback(FRHICommandListImmediate& RHICmdLis
 	FRHIGPUTextureReadback Texture2DReadback(TEXT("Texture2DReadback"));
 	Texture2DReadback.EnqueueCopy(RHICmdList, Texture2D, ReadbackOffset2D, 0, ReadbackSize2D);
 
+	//Some platform only support readback on single 2d texture therefore the test can happen only on one slice
+	if (GRHIGlobals.SupportLinearTextureVolumeFormat == false)
+	{
+		ReadbackOffset3D.Z = 0;
+		ReadbackSize3D.Z = 1;
+	}
+
 	FRHIGPUTextureReadback Texture3DReadback(TEXT("Texture3DReadback"));
 	Texture3DReadback.EnqueueCopy(RHICmdList, Texture3D, ReadbackOffset3D, 0, ReadbackSize3D);
 
diff --git a/Engine/Source/Runtime/RHI/Private/RHIGPUReadback.cpp b/Engine/Source/Runtime/RHI/Private/RHIGPUReadback.cpp
--- a/Engine/Source/Runtime/RHI/Private/RHIGPUReadback.cpp
+++ b/Engine/Source/Runtime/RHI/Private/RHIGPUReadback.cpp
@@ -144,7 +144,22 @@ void FRHIGPUTextureReadback::EnqueueCopy(FRHICommandList& RHICmdList, FRHITextur
 			SCOPED_GPU_MASK(RHICmdList, FRHIGPUMask::FromIndex(GPUIndex));
 
 			// Assume for now that every enqueue happens on a texture of the same format and size (when reused).
-			if (!DestinationStagingTextures[GPUIndex] || DestinationStagingTextures[GPUIndex]->GetDesc().Dimension != SourceTexture->GetDesc().Dimension)
+			bool bNeedTextureCreation = true;
+
+			if (DestinationStagingTextures[GPUIndex])
+			{
+				if (GRHIGlobals.SupportLinearTextureVolumeFormat)
+				{
+					bNeedTextureCreation = DestinationStagingTextures[GPUIndex]->GetDesc().Dimension != SourceTexture->GetDesc().Dimension;
+				}
+				else
+				{
+					// Some platform support only 2d texture for readback, so assuming that if we have a staging texture is always a 2d texture 
+					bNeedTextureCreation = false;
+				}
+			}
+
+			if (bNeedTextureCreation)
 			{
 				FIntVector StagingTextureSize;
 
@@ -164,8 +179,10 @@ void FRHIGPUTextureReadback::EnqueueCopy(FRHICommandList& RHICmdList, FRHITextur
 					.SetExtent(StagingTextureSize.X, StagingTextureSize.Y)
 					.SetFlags(ETextureCreateFlags::CPUReadback | ETextureCreateFlags::HideInVisualizeTexture);
 
-				switch (Desc.Dimension)
+				if (GRHIGlobals.SupportLinearTextureVolumeFormat)
 				{
+					switch (Desc.Dimension)
+					{
 					case ETextureDimension::Texture2DArray:
 					case ETextureDimension::TextureCubeArray:
 						ensureMsgf(Size.Z <= 1, TEXT("Readback for texture arrays supports only one slice at a time. Texture Name: %s, SourcePosition: (%d, %d, %d), SourceSlice: %u, Size: (%d, %d, %d)."),
@@ -175,6 +192,17 @@ void FRHIGPUTextureReadback::EnqueueCopy(FRHICommandList& RHICmdList, FRHITextur
 					case ETextureDimension::Texture3D:
 						Desc.SetDepth(StagingTextureSize.Z);
 						break;
+					}
+				}
+				else
+				{
+					ensureMsgf(Size.Z <= 1, TEXT("Readback for texture arrays/volume texture supports only one slice at a time. Texture Name: %s, SourcePosition: (%d, %d, %d), SourceSlice: %u, Size: (%d, %d, %d)."),
+						*SourceTexture->GetName().ToString(), SourcePosition.X, SourcePosition.Y, SourcePosition.Z, SourceSlice, Size.X, Size.Y, Size.Z);
+
+					// Some platforms only support 2d texture for readback, create a 2d texture staging texture and force the slice size to 1
+					Desc.SetArraySize(1);
+					Desc.SetDepth(1);
+					Desc.SetDimension(ETextureDimension::Texture2D);
 				}
 
 				DestinationStagingTextures[GPUIndex] = RHICmdList.CreateTexture(Desc);
diff --git a/Engine/Source/Runtime/RHI/Public/RHIGlobals.h b/Engine/Source/Runtime/RHI/Public/RHIGlobals.h
--- a/Engine/Source/Runtime/RHI/Public/RHIGlobals.h
+++ b/Engine/Source/Runtime/RHI/Public/RHIGlobals.h
@@ -673,6 +673,9 @@ struct FRHIGlobals
 
 	/** Whether Depth Stencil MSAA Resolve Targets are supported. */
 	bool SupportsDepthStencilResolve = false;
+
+	/** True if RHI supports Linear texture format in 3D/Cube/Array texture */
+	bool SupportLinearTextureVolumeFormat = true;
 };
 
 extern RHI_API FRHIGlobals GRHIGlobals;
