[RHI] Re-enabled the optimization to register CPU scopes in constructor of FRHIBreadcrumb, originaly implemented in CL 34739050 by Luke T. (and reverted in CL 35340934), but this time using the new FCpuProfilerTrace::OutputDynamicEventType API.

diff --git a/Engine/Source/Runtime/RHI/Public/RHIBreadcrumbs.h b/Engine/Source/Runtime/RHI/Public/RHIBreadcrumbs.h
--- a/Engine/Source/Runtime/RHI/Public/RHIBreadcrumbs.h
+++ b/Engine/Source/Runtime/RHI/Public/RHIBreadcrumbs.h
@@ -444,6 +444,22 @@
 	public:
 		uint32 const ID;
 
+#if RHI_BREADCRUMBS_EMIT_CPU
+		uint32 const CPUTraceMarkerID = 0;
+#endif
+
+		void CreateTraceMarkers()
+		{
+#if RHI_BREADCRUMBS_EMIT_CPU
+			if (TRACE_CPUPROFILER_EVENT_MANUAL_IS_ENABLED())
+			{
+				FBuffer Buffer;
+				TCHAR const* Str = GetTCHAR(Buffer);
+				const_cast<uint32&>(CPUTraceMarkerID) = FCpuProfilerTrace::OutputDynamicEventType(Str, __FILE__, __LINE__);
+			}
+#endif
+		}
+
 		struct FBuffer
 		{
 			TCHAR Data[128];
@@ -466,11 +482,9 @@
 	inline void FRHIBreadcrumbNode::BeginCPU() const
 	{
 #if RHI_BREADCRUMBS_EMIT_CPU
-		if (TRACE_CPUPROFILER_EVENT_MANUAL_IS_ENABLED())
+		if (Name.CPUTraceMarkerID)
 		{
-			FRHIBreadcrumb::FBuffer Buffer;
-			TCHAR const* Str = Name.GetTCHAR(Buffer);
-			FCpuProfilerTrace::OutputBeginDynamicEvent(Str);
+			FCpuProfilerTrace::OutputBeginEvent(Name.CPUTraceMarkerID);
 		}
 #endif
 	}
@@ -478,7 +492,7 @@
 	inline void FRHIBreadcrumbNode::EndCPU() const
 	{
 #if RHI_BREADCRUMBS_EMIT_CPU
-		if (TRACE_CPUPROFILER_EVENT_MANUAL_IS_ENABLED())
+		if (Name.CPUTraceMarkerID)
 		{
 			FCpuProfilerTrace::OutputEndEvent();
 		}
@@ -645,7 +659,9 @@
 		TRHIBreadcrumb(TCHAR const* FormatString, TArgs const&... Args)
 			: FormatString(FormatString)
 			, Values(Args...)
-		{}
+		{
+			CreateTraceMarkers();
+		}
 
 		virtual TCHAR const* GetTCHAR(FBuffer& Buffer) const override
 		{
@@ -672,7 +688,9 @@
 	public:
 		TRHIBreadcrumb(TCHAR const* StringLiteral)
 			: StringLiteral(StringLiteral)
-		{}
+		{
+			CreateTraceMarkers();
+		}
 
 		virtual TCHAR const* GetTCHAR(FBuffer&) const override
 		{
