RHI - Fixed crash when accessing pipeline after the RHI context is cleaned up.

diff --git a/Engine/Source/Runtime/RHI/Private/RHICommandList.cpp b/Engine/Source/Runtime/RHI/Private/RHICommandList.cpp
--- a/Engine/Source/Runtime/RHI/Private/RHICommandList.cpp
+++ b/Engine/Source/Runtime/RHI/Private/RHICommandList.cpp
@@ -949,11 +949,12 @@ FGraphEventRef FRHICommandListExecutor::FTranslateState::Finalize()
 				
 				GDynamicRHI->RHIFinalizeContext(MoveTemp(FinalizeArgs), PlatformCommandLists);
 				
-				for (auto& State : PipelineStates)
+				for (ERHIPipeline Pipeline : MakeFlagsRange(ERHIPipeline::All))
 				{
+					auto& State = PipelineStates[Pipeline]; 
 					if (State.Context)
 					{
-						State.FinalizedCmdList = PlatformCommandLists[State.Context->GetPipeline()];
+						State.FinalizedCmdList = PlatformCommandLists[Pipeline];
 					}
 				}
 			}
