Add draw calls and primitives counts to new RHI GPU profiler event stream
 - Platform RHIs emit the FStats event type, which accumulates the counts to the current active breadcrumb node.

diff --git a/Engine/Source/Runtime/D3D12RHI/Private/D3D12CommandContext.cpp b/Engine/Source/Runtime/D3D12RHI/Private/D3D12CommandContext.cpp
--- a/Engine/Source/Runtime/D3D12RHI/Private/D3D12CommandContext.cpp
+++ b/Engine/Source/Runtime/D3D12RHI/Private/D3D12CommandContext.cpp
@@ -200,6 +200,8 @@ void FD3D12ContextCommon::BindDiagnosticBuffer(FD3D12RootSignature const* RootSi
 
 	#if RHI_NEW_GPU_PROFILER
 		{
+			FlushProfilerStats();
+
 			auto& Event = GetCommandList().EmplaceEvent<UE::RHI::GPUProfiler::FEvent::FBeginBreadcrumb>();
 			Event.Breadcrumb = Breadcrumb;
 
@@ -222,6 +224,8 @@ void FD3D12ContextCommon::BindDiagnosticBuffer(FD3D12RootSignature const* RootSi
 	{
 	#if RHI_NEW_GPU_PROFILER
 		{
+			FlushProfilerStats();
+
 			auto& Event = GetCommandList().EmplaceEvent<UE::RHI::GPUProfiler::FEvent::FEndBreadcrumb>();
 			Event.Breadcrumb = Breadcrumb;
 
@@ -488,6 +492,10 @@ void FD3D12CommandContext::Finalize(TArray<FD3D12Payload*>& OutPayloads)
 	GetParentDevice()->GetBindlessDescriptorManager().FinalizeContext(*this);
 #endif
 
+#if RHI_NEW_GPU_PROFILER
+	FlushProfilerStats();
+#endif
+
 	FD3D12ContextCommon::Finalize(OutPayloads);
 }
 
diff --git a/Engine/Source/Runtime/D3D12RHI/Private/D3D12CommandContext.h b/Engine/Source/Runtime/D3D12RHI/Private/D3D12CommandContext.h
--- a/Engine/Source/Runtime/D3D12RHI/Private/D3D12CommandContext.h
+++ b/Engine/Source/Runtime/D3D12RHI/Private/D3D12CommandContext.h
@@ -527,6 +527,18 @@ public:
 	void ClearShaderResources(FD3D12BaseShaderResource* Resource, EShaderParameterTypeMask ShaderParameterTypeMask);
 	void ClearAllShaderResources();
 
+#if RHI_NEW_GPU_PROFILER
+	void FlushProfilerStats()
+	{
+		// Flush accumulated draw stats
+		if (StatEvent)
+		{
+			GetCommandList().EmplaceEvent<UE::RHI::GPUProfiler::FEvent::FStats>() = StatEvent;
+			StatEvent = {};
+		}
+	}
+#endif
+
 	FD3D12FastConstantAllocator ConstantsAllocator;
 
 	// Current GPU event stack
diff --git a/Engine/Source/Runtime/RHI/Private/GPUProfiler.cpp b/Engine/Source/Runtime/RHI/Private/GPUProfiler.cpp
--- a/Engine/Source/Runtime/RHI/Private/GPUProfiler.cpp
+++ b/Engine/Source/Runtime/RHI/Private/GPUProfiler.cpp
@@ -875,6 +875,9 @@ namespace UE::RHI::GPUProfiler
 			uint32 Level = 0;
 			TArray<uint64> Timestamps;
 
+			uint32 NumDraws = 0;
+			uint32 NumPrimitives = 0;
+
 			uint64 BusyCycles = 0;
 
 			FNode(FString&& Name)
@@ -1060,6 +1063,14 @@ namespace UE::RHI::GPUProfiler
 					break;
 			#endif // WITH_RHI_BREADCRUMBS
 
+				case FEvent::EType::Stats:
+					{
+						auto& Stats = Event->Value.Get<FEvent::FStats>();
+						Current->NumDraws += Stats.NumDraws;
+						Current->NumPrimitives += Stats.NumPrimitives;
+					}
+					break;
+
 				case FEvent::EType::FrameBoundary:
 					{
 						FEvent::FFrameBoundary& FrameBoundary = Event->Value.Get<FEvent::FFrameBoundary>();
@@ -1151,15 +1162,20 @@ namespace UE::RHI::GPUProfiler
 			}
 		}
 
-		double RootMilliseconds = FPlatformTime::ToMilliseconds64(First->BusyCycles);
-
 		FString LogMessage;
 		for (FNode* Node = First; Node; Node = Node->Next)
 		{
 			double Milliseconds = FPlatformTime::ToMilliseconds64(Node->BusyCycles);
 
-			TUnicodeHorizontalBar<8> Bar = Milliseconds / RootMilliseconds;
-			LogMessage += FString::Printf(TEXT("%9.3f ms |%s| %*s\n"), Milliseconds, Bar.Text, Node->Name.Len() + (Node->Level * 4), *Node->Name);
+			TUnicodeHorizontalBar<8> Bar = Milliseconds / FPlatformTime::ToMilliseconds64(First->BusyCycles);
+			LogMessage += FString::Printf(TEXT("%9.3f ms |%s| %6d | %6d | %*s\n")
+				, Milliseconds
+				, Bar.Text
+				, Node->NumDraws
+				, Node->NumPrimitives
+				, Node->Name.Len() + (Node->Level * 4)
+				, *Node->Name
+			);
 		}
 
 		UE_LOG(LogRHI, Display, TEXT("GPU Profile for Frame %d, Queue [%s, GPU: %d, Idx: %d]:\n%s\n\n"), FrameNumber, Queue.GetTypeString(), Queue.GPU, Queue.Index, *LogMessage);
diff --git a/Engine/Source/Runtime/RHI/Public/GPUProfiler.h b/Engine/Source/Runtime/RHI/Public/GPUProfiler.h
--- a/Engine/Source/Runtime/RHI/Public/GPUProfiler.h
+++ b/Engine/Source/Runtime/RHI/Public/GPUProfiler.h
@@ -7,8 +7,9 @@
 #pragma once
 
 #include "CoreMinimal.h"
-#include "Templates/RefCounting.h"
-#include "RHI.h"
+#include "Misc/TVariant.h"
+
+#include "RHIBreadcrumbs.h"
 
 #if RHI_NEW_GPU_PROFILER
 
@@ -124,6 +125,17 @@ namespace UE::RHI::GPUProfiler
 			uint64 GPUTimestampBOP;
 		};
 
+		struct FStats
+		{
+			uint32 NumDraws;
+			uint32 NumPrimitives;
+
+			operator bool() const
+			{
+				return NumDraws > 0 || NumPrimitives > 0;
+			}
+		};
+
 		struct FSignalFence
 		{
 			uint64 CPUTimestamp;
@@ -154,6 +166,7 @@ namespace UE::RHI::GPUProfiler
 #endif
 			, FBeginWork
 			, FEndWork
+			, FStats
 			, FSignalFence
 			, FWaitFence
 			, FFlip
@@ -169,6 +182,7 @@ namespace UE::RHI::GPUProfiler
 #endif
 			BeginWork       = FStorage::IndexOfType<FBeginWork      >(),
 			EndWork         = FStorage::IndexOfType<FEndWork        >(),
+			Stats           = FStorage::IndexOfType<FStats          >(),
 			SignalFence     = FStorage::IndexOfType<FSignalFence    >(),
 			WaitFence       = FStorage::IndexOfType<FWaitFence      >(),
 			Flip            = FStorage::IndexOfType<FFlip           >(),
diff --git a/Engine/Source/Runtime/RHI/Public/RHIBreadcrumbs.h b/Engine/Source/Runtime/RHI/Public/RHIBreadcrumbs.h
--- a/Engine/Source/Runtime/RHI/Public/RHIBreadcrumbs.h
+++ b/Engine/Source/Runtime/RHI/Public/RHIBreadcrumbs.h
@@ -2,16 +2,15 @@
 
 #pragma once
 
-#include "Containers/Array.h"
-#include "RHIFwd.h"
-#include "RHIDefinitions.h"
-#include "UObject/NameTypes.h"
+#include "Misc/Build.h"
 #include "Misc/MemStack.h"
+#include "HAL/Platform.h"
 #include "Templates/SharedPointer.h"
+#include "GenericPlatform/GenericPlatformCrashContext.h"
+
+#include "RHIFwd.h"
 #include "RHIPipeline.h"
-#include "Math/Color.h"
 #include "MultiGPU.h"
-#include "GenericPlatform/GenericPlatformCrashContext.h"
 
 #ifndef WITH_RHI_BREADCRUMBS
 #define WITH_RHI_BREADCRUMBS ((UE_BUILD_DEBUG || UE_BUILD_DEVELOPMENT) || WITH_PROFILEGPU)
diff --git a/Engine/Source/Runtime/RHI/Public/RHIContext.h b/Engine/Source/Runtime/RHI/Public/RHIContext.h
--- a/Engine/Source/Runtime/RHI/Public/RHIContext.h
+++ b/Engine/Source/Runtime/RHI/Public/RHIContext.h
@@ -19,6 +19,7 @@
 #include "RHIBreadcrumbs.h"
 #include "RHIResources.h"
 #include "RHIShaderParameters.h"
+#include "GPUProfiler.h"
 
 class FRHIDepthRenderTargetView;
 class FRHIRenderTargetView;
@@ -517,6 +518,13 @@ public:
 	//
 	inline bool ShouldEmitBreadcrumbs() const;
 #endif
+
+protected:
+#if RHI_NEW_GPU_PROFILER
+	// Used to accumulate draw call and primitive counts,
+	// via the RHI_DRAW_CALL_INC / RHI_DRAW_CALL_STATS macros.
+	UE::RHI::GPUProfiler::FEvent::FStats StatEvent {};
+#endif
 };
 
 /** Context that is used to generate Upload commands. */
diff --git a/Engine/Source/Runtime/RHI/Public/RHIStats.h b/Engine/Source/Runtime/RHI/Public/RHIStats.h
--- a/Engine/Source/Runtime/RHI/Public/RHIStats.h
+++ b/Engine/Source/Runtime/RHI/Public/RHIStats.h
@@ -103,8 +103,13 @@ struct FRHIDrawStatsCategory
 
 // Macros for use inside RHI context Draw functions.
 // Updates the Stats structure on the executing RHI command list
-#define RHI_DRAW_CALL_INC()                              do { GetExecutingCommandList().Stats_AddDraw();                                          } while (false)
-#define RHI_DRAW_CALL_STATS(PrimitiveType,NumPrimitives) do { GetExecutingCommandList().Stats_AddDrawAndPrimitives(PrimitiveType, NumPrimitives); } while (false)
+#if RHI_NEW_GPU_PROFILER
+	#define RHI_DRAW_CALL_INC()             do { StatEvent.NumDraws++; GetExecutingCommandList().Stats_AddDraw(); } while (false)
+	#define RHI_DRAW_CALL_STATS(Type,Prims) do { StatEvent.NumDraws++; StatEvent.NumPrimitives += Prims; GetExecutingCommandList().Stats_AddDrawAndPrimitives(Type, Prims); } while (false)
+#else
+	#define RHI_DRAW_CALL_INC()             do { GetExecutingCommandList().Stats_AddDraw(); } while (false)
+	#define RHI_DRAW_CALL_STATS(Type,Prims) do { GetExecutingCommandList().Stats_AddDrawAndPrimitives(Type, Prims); } while (false)
+#endif
 
 struct FRHIDrawStats
 {
