- eEverting a change to cache a bindless setting, but seems to have had side effects

diff --git a/Engine/Source/Runtime/RHI/Private/RHI.cpp b/Engine/Source/Runtime/RHI/Private/RHI.cpp
--- a/Engine/Source/Runtime/RHI/Private/RHI.cpp
+++ b/Engine/Source/Runtime/RHI/Private/RHI.cpp
@@ -1178,25 +1178,9 @@ static ERHIBindlessConfiguration ParseConfigurationFromString(const FString& InS
 static bool GetBindlessConfigurationSetting(FString& OutSetting, EShaderPlatform Platform, const TCHAR* SettingName)
 {
 	const FString ShaderFormat = FDataDrivenShaderPlatformInfo::GetShaderFormat(Platform).ToString();
-	static TMap<FString, FString> ShaderFormatToBindlessConfig;
 	if (!ShaderFormat.IsEmpty())
 	{
-		FString* Config = ShaderFormatToBindlessConfig.Find(ShaderFormat);
-		if (Config == nullptr)
-		{
-			if (!GConfig->GetString(*ShaderFormat, SettingName, OutSetting, GEngineIni))
-			{
-				OutSetting = TEXT("");
-			}
-			ShaderFormatToBindlessConfig.Add(ShaderFormat, OutSetting);
-		}
-		else
-		{
-			OutSetting = *Config;
-		}
-
-		// return true if we ever found a setting
-		return OutSetting.Len() > 0;
+		return GConfig->GetString(*ShaderFormat, SettingName, OutSetting, GEngineIni);
 	}
 
 	return false;
