Wrapping a code section in DO_CHECK checks.

diff --git a/Engine/Source/Runtime/RHI/Private/RHIResources.cpp b/Engine/Source/Runtime/RHI/Private/RHIResources.cpp
--- a/Engine/Source/Runtime/RHI/Private/RHIResources.cpp
+++ b/Engine/Source/Runtime/RHI/Private/RHIResources.cpp
@@ -302,6 +302,7 @@ FRHIViewDesc::FTexture::FViewInfo FRHIViewDesc::FTexture::GetViewInfo(FRHITextur
 
 	checkf(ArrayRange.Num > 0 || ArrayRange.First == 0, TEXT("ArrayRange.Num cannot be zero, unless creating a view of the entire range."));
 
+#if DO_CHECK
 	// make sure the view fits in the texture
 	{
 		uint16 TextureArraySize = Desc.IsTextureCube() ? Desc.ArraySize * 6 : Desc.ArraySize;
@@ -323,6 +324,7 @@ FRHIViewDesc::FTexture::FViewInfo FRHIViewDesc::FTexture::GetViewInfo(FRHITextur
 			TextureArraySize / SliceDividerForCheckMessage
 		);
 	}
+#endif
 
 	// When ArrayRange.Num == 0, we use the number of elements from the texture. If the view is a 2D array and the texture a cube (array), we need to do x6 on the number of slices
 	// We already checked that we can only create cube views on cube textures, so we only need to take into account the 2D view on cube texture case
