Fix state validation when binding resources for graphics shaders.

RHISetShaderParameters was requiring SRVGraphics for graphics shaders, which is incorrect now that this state is a combination of bits (for pixel and non-pixel). It needs to check for the new SRVGraphicsPixel and SRVGraphicsNonPixel states based on the type of graphics shader that's passed in.

#rnx

diff --git a/Engine/Source/Runtime/RHI/Public/RHIValidationContext.h b/Engine/Source/Runtime/RHI/Public/RHIValidationContext.h
--- a/Engine/Source/Runtime/RHI/Public/RHIValidationContext.h
+++ b/Engine/Source/Runtime/RHI/Public/RHIValidationContext.h
@@ -804,8 +804,10 @@ public:
 
 		RHIValidation::FStageBoundUniformBuffers& BoundUniformBuffers = State.BoundUniformBuffers.Get(Shader->GetFrequency());
 
-		ValidateShaderParameters(Shader, Tracker, State.StaticUniformBuffers, BoundUniformBuffers, InResourceParameters, ERHIAccess::SRVGraphics, RHIValidation::EUAVMode::Graphics);
-		ValidateShaderParameters(Shader, Tracker, State.StaticUniformBuffers, BoundUniformBuffers, InBindlessParameters, ERHIAccess::SRVGraphics, RHIValidation::EUAVMode::Graphics);
+		ERHIAccess RequiredAccess = Shader->GetFrequency() == SF_Pixel ? ERHIAccess::SRVGraphicsPixel : ERHIAccess::SRVGraphicsNonPixel;
+
+		ValidateShaderParameters(Shader, Tracker, State.StaticUniformBuffers, BoundUniformBuffers, InResourceParameters, RequiredAccess, RHIValidation::EUAVMode::Graphics);
+		ValidateShaderParameters(Shader, Tracker, State.StaticUniformBuffers, BoundUniformBuffers, InBindlessParameters, RequiredAccess, RHIValidation::EUAVMode::Graphics);
 
 		RHIContext->RHISetShaderParameters(Shader, InParametersData, InParameters, InResourceParameters, InBindlessParameters);
 	}
