Refactor delegate-based "new PSO" reporting to be once per frame.

Currently, new PSOs are reported through the delegate as soon as they're encountered. This is problematic in multiple ways, but especially if this is being done during a GC cycle when the GC lock is being held and the broadcast tries to pin a UObject.

This change moves the reporting to be once per frame, in bulk for every new PSO that has been encountered since last frame.

diff --git a/Engine/Source/Runtime/RHI/Private/PipelineFileCache.cpp b/Engine/Source/Runtime/RHI/Private/PipelineFileCache.cpp
--- a/Engine/Source/Runtime/RHI/Private/PipelineFileCache.cpp
+++ b/Engine/Source/Runtime/RHI/Private/PipelineFileCache.cpp
@@ -112,7 +112,7 @@ static TAutoConsoleVariable<int32> CVarPSOFileCacheLogPSO(
 static TAutoConsoleVariable<int32> CVarPSOFileCacheReportPSO(
 														   TEXT("r.ShaderPipelineCache.ReportPSO"),
 														   PIPELINE_CACHE_DEFAULT_ENABLED,
-														   TEXT("1 reports new PSO entries via a delegate, but does not record or modify any cache file."),
+														   TEXT("1 reports new PSO entries via a delegate, but does not record or modify any cache file. New PSOs are reported in bulk once per frame."),
 														   ECVF_Default | ECVF_RenderThreadSafe
 														   );
 
@@ -181,6 +181,7 @@ TMap<uint32, FPSOUsageData> FPipelineFileCacheManager::RunTimeToPSOUsage;
 TMap<uint32, FPSOUsageData> FPipelineFileCacheManager::NewPSOUsage;
 TMap<uint32, FPipelineStateStats*> FPipelineFileCacheManager::Stats;
 TSet<FPipelineCacheFileFormatPSO> FPipelineFileCacheManager::NewPSOs;
+TArray<FPipelineCacheFileFormatPSO> FPipelineFileCacheManager::NewPSOsToReport;
 TSet<uint32> FPipelineFileCacheManager::NewPSOHashes;
 uint32 FPipelineFileCacheManager::NumNewPSOs;
 FString FPipelineFileCacheManager::UserCacheKey;
@@ -3523,6 +3524,24 @@ void FPipelineFileCacheManager::LogNewRaytracingPSOToConsole(FPipelineCacheFileF
 	}
 }
 
+void FPipelineFileCacheManager::BroadcastNewPSOsDelegate()
+{
+	TArray<FPipelineCacheFileFormatPSO> PSOs;
+	{
+		FRWScopeLock Lock(FileCacheLock, SLT_Write);
+		PSOs = MoveTemp(NewPSOsToReport);
+		NewPSOsToReport.Empty(32);
+	}
+
+	if (ReportNewPSOs() && PSOLoggedEvent.IsBound())
+	{
+		for (FPipelineCacheFileFormatPSO& PSO : PSOs)
+		{
+			PSOLoggedEvent.Broadcast(PSO);
+		}
+	}
+}
+
 void FPipelineFileCacheManager::CacheGraphicsPSO(uint32 RunTimeHash, FGraphicsPipelineStateInitializer const& Initializer, bool bWasPSOPrecached)
 {
 	if(IsPipelineFileCacheEnabled() && (LogPSOtoFileCache() || ReportNewPSOs()))
@@ -3584,7 +3603,7 @@ void FPipelineFileCacheManager::CacheGraphicsPSO(uint32 RunTimeHash, FGraphicsPi
 							
 						if (ReportNewPSOs() && PSOLoggedEvent.IsBound())
 						{
-							PSOLoggedEvent.Broadcast(NewEntry);
+							NewPSOsToReport.Emplace(MoveTemp(NewEntry));
 						}
 					}
 				}
@@ -3659,7 +3678,7 @@ void FPipelineFileCacheManager::CacheComputePSO(uint32 RunTimeHash, FRHIComputeS
 
 							if (ReportNewPSOs() && PSOLoggedEvent.IsBound())
 							{
-								PSOLoggedEvent.Broadcast(NewEntry);
+								NewPSOsToReport.Emplace(MoveTemp(NewEntry));
 							}
 						}
 					}
@@ -3740,7 +3759,7 @@ void FPipelineFileCacheManager::CacheRayTracingPSO(const FRayTracingPipelineStat
 
 						if (ReportNewPSOs() && PSOLoggedEvent.IsBound())
 						{
-							PSOLoggedEvent.Broadcast(NewEntry);
+							NewPSOsToReport.Emplace(MoveTemp(NewEntry));
 						}
 					}
 
diff --git a/Engine/Source/Runtime/RHI/Private/PipelineStateCache.cpp b/Engine/Source/Runtime/RHI/Private/PipelineStateCache.cpp
--- a/Engine/Source/Runtime/RHI/Private/PipelineStateCache.cpp
+++ b/Engine/Source/Runtime/RHI/Private/PipelineStateCache.cpp
@@ -2481,6 +2481,8 @@ void PipelineStateCache::FlushResources()
 	GPrecacheGraphicsPipelineCache.ProcessDelayedCleanup();
 	GPrecacheComputePipelineCache.ProcessDelayedCleanup();
 
+	FPipelineFileCacheManager::BroadcastNewPSOsDelegate();
+
 	{
 		int32 NumMissesThisFrame = GraphicsPipelineCacheMisses.Load(EMemoryOrder::Relaxed);
 		int32 NumMissesLastFrame = GraphicsPipelineCacheMissesHistory.Num() >= 2 ? GraphicsPipelineCacheMissesHistory[1] : 0;
diff --git a/Engine/Source/Runtime/RHI/Public/PipelineFileCache.h b/Engine/Source/Runtime/RHI/Public/PipelineFileCache.h
--- a/Engine/Source/Runtime/RHI/Public/PipelineFileCache.h
+++ b/Engine/Source/Runtime/RHI/Public/PipelineFileCache.h
@@ -364,6 +364,11 @@ public:
 	 * Gets the event delegate to register for pipeline state logging events.
 	 */
 	RHI_API static FPipelineStateLoggedEvent& OnPipelineStateLogged();
+
+	/*
+	 * If the delegate is set, broadcasts any new PSOs that were encountered since the last time the delegate was broadcast.
+	 */
+	RHI_API static void BroadcastNewPSOsDelegate();
 	
 	RHI_API static void GetOrderedPSOHashes(const FString& PSOCacheKey, TArray<FPipelineCachePSOHeader>& PSOHashes, PSOOrder Order, int64 MinBindCount, TSet<uint32> const& AlreadyCompiledHashes);
 	RHI_API static void FetchPSODescriptors(const FString& PSOCacheKey, TDoubleLinkedList<FPipelineCacheFileFormatPSORead*>& LoadedBatch);
@@ -428,6 +433,7 @@ private:
 	static TMap<uint32, FPSOUsageData> NewPSOUsage;				// For mask or engine updates - Merged + Saved (Our internal PSO hash to latest usage data) - temp working scratch, only holds updates since last "save" so is not the authority on state
 	static TMap<uint32, FPipelineStateStats*> Stats;
 	static TSet<FPipelineCacheFileFormatPSO> NewPSOs;
+	static TArray<FPipelineCacheFileFormatPSO> NewPSOsToReport; // New PSOs that will be broadcast via a delegate (if bound). Cleared when the delegate is broadcast.
  	static TSet<uint32> NewPSOHashes;
 	static uint32 NumNewPSOs;
 	static PSOOrder RequestedOrder;
