Fixed linker errors on some console platform when building modular builds. Fixes involves moving ctor/dtor and friends from header to cpp and dll import/export

diff --git a/Engine/Source/Runtime/MovieScene/Public/Evaluation/MovieScenePropertyTemplate.h b/Engine/Source/Runtime/MovieScene/Public/Evaluation/MovieScenePropertyTemplate.h
--- a/Engine/Source/Runtime/MovieScene/Public/Evaluation/MovieScenePropertyTemplate.h
+++ b/Engine/Source/Runtime/MovieScene/Public/Evaluation/MovieScenePropertyTemplate.h
@@ -231,8 +231,10 @@ struct FMovieScenePropertySectionTemplate : public FMovieSceneEvalTemplate
 {
 	GENERATED_BODY()
 	
-	FMovieScenePropertySectionTemplate(){}
+	MOVIESCENE_API FMovieScenePropertySectionTemplate();
 	MOVIESCENE_API FMovieScenePropertySectionTemplate(FName PropertyName, const FString& InPropertyPath);
+	MOVIESCENE_API FMovieScenePropertySectionTemplate(const FMovieScenePropertySectionTemplate&);
+	MOVIESCENE_API ~FMovieScenePropertySectionTemplate();
 public:
 	//use thse keys for setting and iterating the correct types.
 	MOVIESCENE_API const static FMovieSceneInterrogationKey GetFloatInterrogationKey();
diff --git a/Engine/Source/Runtime/MovieScene/Public/MovieSceneSection.h b/Engine/Source/Runtime/MovieScene/Public/MovieSceneSection.h
--- a/Engine/Source/Runtime/MovieScene/Public/MovieSceneSection.h
+++ b/Engine/Source/Runtime/MovieScene/Public/MovieSceneSection.h
@@ -236,8 +236,6 @@ class UMovieSceneSection
 {
 	GENERATED_UCLASS_BODY()
 
-	~UMovieSceneSection() {};
-
 public:
 
 	UPROPERTY(EditAnywhere, Category="Section", meta=(ShowOnlyInnerProperties))
diff --git a/Engine/Source/Runtime/MovieScene/Public/MovieSceneSpawnRegister.h b/Engine/Source/Runtime/MovieScene/Public/MovieSceneSpawnRegister.h
--- a/Engine/Source/Runtime/MovieScene/Public/MovieSceneSpawnRegister.h
+++ b/Engine/Source/Runtime/MovieScene/Public/MovieSceneSpawnRegister.h
@@ -53,12 +53,13 @@ public:
 
 	static MOVIESCENE_API UE::MovieScene::TPlaybackCapabilityID<FMovieSceneSpawnRegister> ID;
 
+	MOVIESCENE_API FMovieSceneSpawnRegister();
+	MOVIESCENE_API FMovieSceneSpawnRegister(const FMovieSceneSpawnRegister&);
+
 	/**
 	 * Virtual destructor
 	 */
-	virtual ~FMovieSceneSpawnRegister()
-	{
-	}
+	MOVIESCENE_API virtual ~FMovieSceneSpawnRegister();
 
 	/**
 	 * Attempt to find a previously spawned object represented by the specified object and template IDs
diff --git a/Engine/Source/Runtime/NavigationSystem/Private/NavMesh/NavMeshPath.cpp b/Engine/Source/Runtime/NavigationSystem/Private/NavMesh/NavMeshPath.cpp
--- a/Engine/Source/Runtime/NavigationSystem/Private/NavMesh/NavMeshPath.cpp
+++ b/Engine/Source/Runtime/NavigationSystem/Private/NavMesh/NavMeshPath.cpp
@@ -31,6 +31,14 @@ FNavMeshPath::FNavMeshPath()
 	InternalResetNavMeshPath();
 }
 
+PRAGMA_DISABLE_DEPRECATION_WARNINGS
+FNavMeshPath::~FNavMeshPath() = default;
+FNavMeshPath::FNavMeshPath(const FNavMeshPath&) = default;
+FNavMeshPath::FNavMeshPath(FNavMeshPath&& Other) = default;
+FNavMeshPath& FNavMeshPath::operator=(const FNavMeshPath& Other) = default;
+FNavMeshPath& FNavMeshPath::operator=(FNavMeshPath&& Other) = default;
+PRAGMA_ENABLE_DEPRECATION_WARNINGS
+
 void FNavMeshPath::ResetForRepath()
 {
 	Super::ResetForRepath();
diff --git a/Engine/Source/Runtime/NavigationSystem/Private/NavigationPath.cpp b/Engine/Source/Runtime/NavigationSystem/Private/NavigationPath.cpp
--- a/Engine/Source/Runtime/NavigationSystem/Private/NavigationPath.cpp
+++ b/Engine/Source/Runtime/NavigationSystem/Private/NavigationPath.cpp
@@ -67,6 +67,8 @@ FNavigationPath::FNavigationPath(const TArray<FVector>& Points, AActor* InBase)
 	}
 }
 
+FNavigationPath::~FNavigationPath() = default;
+
 void FNavigationPath::InternalResetNavigationPath()
 {
 	ShortcutNodeRefs.Reset();
diff --git a/Engine/Source/Runtime/NavigationSystem/Public/NavMesh/NavMeshPath.h b/Engine/Source/Runtime/NavigationSystem/Public/NavMesh/NavMeshPath.h
--- a/Engine/Source/Runtime/NavigationSystem/Public/NavMesh/NavMeshPath.h
+++ b/Engine/Source/Runtime/NavigationSystem/Public/NavMesh/NavMeshPath.h
@@ -40,13 +40,12 @@ struct FNavMeshPath : public FNavigationPath
 	typedef FNavigationPath Super;
 
 	NAVIGATIONSYSTEM_API FNavMeshPath();
+	NAVIGATIONSYSTEM_API ~FNavMeshPath();
 
-	PRAGMA_DISABLE_DEPRECATION_WARNINGS
-	FNavMeshPath(const FNavMeshPath&) = default;
-	FNavMeshPath(FNavMeshPath&& Other) = default;
-	FNavMeshPath& operator=(const FNavMeshPath& Other) = default;
-	FNavMeshPath& operator=(FNavMeshPath&& Other) = default;
-	PRAGMA_ENABLE_DEPRECATION_WARNINGS
+	NAVIGATIONSYSTEM_API FNavMeshPath(const FNavMeshPath&);
+	NAVIGATIONSYSTEM_API FNavMeshPath(FNavMeshPath&& Other);
+	NAVIGATIONSYSTEM_API FNavMeshPath& operator=(const FNavMeshPath& Other);
+	NAVIGATIONSYSTEM_API FNavMeshPath& operator=(FNavMeshPath&& Other);
 
 	FORCEINLINE void SetWantsStringPulling(const bool bNewWantsStringPulling) { bWantsStringPulling = bNewWantsStringPulling; }
 	FORCEINLINE bool WantsStringPulling() const { return bWantsStringPulling; }
diff --git a/Engine/Source/Runtime/NavigationSystem/Public/NavigationData.h b/Engine/Source/Runtime/NavigationSystem/Public/NavigationData.h
--- a/Engine/Source/Runtime/NavigationSystem/Public/NavigationData.h
+++ b/Engine/Source/Runtime/NavigationSystem/Public/NavigationData.h
@@ -89,8 +89,7 @@ struct FNavigationPath : public TSharedFromThis<FNavigationPath, ESPMode::Thread
 
 	NAVIGATIONSYSTEM_API FNavigationPath();
 	NAVIGATIONSYSTEM_API FNavigationPath(const TArray<FVector>& Points, AActor* Base = NULL);
-	virtual ~FNavigationPath()
-	{ }
+	NAVIGATIONSYSTEM_API virtual ~FNavigationPath();
 
 	PRAGMA_DISABLE_DEPRECATION_WARNINGS
 	FNavigationPath(const FNavigationPath&) = default;
diff --git a/Engine/Source/Runtime/Net/Core/Private/Net/Core/Analytics/NetAnalytics.cpp b/Engine/Source/Runtime/Net/Core/Private/Net/Core/Analytics/NetAnalytics.cpp
--- a/Engine/Source/Runtime/Net/Core/Private/Net/Core/Analytics/NetAnalytics.cpp
+++ b/Engine/Source/Runtime/Net/Core/Private/Net/Core/Analytics/NetAnalytics.cpp
@@ -15,6 +15,13 @@ TAtomic<uint8> GNetAnalyticsCounter(0);
  * FNetAnalytics
  */
 
+FNetAnalyticsData::FNetAnalyticsData()
+	: Aggregator(nullptr)
+{
+}
+
+FNetAnalyticsData::~FNetAnalyticsData() = default;
+
 void FNetAnalyticsData::InternalSendAnalytics()
 {
 	SendAnalytics();
diff --git a/Engine/Source/Runtime/Net/Core/Public/Net/Core/Analytics/NetAnalytics.h b/Engine/Source/Runtime/Net/Core/Public/Net/Core/Analytics/NetAnalytics.h
--- a/Engine/Source/Runtime/Net/Core/Public/Net/Core/Analytics/NetAnalytics.h
+++ b/Engine/Source/Runtime/Net/Core/Public/Net/Core/Analytics/NetAnalytics.h
@@ -120,10 +120,8 @@ public:
 	/**
 	 * Default constructor
 	 */
-	FNetAnalyticsData()
-		: Aggregator(nullptr)
-	{
-	}
+	NETCORE_API FNetAnalyticsData();
+	NETCORE_API virtual ~FNetAnalyticsData();
 
 	/**
 	 * Dispatches the aggregated analytics data - no further data should be added
diff --git a/Engine/Source/Runtime/NetworkReplayStreaming/HttpNetworkReplayStreaming/Public/HttpNetworkReplayStreaming.h b/Engine/Source/Runtime/NetworkReplayStreaming/HttpNetworkReplayStreaming/Public/HttpNetworkReplayStreaming.h
--- a/Engine/Source/Runtime/NetworkReplayStreaming/HttpNetworkReplayStreaming/Public/HttpNetworkReplayStreaming.h
+++ b/Engine/Source/Runtime/NetworkReplayStreaming/HttpNetworkReplayStreaming/Public/HttpNetworkReplayStreaming.h
@@ -218,7 +218,7 @@ HTTPNETWORKREPLAYSTREAMING_API const TCHAR* LexToString(EHttpReplayResult Enum);
 /**
  * Http network replay streaming manager
  */
-class HTTPNETWORKREPLAYSTREAMING_API FHttpNetworkReplayStreamer : public INetworkReplayStreamer
+class FHttpNetworkReplayStreamer : public INetworkReplayStreamer
 {
 	using FHttpReplayResult = UE::Net::TNetResult<EHttpReplayResult>;
 
diff --git a/Engine/Source/Runtime/NetworkReplayStreaming/NetworkReplayStreaming/Private/NetworkReplayStreaming.cpp b/Engine/Source/Runtime/NetworkReplayStreaming/NetworkReplayStreaming/Private/NetworkReplayStreaming.cpp
--- a/Engine/Source/Runtime/NetworkReplayStreaming/NetworkReplayStreaming/Private/NetworkReplayStreaming.cpp
+++ b/Engine/Source/Runtime/NetworkReplayStreaming/NetworkReplayStreaming/Private/NetworkReplayStreaming.cpp
@@ -182,6 +182,9 @@ ENetworkReplayError::Type INetworkReplayStreamer::GetLastError() const
 }
 PRAGMA_ENABLE_DEPRECATION_WARNINGS
 
+INetworkReplayStreamer::INetworkReplayStreamer() = default;
+INetworkReplayStreamer::~INetworkReplayStreamer() = default;
+
 bool INetworkReplayStreamer::HasError() const
 {
 	return ExtendedError.IsValid();
diff --git a/Engine/Source/Runtime/NetworkReplayStreaming/NetworkReplayStreaming/Public/NetworkReplayStreaming.h b/Engine/Source/Runtime/NetworkReplayStreaming/NetworkReplayStreaming/Public/NetworkReplayStreaming.h
--- a/Engine/Source/Runtime/NetworkReplayStreaming/NetworkReplayStreaming/Public/NetworkReplayStreaming.h
+++ b/Engine/Source/Runtime/NetworkReplayStreaming/NetworkReplayStreaming/Public/NetworkReplayStreaming.h
@@ -515,7 +515,8 @@ DECLARE_MULTICAST_DELEGATE_TwoParams(FOnReplayGetAnalyticsAttributes, const clas
 class INetworkReplayStreamer 
 {
 public:
-	virtual ~INetworkReplayStreamer() {}
+	NETWORKREPLAYSTREAMING_API INetworkReplayStreamer();
+	NETWORKREPLAYSTREAMING_API virtual ~INetworkReplayStreamer();
 
 	virtual void StartStreaming(const FStartStreamingParameters& Params, const FStartStreamingCallback& Delegate) = 0;
 
diff --git a/Engine/Source/Runtime/Online/HTTP/Private/HttpRequestAdapter.cpp b/Engine/Source/Runtime/Online/HTTP/Private/HttpRequestAdapter.cpp
--- a/Engine/Source/Runtime/Online/HTTP/Private/HttpRequestAdapter.cpp
+++ b/Engine/Source/Runtime/Online/HTTP/Private/HttpRequestAdapter.cpp
@@ -7,6 +7,8 @@ FHttpRequestAdapterBase::FHttpRequestAdapterBase(const TSharedRef<IHttpRequest>&
 {
 }
 
+FHttpRequestAdapterBase::~FHttpRequestAdapterBase() = default;
+
 FString FHttpRequestAdapterBase::GetURL() const 
 { 
 	return HttpRequest->GetURL(); 
diff --git a/Engine/Source/Runtime/Online/HTTP/Public/HttpRequestAdapter.h b/Engine/Source/Runtime/Online/HTTP/Public/HttpRequestAdapter.h
--- a/Engine/Source/Runtime/Online/HTTP/Public/HttpRequestAdapter.h
+++ b/Engine/Source/Runtime/Online/HTTP/Public/HttpRequestAdapter.h
@@ -14,6 +14,7 @@ class FHttpRequestAdapterBase : public FHttpRequestImpl
 {
 public:
 	HTTP_API FHttpRequestAdapterBase(const TSharedRef<IHttpRequest>& InHttpRequest);
+	HTTP_API ~FHttpRequestAdapterBase();
 
 	// IHttpRequest interface
 	HTTP_API virtual FString GetURL() const override;
diff --git a/Engine/Source/Runtime/PacketHandlers/PacketHandler/Private/PacketHandler.cpp b/Engine/Source/Runtime/PacketHandlers/PacketHandler/Private/PacketHandler.cpp
--- a/Engine/Source/Runtime/PacketHandlers/PacketHandler/Private/PacketHandler.cpp
+++ b/Engine/Source/Runtime/PacketHandlers/PacketHandler/Private/PacketHandler.cpp
@@ -1245,6 +1245,9 @@ HandlerComponent::HandlerComponent(FName InName)
 {
 }
 
+HandlerComponent::~HandlerComponent() = default;
+
+
 bool HandlerComponent::IsActive() const
 {
 	return bActive;
diff --git a/Engine/Source/Runtime/PacketHandlers/PacketHandler/Public/PacketHandler.h b/Engine/Source/Runtime/PacketHandlers/PacketHandler/Public/PacketHandler.h
--- a/Engine/Source/Runtime/PacketHandlers/PacketHandler/Public/PacketHandler.h
+++ b/Engine/Source/Runtime/PacketHandlers/PacketHandler/Public/PacketHandler.h
@@ -740,9 +740,7 @@ public:
 	/**
 	 * Base destructor
 	 */
-	virtual ~HandlerComponent()
-	{
-	}
+	PACKETHANDLER_API virtual ~HandlerComponent();
 
 	/**
 	 * Returns whether this handler is currently active
diff --git a/Engine/Source/Runtime/PacketHandlers/ReliabilityHandlerComponent/Public/ReliabilityHandlerComponent.h b/Engine/Source/Runtime/PacketHandlers/ReliabilityHandlerComponent/Public/ReliabilityHandlerComponent.h
--- a/Engine/Source/Runtime/PacketHandlers/ReliabilityHandlerComponent/Public/ReliabilityHandlerComponent.h
+++ b/Engine/Source/Runtime/PacketHandlers/ReliabilityHandlerComponent/Public/ReliabilityHandlerComponent.h
@@ -17,11 +17,11 @@ struct FBitWriter;
 
 
 class UE_DEPRECATED(5.3, "This handler is not supported. For fully-reliable data, use reliable RPCs or a separate connection with a reliable protocol.")
-RELIABILITYHANDLERCOMPONENT_API ReliabilityHandlerComponent : public HandlerComponent
+ReliabilityHandlerComponent : public HandlerComponent
 {
 public:
 	/* Initializes default data */
-	ReliabilityHandlerComponent();
+	RELIABILITYHANDLERCOMPONENT_API ReliabilityHandlerComponent();
 
 	virtual void CountBytes(FArchive& Ar) const;
 
@@ -35,7 +35,7 @@ public:
 	virtual void Outgoing(FBitWriter& Packet, FOutPacketTraits& Traits) override;
 
 	/* Queues a packet for resending */
-	void QueuePacketForResending(uint8* Packet, int32 CountBits, FOutPacketTraits& Traits);
+	RELIABILITYHANDLERCOMPONENT_API void QueuePacketForResending(uint8* Packet, int32 CountBits, FOutPacketTraits& Traits);
 
 	UE_DEPRECATED(4.21, "Use the PacketTraits version for sending packets with additional flags and options")
 	FORCEINLINE void QueueHandlerPacketForResending(HandlerComponent* InComponent, uint8* Packet, int32 CountBits)
diff --git a/Engine/Source/Runtime/RHI/Private/DynamicRHI.cpp b/Engine/Source/Runtime/RHI/Private/DynamicRHI.cpp
--- a/Engine/Source/Runtime/RHI/Private/DynamicRHI.cpp
+++ b/Engine/Source/Runtime/RHI/Private/DynamicRHI.cpp
@@ -449,6 +449,8 @@ void RHIExit()
 	GRHICommandList.CleanupGraphEvents();
 }
 
+FDynamicRHI::~FDynamicRHI() = default;
+
 void FDynamicRHI::RHIEndFrame_RenderThread(FRHICommandListImmediate& RHICmdList)
 {
 	RHICmdList.ImmediateFlush(EImmediateFlushType::DispatchToRHIThread, ERHISubmitFlags::EndFrame);
diff --git a/Engine/Source/Runtime/RHI/Public/DynamicRHI.h b/Engine/Source/Runtime/RHI/Public/DynamicRHI.h
--- a/Engine/Source/Runtime/RHI/Public/DynamicRHI.h
+++ b/Engine/Source/Runtime/RHI/Public/DynamicRHI.h
@@ -127,7 +127,7 @@ public:
 	using FRHICalcTextureSizeResult = ::FRHICalcTextureSizeResult;
 
 	/** Declare a virtual destructor, so the dynamic RHI can be deleted without knowing its type. */
-	virtual ~FDynamicRHI() {}
+	RHI_API virtual ~FDynamicRHI();
 
 	/** Initializes the RHI; separate from IDynamicRHIModule::CreateRHI so that GDynamicRHI is set when it is called. */
 	virtual void Init() = 0;
diff --git a/Engine/Source/Runtime/RHI/Public/RHITransientResourceAllocator.h b/Engine/Source/Runtime/RHI/Public/RHITransientResourceAllocator.h
--- a/Engine/Source/Runtime/RHI/Public/RHITransientResourceAllocator.h
+++ b/Engine/Source/Runtime/RHI/Public/RHITransientResourceAllocator.h
@@ -92,35 +92,22 @@ public:
 
 	using FResourceTask = UE::Tasks::TTask<FResourceTaskResult>;
 
-	FRHITransientResource(
+	RHI_API FRHITransientResource(
 		FRHIResource* InResource,
 		uint64 InGpuVirtualAddress,
 		uint64 InHash,
 		uint64 InSize,
 		ERHITransientAllocationType InAllocationType,
-		ERHITransientResourceType InResourceType)
-		: Resource(InResource)
-		, GpuVirtualAddress(InGpuVirtualAddress)
-		, Hash(InHash)
-		, Size(InSize)
-		, AllocationType(InAllocationType)
-		, ResourceType(InResourceType)
-	{}
+		ERHITransientResourceType InResourceType);
 
-	FRHITransientResource(
+	RHI_API FRHITransientResource(
 		const FResourceTask& InResourceTask,
 		uint64 InHash,
 		uint64 InSize,
 		ERHITransientAllocationType InAllocationType,
-		ERHITransientResourceType InResourceType)
-		: ResourceTask(InResourceTask)
-		, Hash(InHash)
-		, Size(InSize)
-		, AllocationType(InAllocationType)
-		, ResourceType(InResourceType)
-	{}
+		ERHITransientResourceType InResourceType);
 
-	virtual ~FRHITransientResource() = default;
+	RHI_API virtual ~FRHITransientResource();
 
 	//////////////////////////////////////////////////////////////////////////////////////////////////
 	//! Internal Allocator API
@@ -288,26 +275,22 @@ private:
 class FRHITransientTexture final : public FRHITransientResource
 {
 public:
-	FRHITransientTexture(
+	RHI_API FRHITransientTexture(
 		const FResourceTask& InResourceTask,
 		uint64 InHash,
 		uint64 InSize,
 		ERHITransientAllocationType InAllocationType,
-		const FRHITextureCreateInfo& InCreateInfo)
-		: FRHITransientResource(InResourceTask, InHash, InSize, InAllocationType, ERHITransientResourceType::Texture)
-		, CreateInfo(InCreateInfo)
-	{}
+		const FRHITextureCreateInfo& InCreateInfo);
 
-	FRHITransientTexture(
+	RHI_API FRHITransientTexture(
 		FRHIResource* InTexture,
 		uint64 InGpuVirtualAddress,
 		uint64 InHash,
 		uint64 InSize,
 		ERHITransientAllocationType InAllocationType,
-		const FRHITextureCreateInfo& InCreateInfo)
-		: FRHITransientResource(InTexture, InGpuVirtualAddress, InHash, InSize, InAllocationType, ERHITransientResourceType::Texture)
-		, CreateInfo(InCreateInfo)
-	{}
+		const FRHITextureCreateInfo& InCreateInfo);
+
+	RHI_API virtual ~FRHITransientTexture();
 
 	// Returns the underlying RHI texture.
 	FRHITexture* GetRHI() const { return static_cast<FRHITexture*>(FRHITransientResource::GetRHI()); }
@@ -355,6 +338,8 @@ public:
 		, CreateInfo(InCreateInfo)
 	{}
 
+	RHI_API virtual ~FRHITransientBuffer();
+
 	// Returns the underlying RHI buffer.
 	FRHIBuffer* GetRHI() const { return static_cast<FRHIBuffer*>(FRHITransientResource::GetRHI()); }
 
diff --git a/Engine/Source/Runtime/RHICore/Private/RHICoreTransientResourceAllocator.cpp b/Engine/Source/Runtime/RHICore/Private/RHICoreTransientResourceAllocator.cpp
--- a/Engine/Source/Runtime/RHICore/Private/RHICoreTransientResourceAllocator.cpp
+++ b/Engine/Source/Runtime/RHICore/Private/RHICoreTransientResourceAllocator.cpp
@@ -629,6 +629,12 @@ void FRHITransientHeapCache::GarbageCollect()
 
 //////////////////////////////////////////////////////////////////////////
 
+FRHITransientResourceHeapAllocator::FRHITransientResourceHeapAllocator(FRHITransientHeapCache& InHeapCache)
+	: HeapCache(InHeapCache)
+{}
+
+FRHITransientResourceHeapAllocator::~FRHITransientResourceHeapAllocator() = default;
+
 FRHITransientTexture* FRHITransientResourceHeapAllocator::CreateTextureInternal(
 	const FRHITextureCreateInfo& CreateInfo,
 	const TCHAR* DebugName,
diff --git a/Engine/Source/Runtime/RHICore/Public/RHICoreTransientResourceAllocator.h b/Engine/Source/Runtime/RHICore/Public/RHICoreTransientResourceAllocator.h
--- a/Engine/Source/Runtime/RHICore/Public/RHICoreTransientResourceAllocator.h
+++ b/Engine/Source/Runtime/RHICore/Public/RHICoreTransientResourceAllocator.h
@@ -671,9 +671,8 @@ private:
 class FRHITransientResourceHeapAllocator : public IRHITransientResourceAllocator
 {
 public:
-	FRHITransientResourceHeapAllocator(FRHITransientHeapCache& InHeapCache)
-		: HeapCache(InHeapCache)
-	{}
+	RHICORE_API FRHITransientResourceHeapAllocator(FRHITransientHeapCache& InHeapCache);
+	RHICORE_API ~FRHITransientResourceHeapAllocator();
 
 	// Sets the create mode for allocations.
 	RHICORE_API void SetCreateMode(ERHITransientResourceCreateMode InCreateMode) override;
diff --git a/Engine/Source/Runtime/RenderCore/Private/RenderGraphResources.cpp b/Engine/Source/Runtime/RenderCore/Private/RenderGraphResources.cpp
--- a/Engine/Source/Runtime/RenderCore/Private/RenderGraphResources.cpp
+++ b/Engine/Source/Runtime/RenderCore/Private/RenderGraphResources.cpp
@@ -175,6 +175,8 @@ void FRDGPooledBuffer::SetDebugLabelName(FRHICommandListBase& RHICmdList, const
 #endif
 }
 
+FRDGUniformBuffer::~FRDGUniformBuffer() = default;
+
 void FRDGUniformBuffer::InitRHI()
 {
 	check(!HasRHI());
diff --git a/Engine/Source/Runtime/RenderCore/Private/RenderResource.cpp b/Engine/Source/Runtime/RenderCore/Private/RenderResource.cpp
--- a/Engine/Source/Runtime/RenderCore/Private/RenderResource.cpp
+++ b/Engine/Source/Runtime/RenderCore/Private/RenderResource.cpp
@@ -237,6 +237,11 @@ FRenderResource::FRenderResource(ERHIFeatureLevel::Type InFeatureLevel)
 {
 }
 
+FRenderResource::FRenderResource(const FRenderResource&) = default;
+FRenderResource::FRenderResource(FRenderResource&&) = default;
+FRenderResource& FRenderResource::operator=(const FRenderResource& Other) = default;
+FRenderResource& FRenderResource::operator=(FRenderResource&& Other) = default;
+
 FRenderResource::~FRenderResource()
 {
 	checkf(ResourceState == ERenderResourceState::Default, TEXT(" Invalid Resource State: %s"), ResourceState == ERenderResourceState::BatchReleased ? TEXT("BatchReleased") : TEXT("Deleted"));
@@ -417,6 +422,10 @@ TGlobalResource<FTextureSamplerStateCache> GTextureSamplerStateCache;
 
 FTexture::FTexture() = default;
 FTexture::~FTexture() = default;
+FTexture::FTexture(const FTexture&) = default;
+FTexture::FTexture(FTexture&&) = default;
+FTexture& FTexture::operator=(const FTexture& Other) = default;
+FTexture& FTexture::operator=(FTexture&& Other) = default;
 
 uint32 FTexture::GetSizeX() const
 {
@@ -545,6 +554,8 @@ FString FTextureReference::GetFriendlyName() const
 // FVertexBuffer
 
 FVertexBuffer::FVertexBuffer() = default;
+FVertexBuffer::FVertexBuffer(const FVertexBuffer&) = default;
+FVertexBuffer& FVertexBuffer::operator=(const FVertexBuffer& Other) = default;
 FVertexBuffer::~FVertexBuffer() = default;
 
 void FVertexBuffer::ReleaseRHI()
@@ -579,6 +590,8 @@ void FVertexBufferWithSRV::ReleaseRHI()
 // FIndexBuffer
 
 FIndexBuffer::FIndexBuffer() = default;
+FIndexBuffer::FIndexBuffer(const FIndexBuffer&) = default;
+FIndexBuffer& FIndexBuffer::operator=(const FIndexBuffer& Other) = default;
 FIndexBuffer::~FIndexBuffer() = default;
 
 void FIndexBuffer::ReleaseRHI()
diff --git a/Engine/Source/Runtime/RenderCore/Public/RenderGraphResources.h b/Engine/Source/Runtime/RenderCore/Public/RenderGraphResources.h
--- a/Engine/Source/Runtime/RenderCore/Public/RenderGraphResources.h
+++ b/Engine/Source/Runtime/RenderCore/Public/RenderGraphResources.h
@@ -178,7 +178,7 @@ class FRDGUniformBuffer
 {
 public:
 
-	virtual ~FRDGUniformBuffer() {};
+	RENDERCORE_API virtual ~FRDGUniformBuffer();
 
 	FORCEINLINE const FRDGParameterStruct& GetParameters() const
 	{
diff --git a/Engine/Source/Runtime/RenderCore/Public/RenderResource.h b/Engine/Source/Runtime/RenderCore/Public/RenderResource.h
--- a/Engine/Source/Runtime/RenderCore/Public/RenderResource.h
+++ b/Engine/Source/Runtime/RenderCore/Public/RenderResource.h
@@ -87,6 +87,12 @@ public:
 	/** Constructor when we know what feature level this resource should support */
 	RENDERCORE_API FRenderResource(ERHIFeatureLevel::Type InFeatureLevel);
 
+	/** Misc copy/assignment */
+	RENDERCORE_API FRenderResource(const FRenderResource&);
+	RENDERCORE_API FRenderResource(FRenderResource&&);
+	RENDERCORE_API FRenderResource& operator=(const FRenderResource& Other);
+	RENDERCORE_API FRenderResource& operator=(FRenderResource&& Other);
+
 	/** Destructor used to catch unreleased resources. */
 	RENDERCORE_API virtual ~FRenderResource();
 
@@ -395,6 +401,10 @@ public:
 
 	RENDERCORE_API FTexture();
 	RENDERCORE_API virtual ~FTexture();
+	RENDERCORE_API FTexture(const FTexture&);
+	RENDERCORE_API FTexture(FTexture&&);
+	RENDERCORE_API FTexture& operator=(const FTexture& Other);
+	RENDERCORE_API FTexture& operator=(FTexture&& Other);
 
 	const FTextureRHIRef& GetTextureRHI() { return TextureRHI; }
 
@@ -472,6 +482,8 @@ class FVertexBuffer : public FRenderResource
 {
 public:
 	RENDERCORE_API FVertexBuffer();
+	RENDERCORE_API FVertexBuffer(const FVertexBuffer&);
+	RENDERCORE_API FVertexBuffer& operator=(const FVertexBuffer& Other);
 	RENDERCORE_API virtual ~FVertexBuffer();
 
 	// FRenderResource interface.
@@ -505,6 +517,8 @@ class FIndexBuffer : public FRenderResource
 {
 public:
 	RENDERCORE_API FIndexBuffer();
+	RENDERCORE_API FIndexBuffer(const FIndexBuffer&);
+	RENDERCORE_API FIndexBuffer& operator=(const FIndexBuffer& Other);
 	RENDERCORE_API virtual ~FIndexBuffer();
 
 	// FRenderResource interface.
diff --git a/Engine/Source/Runtime/Renderer/Public/VolumeRendering.h b/Engine/Source/Runtime/Renderer/Public/VolumeRendering.h
--- a/Engine/Source/Runtime/Renderer/Public/VolumeRendering.h
+++ b/Engine/Source/Runtime/Renderer/Public/VolumeRendering.h
@@ -94,10 +94,10 @@ private:
 extern ENGINE_API void RasterizeToVolumeTexture(FRHICommandList& RHICmdList, FVolumeBounds VolumeBounds);
 
 /** Vertex buffer used for rendering into a volume texture. */
-class ENGINE_API FVolumeRasterizeVertexBuffer : public FVertexBuffer
+class FVolumeRasterizeVertexBuffer : public FVertexBuffer
 {
 public:
-	virtual void InitRHI(FRHICommandListBase& RHICmdList) override;
+	ENGINE_API virtual void InitRHI(FRHICommandListBase& RHICmdList) override;
 };
 
 extern ENGINE_API TGlobalResource<FVolumeRasterizeVertexBuffer> GVolumeRasterizeVertexBuffer;
\ No newline at end of file
diff --git a/Engine/Source/Runtime/SignalProcessing/Private/AudioChannelFormatConverter.cpp b/Engine/Source/Runtime/SignalProcessing/Private/AudioChannelFormatConverter.cpp
--- a/Engine/Source/Runtime/SignalProcessing/Private/AudioChannelFormatConverter.cpp
+++ b/Engine/Source/Runtime/SignalProcessing/Private/AudioChannelFormatConverter.cpp
@@ -9,6 +9,8 @@
 
 namespace Audio
 {
+	FBaseChannelFormatConverter::~FBaseChannelFormatConverter() = default;
+
 	const IChannelFormatConverter::FInputFormat& FBaseChannelFormatConverter::GetInputFormat() const
 	{
 		return InputFormat;
diff --git a/Engine/Source/Runtime/SignalProcessing/Private/Delay.cpp b/Engine/Source/Runtime/SignalProcessing/Private/Delay.cpp
--- a/Engine/Source/Runtime/SignalProcessing/Private/Delay.cpp
+++ b/Engine/Source/Runtime/SignalProcessing/Private/Delay.cpp
@@ -27,6 +27,8 @@ namespace Audio
 		Reset();
 	}
 
+	FDelay::~FDelay() = default;
+
 	// update metadata, call Reset()
 	void FDelay::Init(const float InSampleRate, const float InBufferLengthSec)
 	{
diff --git a/Engine/Source/Runtime/SignalProcessing/Private/Filter.cpp b/Engine/Source/Runtime/SignalProcessing/Private/Filter.cpp
--- a/Engine/Source/Runtime/SignalProcessing/Private/Filter.cpp
+++ b/Engine/Source/Runtime/SignalProcessing/Private/Filter.cpp
@@ -482,9 +482,9 @@ namespace Audio
 	{
 	}
 
-	IFilter::~IFilter()
-	{
-	}
+	IFilter::IFilter(const IFilter&) = default;
+
+	IFilter::~IFilter() = default;
 
 	void IFilter::Init(const float InSampleRate, const int32 InNumChannels, const int32 InVoiceId, FModulationMatrix* InModMatrix)
 	{
diff --git a/Engine/Source/Runtime/SignalProcessing/Private/LFO.cpp b/Engine/Source/Runtime/SignalProcessing/Private/LFO.cpp
--- a/Engine/Source/Runtime/SignalProcessing/Private/LFO.cpp
+++ b/Engine/Source/Runtime/SignalProcessing/Private/LFO.cpp
@@ -21,6 +21,8 @@ namespace Audio
 	{
 	}
 
+	FLFO::~FLFO() = default;
+
 	void FLFO::Init(const float InSampleRate, const int32 InVoiceId, FModulationMatrix* InMatrix, const int32 ModMatrixStage)
 	{
 		IOscBase::Init(InSampleRate, InVoiceId, InMatrix, ModMatrixStage);
diff --git a/Engine/Source/Runtime/SignalProcessing/Public/DSP/AudioChannelFormatConverter.h b/Engine/Source/Runtime/SignalProcessing/Public/DSP/AudioChannelFormatConverter.h
--- a/Engine/Source/Runtime/SignalProcessing/Public/DSP/AudioChannelFormatConverter.h
+++ b/Engine/Source/Runtime/SignalProcessing/Public/DSP/AudioChannelFormatConverter.h
@@ -71,7 +71,7 @@ namespace Audio
 				float Gain = 0.f;
 			};
 
-			virtual ~FBaseChannelFormatConverter() = default;
+			SIGNALPROCESSING_API virtual ~FBaseChannelFormatConverter();
 
 			/** Return the input format handled by this converter. */
 			SIGNALPROCESSING_API const FInputFormat& GetInputFormat() const override;
diff --git a/Engine/Source/Runtime/SignalProcessing/Public/DSP/Delay.h b/Engine/Source/Runtime/SignalProcessing/Public/DSP/Delay.h
--- a/Engine/Source/Runtime/SignalProcessing/Public/DSP/Delay.h
+++ b/Engine/Source/Runtime/SignalProcessing/Public/DSP/Delay.h
@@ -16,7 +16,7 @@ namespace Audio
 		SIGNALPROCESSING_API FDelay();
 
 		// Virtual Destructor
-		SIGNALPROCESSING_API virtual ~FDelay() = default;
+		SIGNALPROCESSING_API virtual ~FDelay();
 
 		// Initialization of the delay with given sample rate and max buffer size in samples.
 		// (calls Reset())
diff --git a/Engine/Source/Runtime/SignalProcessing/Public/DSP/Filter.h b/Engine/Source/Runtime/SignalProcessing/Public/DSP/Filter.h
--- a/Engine/Source/Runtime/SignalProcessing/Public/DSP/Filter.h
+++ b/Engine/Source/Runtime/SignalProcessing/Public/DSP/Filter.h
@@ -130,6 +130,7 @@ namespace Audio
 	{
 	public:
 		SIGNALPROCESSING_API IFilter();
+		SIGNALPROCESSING_API IFilter(const IFilter&);
 		SIGNALPROCESSING_API virtual ~IFilter();
 
 		// Initialize the filter
diff --git a/Engine/Source/Runtime/SignalProcessing/Public/DSP/LFO.h b/Engine/Source/Runtime/SignalProcessing/Public/DSP/LFO.h
--- a/Engine/Source/Runtime/SignalProcessing/Public/DSP/LFO.h
+++ b/Engine/Source/Runtime/SignalProcessing/Public/DSP/LFO.h
@@ -45,7 +45,7 @@ namespace Audio
 	{
 	public:
 		SIGNALPROCESSING_API FLFO();
-		virtual ~FLFO() = default;
+		SIGNALPROCESSING_API virtual ~FLFO();
 
 		//~ Begin FOscBase
 		SIGNALPROCESSING_API virtual void Init(const float InSampleRate, const int32 InVoiceId = 0, FModulationMatrix* InMatrix = nullptr, const int32 ModMatrixStage = 0) override;
diff --git a/Engine/Source/Runtime/Slate/Private/Widgets/Layout/SBackgroundBlur.cpp b/Engine/Source/Runtime/Slate/Private/Widgets/Layout/SBackgroundBlur.cpp
--- a/Engine/Source/Runtime/Slate/Private/Widgets/Layout/SBackgroundBlur.cpp
+++ b/Engine/Source/Runtime/Slate/Private/Widgets/Layout/SBackgroundBlur.cpp
@@ -31,6 +31,7 @@ SBackgroundBlur::SBackgroundBlur()
 {
 }
 
+SBackgroundBlur::~SBackgroundBlur() = default;
 
 void SBackgroundBlur::Construct(const FArguments& InArgs)
 {
diff --git a/Engine/Source/Runtime/Slate/Public/Widgets/Layout/SBackgroundBlur.h b/Engine/Source/Runtime/Slate/Public/Widgets/Layout/SBackgroundBlur.h
--- a/Engine/Source/Runtime/Slate/Public/Widgets/Layout/SBackgroundBlur.h
+++ b/Engine/Source/Runtime/Slate/Public/Widgets/Layout/SBackgroundBlur.h
@@ -47,6 +47,7 @@ public:
 
 public:
 	SLATE_API SBackgroundBlur();
+	SLATE_API virtual ~SBackgroundBlur();
 	SLATE_API void Construct(const FArguments& InArgs);
 	SLATE_API virtual int32 OnPaint(const FPaintArgs& Args, const FGeometry& AllottedGeometry, const FSlateRect& MyCullingRect, FSlateWindowElementList& OutDrawElements, int32 LayerId, const FWidgetStyle& InWidgetStyle, bool bParentEnabled) const override;
 
diff --git a/Engine/Source/Runtime/SlateCore/Private/Rendering/ShaderResourceManager.cpp b/Engine/Source/Runtime/SlateCore/Private/Rendering/ShaderResourceManager.cpp
--- a/Engine/Source/Runtime/SlateCore/Private/Rendering/ShaderResourceManager.cpp
+++ b/Engine/Source/Runtime/SlateCore/Private/Rendering/ShaderResourceManager.cpp
@@ -23,6 +23,12 @@ FSlateResourceHandle::FSlateResourceHandle(const TSharedPtr<FSlateSharedHandleDa
 {
 }
 
+FSlateShaderResourceManager::FSlateShaderResourceManager() = default;
+
+FSlateShaderResourceManager::~FSlateShaderResourceManager()
+{
+	ClearTextureMap();
+}
 
 FSlateResourceHandle FSlateShaderResourceManager::GetResourceHandle(const FSlateBrush& Brush, FVector2f LocalSize, float DrawScale)
 {
diff --git a/Engine/Source/Runtime/SlateCore/Private/Rendering/SlateDrawBuffer.cpp b/Engine/Source/Runtime/SlateCore/Private/Rendering/SlateDrawBuffer.cpp
--- a/Engine/Source/Runtime/SlateCore/Private/Rendering/SlateDrawBuffer.cpp
+++ b/Engine/Source/Runtime/SlateCore/Private/Rendering/SlateDrawBuffer.cpp
@@ -16,6 +16,8 @@ FSlateDrawBuffer::FSlateDrawBuffer()
 	, ResourceVersion(0)
 { }
 
+FSlateDrawBuffer::~FSlateDrawBuffer() = default;
+
 FSlateWindowElementList& FSlateDrawBuffer::AddWindowElementList(TSharedRef<SWindow> ForWindow)
 {
 	ensureMsgf(IsLocked(), TEXT("The SlateDrawBuffer should be lock before modifying it."));
diff --git a/Engine/Source/Runtime/SlateCore/Private/Styling/SlateStyleSet.cpp b/Engine/Source/Runtime/SlateCore/Private/Styling/SlateStyleSet.cpp
--- a/Engine/Source/Runtime/SlateCore/Private/Styling/SlateStyleSet.cpp
+++ b/Engine/Source/Runtime/SlateCore/Private/Styling/SlateStyleSet.cpp
@@ -23,6 +23,9 @@ FSlateStyleSet::FSlateStyleSet(const FName& InStyleSetName)
 	Set(TEXT("Default"), GetDefaultBrush());
 }
 
+FSlateStyleSet::FSlateStyleSet(FSlateStyleSet&&) = default;
+FSlateStyleSet::FSlateStyleSet(const FSlateStyleSet&) = default;
+
 FSlateStyleSet::~FSlateStyleSet()
 {
 	// Delete all allocated brush resources.
diff --git a/Engine/Source/Runtime/SlateCore/Private/Styling/SlateTypes.cpp b/Engine/Source/Runtime/SlateCore/Private/Styling/SlateTypes.cpp
--- a/Engine/Source/Runtime/SlateCore/Private/Styling/SlateTypes.cpp
+++ b/Engine/Source/Runtime/SlateCore/Private/Styling/SlateTypes.cpp
@@ -471,6 +471,7 @@ FSliderStyle::FSliderStyle()
 {
 }
 
+FSliderStyle::FSliderStyle(const FSliderStyle&) = default;
 FSliderStyle::~FSliderStyle() = default;
 
 void FSliderStyle::GetResources( TArray< const FSlateBrush* >& OutBrushes ) const
@@ -547,6 +548,7 @@ FSpinBoxStyle::FSpinBoxStyle()
 {
 }
 
+FSpinBoxStyle::FSpinBoxStyle(const FSpinBoxStyle&) = default;
 FSpinBoxStyle::~FSpinBoxStyle() = default;
 
 void FSpinBoxStyle::GetResources( TArray< const FSlateBrush* >& OutBrushes ) const
diff --git a/Engine/Source/Runtime/SlateCore/Private/Textures/TextureAtlas.cpp b/Engine/Source/Runtime/SlateCore/Private/Textures/TextureAtlas.cpp
--- a/Engine/Source/Runtime/SlateCore/Private/Textures/TextureAtlas.cpp
+++ b/Engine/Source/Runtime/SlateCore/Private/Textures/TextureAtlas.cpp
@@ -50,6 +50,23 @@ uint32 GetSlateFontAtlasContentBytesPerPixel(ESlateFontAtlasContentType InConten
 /* FSlateTextureAtlas helper class
  *****************************************************************************/
 
+FSlateTextureAtlas::FSlateTextureAtlas( uint32 InWidth, uint32 InHeight, uint32 InBytesPerPixel, ESlateTextureAtlasPaddingStyle InPaddingStyle, bool bInUpdatesAfterInitialization )
+	: AtlasData()
+	, AtlasUsedSlots(NULL)
+	, AtlasEmptySlotsMap()
+	, AtlasWidth( InWidth )
+	, AtlasHeight( InHeight )
+	, BytesPerPixel( InBytesPerPixel )
+	, PaddingStyle( InPaddingStyle )
+	, bNeedsUpdate( false )
+	, bUpdatesAfterInitialization(bInUpdatesAfterInitialization)
+#if !(UE_BUILD_SHIPPING || UE_BUILD_TEST)
+	, AtlasOwnerThread( ESlateTextureAtlasThreadId::Unknown )
+#endif
+{
+	InitAtlasData();
+}
+
 FSlateTextureAtlas::~FSlateTextureAtlas()
 {
 	EmptyAtlasData();
diff --git a/Engine/Source/Runtime/SlateCore/Public/Rendering/ShaderResourceManager.h b/Engine/Source/Runtime/SlateCore/Public/Rendering/ShaderResourceManager.h
--- a/Engine/Source/Runtime/SlateCore/Public/Rendering/ShaderResourceManager.h
+++ b/Engine/Source/Runtime/SlateCore/Public/Rendering/ShaderResourceManager.h
@@ -43,12 +43,8 @@ struct FCompareFNewTextureInfoByTextureSize
 class FSlateShaderResourceManager
 {
 public:
-	FSlateShaderResourceManager() {};
-	virtual ~FSlateShaderResourceManager()
-	{
-		ClearTextureMap();
-	}
-
+	SLATECORE_API FSlateShaderResourceManager();
+	SLATECORE_API virtual ~FSlateShaderResourceManager();
 
 	/** 
 	 * Returns a texture associated with the passed in name.  Should return nullptr if not found 
diff --git a/Engine/Source/Runtime/SlateCore/Public/Rendering/SlateDrawBuffer.h b/Engine/Source/Runtime/SlateCore/Public/Rendering/SlateDrawBuffer.h
--- a/Engine/Source/Runtime/SlateCore/Public/Rendering/SlateDrawBuffer.h
+++ b/Engine/Source/Runtime/SlateCore/Public/Rendering/SlateDrawBuffer.h
@@ -20,6 +20,7 @@ class FSlateDrawBuffer : public FGCObject
 public:
 	/** Default constructor. */
 	SLATECORE_API FSlateDrawBuffer();
+	SLATECORE_API ~FSlateDrawBuffer();
 
 	/** Removes all data from the buffer. */
 	SLATECORE_API void ClearBuffer();
diff --git a/Engine/Source/Runtime/SlateCore/Public/Styling/SlateStyle.h b/Engine/Source/Runtime/SlateCore/Public/Styling/SlateStyle.h
--- a/Engine/Source/Runtime/SlateCore/Public/Styling/SlateStyle.h
+++ b/Engine/Source/Runtime/SlateCore/Public/Styling/SlateStyle.h
@@ -32,6 +32,8 @@ public:
 	 * @param InStyleSetName The name used to identity this style set
 	 */
 	SLATECORE_API FSlateStyleSet(const FName& InStyleSetName);
+	SLATECORE_API FSlateStyleSet(FSlateStyleSet&&);
+	SLATECORE_API FSlateStyleSet(const FSlateStyleSet&);
 
 	/** Destructor. */
 	SLATECORE_API virtual ~FSlateStyleSet();
diff --git a/Engine/Source/Runtime/SlateCore/Public/Styling/SlateTypes.h b/Engine/Source/Runtime/SlateCore/Public/Styling/SlateTypes.h
--- a/Engine/Source/Runtime/SlateCore/Public/Styling/SlateTypes.h
+++ b/Engine/Source/Runtime/SlateCore/Public/Styling/SlateTypes.h
@@ -1319,7 +1319,7 @@ struct FSliderStyle : public FSlateWidgetStyle
 	GENERATED_USTRUCT_BODY()
 
 	SLATECORE_API FSliderStyle();
-
+	SLATECORE_API FSliderStyle(const FSliderStyle&);
 	SLATECORE_API virtual ~FSliderStyle();
 
 	SLATECORE_API virtual void GetResources( TArray< const FSlateBrush* >& OutBrushes ) const override;
@@ -1468,7 +1468,7 @@ struct FSpinBoxStyle : public FSlateWidgetStyle
 	GENERATED_USTRUCT_BODY()
 
 	SLATECORE_API FSpinBoxStyle();
-
+	SLATECORE_API FSpinBoxStyle(const FSpinBoxStyle&);
 	SLATECORE_API virtual ~FSpinBoxStyle();
 
 	SLATECORE_API virtual void GetResources( TArray< const FSlateBrush* >& OutBrushes ) const override;
diff --git a/Engine/Source/Runtime/SlateCore/Public/Textures/TextureAtlas.h b/Engine/Source/Runtime/SlateCore/Public/Textures/TextureAtlas.h
--- a/Engine/Source/Runtime/SlateCore/Public/Textures/TextureAtlas.h
+++ b/Engine/Source/Runtime/SlateCore/Public/Textures/TextureAtlas.h
@@ -98,22 +98,7 @@ struct FAtlasedTextureSlot : public TIntrusiveLinkedList<FAtlasedTextureSlot>
 class FSlateTextureAtlas
 {
 public:
-	FSlateTextureAtlas( uint32 InWidth, uint32 InHeight, uint32 InBytesPerPixel, ESlateTextureAtlasPaddingStyle InPaddingStyle, bool bInUpdatesAfterInitialization )
-		: AtlasData()
-		, AtlasUsedSlots(NULL)
-		, AtlasEmptySlotsMap()
-		, AtlasWidth( InWidth )
-		, AtlasHeight( InHeight )
-		, BytesPerPixel( InBytesPerPixel )
-		, PaddingStyle( InPaddingStyle )
-		, bNeedsUpdate( false )
-		, bUpdatesAfterInitialization(bInUpdatesAfterInitialization)
-#if !(UE_BUILD_SHIPPING || UE_BUILD_TEST)
-		, AtlasOwnerThread( ESlateTextureAtlasThreadId::Unknown )
-#endif
-	{
-		InitAtlasData();
-	}
+	SLATECORE_API FSlateTextureAtlas( uint32 InWidth, uint32 InHeight, uint32 InBytesPerPixel, ESlateTextureAtlasPaddingStyle InPaddingStyle, bool bInUpdatesAfterInitialization );
 
 	SLATECORE_API virtual ~FSlateTextureAtlas();
 
