Change GPU Profiler event storage
 - Implemented UE::RHI::GPUProfiler::FEventStream which holds events in linked chunks, to minimize the allocator overhead and increase cache locality when working with GPU profiler events.
 - Modified existing event sinks and D3D12 RHI to use the new FEventStream type.

Added some constructors to various FEvent types.

diff --git a/Engine/Source/Runtime/D3D12RHI/Private/D3D12CommandContext.cpp b/Engine/Source/Runtime/D3D12RHI/Private/D3D12CommandContext.cpp
--- a/Engine/Source/Runtime/D3D12RHI/Private/D3D12CommandContext.cpp
+++ b/Engine/Source/Runtime/D3D12RHI/Private/D3D12CommandContext.cpp
@@ -202,9 +202,7 @@ void FD3D12ContextCommon::BindDiagnosticBuffer(FD3D12RootSignature const* RootSi
 		{
 			FlushProfilerStats();
 
-			auto& Event = GetCommandList().EmplaceEvent<UE::RHI::GPUProfiler::FEvent::FBeginBreadcrumb>();
-			Event.Breadcrumb = Breadcrumb;
-
+			auto& Event = GetCommandList().EmplaceProfilerEvent<UE::RHI::GPUProfiler::FEvent::FBeginBreadcrumb>(Breadcrumb);
 			FD3D12QueryLocation TimestampQuery = AllocateQuery(ED3D12QueryType::ProfilerTimestampTOP, &Event.GPUTimestampTOP);
 			EndQuery(TimestampQuery);
 		}
@@ -226,9 +224,7 @@ void FD3D12ContextCommon::BindDiagnosticBuffer(FD3D12RootSignature const* RootSi
 		{
 			FlushProfilerStats();
 
-			auto& Event = GetCommandList().EmplaceEvent<UE::RHI::GPUProfiler::FEvent::FEndBreadcrumb>();
-			Event.Breadcrumb = Breadcrumb;
-
+			auto& Event = GetCommandList().EmplaceProfilerEvent<UE::RHI::GPUProfiler::FEvent::FEndBreadcrumb>(Breadcrumb);
 			FD3D12QueryLocation TimestampQuery = AllocateQuery(ED3D12QueryType::ProfilerTimestampBOP, &Event.GPUTimestampBOP);
 			EndQuery(TimestampQuery);
 		}
diff --git a/Engine/Source/Runtime/D3D12RHI/Private/D3D12CommandContext.h b/Engine/Source/Runtime/D3D12RHI/Private/D3D12CommandContext.h
--- a/Engine/Source/Runtime/D3D12RHI/Private/D3D12CommandContext.h
+++ b/Engine/Source/Runtime/D3D12RHI/Private/D3D12CommandContext.h
@@ -533,7 +533,7 @@ public:
 		// Flush accumulated draw stats
 		if (StatEvent)
 		{
-			GetCommandList().EmplaceEvent<UE::RHI::GPUProfiler::FEvent::FStats>() = StatEvent;
+			GetCommandList().EmplaceProfilerEvent<UE::RHI::GPUProfiler::FEvent::FStats>() = StatEvent;
 			StatEvent = {};
 		}
 	}
diff --git a/Engine/Source/Runtime/D3D12RHI/Private/D3D12CommandList.cpp b/Engine/Source/Runtime/D3D12RHI/Private/D3D12CommandList.cpp
--- a/Engine/Source/Runtime/D3D12RHI/Private/D3D12CommandList.cpp
+++ b/Engine/Source/Runtime/D3D12RHI/Private/D3D12CommandList.cpp
@@ -287,7 +287,9 @@ void FD3D12CommandList::Reset(FD3D12CommandAllocator* NewCommandAllocator, FD3D1
 	}
 	D3DX12Residency::Open(ResidencySet);
 
-	State = FState(NewCommandAllocator, TimestampAllocator, PipelineStatsAllocator);
+	(&State)->~FState();
+	new (&State) FState(NewCommandAllocator, TimestampAllocator, PipelineStatsAllocator);
+
 	BeginLocalQueries();
 }
 
@@ -315,47 +317,48 @@ void FD3D12CommandList::Close()
 
 void FD3D12CommandList::BeginLocalQueries()
 {
-	if (!State.bLocalQueriesBegun)
+#if DO_CHECK
+	check(!State.bLocalQueriesBegun);
+	State.bLocalQueriesBegun = true;
+#endif
+
+	if (State.BeginTimestamp)
 	{
-		if (State.BeginTimestamp)
-		{
 #if RHI_NEW_GPU_PROFILER
-			auto& Event = EmplaceEvent<UE::RHI::GPUProfiler::FEvent::FBeginWork>();
-			State.BeginTimestamp.Target = &Event.GPUTimestampTOP;
+		// CPUTimestamp is filled in at submission time in FlushProfilerEvents
+		auto& Event = EmplaceProfilerEvent<UE::RHI::GPUProfiler::FEvent::FBeginWork>(0);
+		State.BeginTimestamp.Target = &Event.GPUTimestampTOP;
 #endif
 
-			EndQuery(State.BeginTimestamp);
-		}
-
-		if (State.PipelineStats)
-		{
-			BeginQuery(State.PipelineStats);
-		}
+		EndQuery(State.BeginTimestamp);
+	}
 
-		State.bLocalQueriesBegun = true;
+	if (State.PipelineStats)
+	{
+		BeginQuery(State.PipelineStats);
 	}
 }
 
 void FD3D12CommandList::EndLocalQueries()
 {
-	if (!State.bLocalQueriesEnded)
+#if DO_CHECK
+	check(!State.bLocalQueriesEnded);
+	State.bLocalQueriesEnded = true;
+#endif
+
+	if (State.PipelineStats)
 	{
-		if (State.PipelineStats)
-		{
-			EndQuery(State.PipelineStats);
-		}
+		EndQuery(State.PipelineStats);
+	}
 
-		if (State.EndTimestamp)
-		{
+	if (State.EndTimestamp)
+	{
 #if RHI_NEW_GPU_PROFILER
-			auto& Event = EmplaceEvent<UE::RHI::GPUProfiler::FEvent::FEndWork>();
-			State.EndTimestamp.Target = &Event.GPUTimestampBOP;
+		auto& Event = EmplaceProfilerEvent<UE::RHI::GPUProfiler::FEvent::FEndWork>();
+		State.EndTimestamp.Target = &Event.GPUTimestampBOP;
 #endif
 
-			EndQuery(State.EndTimestamp);
-		}
-
-		State.bLocalQueriesEnded = true;
+		EndQuery(State.EndTimestamp);
 	}
 }
 
diff --git a/Engine/Source/Runtime/D3D12RHI/Private/D3D12CommandList.h b/Engine/Source/Runtime/D3D12RHI/Private/D3D12CommandList.h
--- a/Engine/Source/Runtime/D3D12RHI/Private/D3D12CommandList.h
+++ b/Engine/Source/Runtime/D3D12RHI/Private/D3D12CommandList.h
@@ -116,17 +116,29 @@ public:
 #endif // ENABLE_RESIDENCY_MANAGEMENT
 
 #if RHI_NEW_GPU_PROFILER
-	template <typename TEventType>
-	TEventType& EmplaceEvent()
+	template <typename TEventType, typename... TArgs>
+	TEventType& EmplaceProfilerEvent(TArgs&&... Args)
 	{
-		return State.EmplaceEvent<TEventType>();
+		TEventType& Data = State.EventStream.Emplace<TEventType>(Forward<TArgs>(Args)...);
+
+		if constexpr (std::is_same_v<UE::RHI::GPUProfiler::FEvent::FBeginWork, TEventType>)
+		{
+			State.BeginEvents.Add(&Data);
+		}
+
+		return Data;
 	}
 
-	TArray<TUniquePtr<UE::RHI::GPUProfiler::FEvent>>&& RetrieveEvents()
+	void FlushProfilerEvents(UE::RHI::GPUProfiler::FEventStream& Destination, uint64 CPUTimestamp)
 	{
-		return MoveTemp(State.Events);
+		for (UE::RHI::GPUProfiler::FEvent::FBeginWork* BeginEvent : State.BeginEvents)
+		{
+			BeginEvent->CPUTimestamp = CPUTimestamp;
+		}
+
+		Destination.Append(MoveTemp(State.EventStream));
 	}
-#endif // RHI_NEW_GPU_PROFILER
+#endif
 
 private:
 	struct FInterfaces
@@ -302,18 +314,16 @@ private:
 		uint32 NumCommands = 0;
 
 		bool IsClosed = false;
+
+	#if DO_CHECK
 		bool bLocalQueriesBegun = false;
 		bool bLocalQueriesEnded = false;
+	#endif
 
-#if RHI_NEW_GPU_PROFILER
-		TArray<TUniquePtr<UE::RHI::GPUProfiler::FEvent>> Events; // @todo dev-pr : optimize storage
-
-		template <typename TEventType>
-		TEventType& EmplaceEvent()
-		{
-			return Events.Emplace_GetRef(MakeUnique<UE::RHI::GPUProfiler::FEvent>(TEventType()))->Value.Get<TEventType>();
-		}
-#endif
+	#if RHI_NEW_GPU_PROFILER
+		UE::RHI::GPUProfiler::FEventStream EventStream;
+		TArray<UE::RHI::GPUProfiler::FEvent::FBeginWork*, TInlineAllocator<8>> BeginEvents;
+	#endif
 		
 	} State;
-};
+};
\ No newline at end of file
diff --git a/Engine/Source/Runtime/D3D12RHI/Private/D3D12Device.h b/Engine/Source/Runtime/D3D12RHI/Private/D3D12Device.h
--- a/Engine/Source/Runtime/D3D12RHI/Private/D3D12Device.h
+++ b/Engine/Source/Runtime/D3D12RHI/Private/D3D12Device.h
@@ -77,7 +77,7 @@ public:
 	TArray<TSharedPtr<FRHIBreadcrumbAllocatorArray>> BreadcrumbAllocators {};
 	#endif
 
-	TArray<TUniquePtr<UE::RHI::GPUProfiler::FEvent>> Events;
+	UE::RHI::GPUProfiler::FEventStream EventStream;
 
 #else
 
diff --git a/Engine/Source/Runtime/D3D12RHI/Private/D3D12RHI.cpp b/Engine/Source/Runtime/D3D12RHI/Private/D3D12RHI.cpp
--- a/Engine/Source/Runtime/D3D12RHI/Private/D3D12RHI.cpp
+++ b/Engine/Source/Runtime/D3D12RHI/Private/D3D12RHI.cpp
@@ -437,18 +437,20 @@ void FD3D12DynamicRHI::FlushTiming(bool bCreateNew, const FRHIEndFrameArgs& Args
 #if RHI_NEW_GPU_PROFILER
 		if (bHasTiming && Payload.Queue.QueueType != ED3D12QueueType::Copy)
 		{
-			UE::RHI::GPUProfiler::FEvent::FFrameBoundary Boundary;
-			Boundary.FrameNumber = Args.FrameNumber;
-
-		#if WITH_RHI_BREADCRUMBS
+			ERHIPipeline Pipeline;
 			switch (Payload.Queue.QueueType)
 			{
-			case ED3D12QueueType::Direct: Boundary.Breadcrumb = Args.GPUBreadcrumbs[ERHIPipeline::Graphics    ]; break;
-			case ED3D12QueueType::Async : Boundary.Breadcrumb = Args.GPUBreadcrumbs[ERHIPipeline::AsyncCompute]; break;
+			default: checkNoEntry(); [[fallthrough]];
+			case ED3D12QueueType::Direct: Pipeline = ERHIPipeline::Graphics; break;
+			case ED3D12QueueType::Async:  Pipeline = ERHIPipeline::AsyncCompute; break;
 			}
-		#endif
 
-			Payload.Events.Emplace(MakeUnique<UE::RHI::GPUProfiler::FEvent>(Boundary));
+			Payload.EventStream.Emplace<UE::RHI::GPUProfiler::FEvent::FFrameBoundary>(
+				Args.FrameNumber
+		#if WITH_RHI_BREADCRUMBS
+				, Args.GPUBreadcrumbs[Pipeline]
+		#endif
+			);
 		}
 #endif // RHI_NEW_GPU_PROFILER
 	});
diff --git a/Engine/Source/Runtime/D3D12RHI/Private/D3D12Resources.cpp b/Engine/Source/Runtime/D3D12RHI/Private/D3D12Resources.cpp
--- a/Engine/Source/Runtime/D3D12RHI/Private/D3D12Resources.cpp
+++ b/Engine/Source/Runtime/D3D12RHI/Private/D3D12Resources.cpp
@@ -1666,12 +1666,13 @@ void FD3D12ResourceBarrierBatcher::FlushIntoCommandList(FD3D12CommandList& Comma
 #if RHI_NEW_GPU_PROFILER
 		if (bBegin)
 		{
-			auto& Event = CommandList.EmplaceEvent<UE::RHI::GPUProfiler::FEvent::FEndWork>();
+			auto& Event = CommandList.EmplaceProfilerEvent<UE::RHI::GPUProfiler::FEvent::FEndWork>();
 			CommandList.EndQuery(TimestampAllocator.Allocate(ED3D12QueryType::ProfilerTimestampBOP, &Event.GPUTimestampBOP));
 		}
 		else
 		{
-			auto& Event = CommandList.EmplaceEvent<UE::RHI::GPUProfiler::FEvent::FBeginWork>();
+			// CPUTimestamp is filled in at submission time in FlushProfilerEvents
+			auto& Event = CommandList.EmplaceProfilerEvent<UE::RHI::GPUProfiler::FEvent::FBeginWork>(0);
 			CommandList.EndQuery(TimestampAllocator.Allocate(ED3D12QueryType::ProfilerTimestampTOP, &Event.GPUTimestampTOP));
 		}
 #else
diff --git a/Engine/Source/Runtime/D3D12RHI/Private/D3D12Submission.cpp b/Engine/Source/Runtime/D3D12RHI/Private/D3D12Submission.cpp
--- a/Engine/Source/Runtime/D3D12RHI/Private/D3D12Submission.cpp
+++ b/Engine/Source/Runtime/D3D12RHI/Private/D3D12Submission.cpp
@@ -955,14 +955,7 @@ void FD3D12DynamicRHI::FlushBatchedPayloads(FD3D12Queue::FPayloadArray& Payloads
 				check(CommandList->IsClosed());
 
 #if RHI_NEW_GPU_PROFILER
-				{
-					TArray<TUniquePtr<UE::RHI::GPUProfiler::FEvent>> Events = CommandList->RetrieveEvents();
-					if (Events.Num() > 0)
-					{
-						Events[0]->Value.Get<UE::RHI::GPUProfiler::FEvent::FBeginWork>().CPUTimestamp = Time;
-						Payload->Events.Append(MoveTemp(Events));
-					}
-				}
+				CommandList->FlushProfilerEvents(Payload->EventStream, Time);
 #endif // RHI_NEW_GPU_PROFILER
 
 				D3DCommandLists.Add(CommandList->Interfaces.CommandList);
@@ -1415,10 +1408,10 @@ FD3D12DynamicRHI::FProcessResult FD3D12DynamicRHI::ProcessInterruptQueue()
 				CurrentQueue.Timing->BreadcrumbAllocators.Append(MoveTemp(Payload->BatchedObjects.BreadcrumbAllocators));
 			}
 
-			if (Payload->Events.Num())
+			if (!Payload->EventStream.IsEmpty())
 			{
 				check(CurrentQueue.Timing);
-				CurrentQueue.Timing->Events.Append(MoveTemp(Payload->Events));
+				CurrentQueue.Timing->EventStream.Append(MoveTemp(Payload->EventStream));
 			}
 #endif
 
@@ -1486,7 +1479,7 @@ void FD3D12DynamicRHI::ProcessTimestamps(FD3D12TimingArray const& TimingPerQueue
 
 	for (auto const& Timing : TimingPerQueue)
 	{
-		UE::RHI::GPUProfiler::ProcessEvents(Timing->Queue.GetProfilerQueue(), Timing->Events);
+		UE::RHI::GPUProfiler::ProcessEvents(Timing->Queue.GetProfilerQueue(), Timing->EventStream);
 	}
 
 #else
diff --git a/Engine/Source/Runtime/D3D12RHI/Private/D3D12Submission.h b/Engine/Source/Runtime/D3D12RHI/Private/D3D12Submission.h
--- a/Engine/Source/Runtime/D3D12RHI/Private/D3D12Submission.h
+++ b/Engine/Source/Runtime/D3D12RHI/Private/D3D12Submission.h
@@ -246,7 +246,7 @@ struct FD3D12PayloadBase : public FD3D12PayloadBaseFixLayout
 	FRHIBreadcrumbRange BreadcrumbRange {};
 #endif
 #if RHI_NEW_GPU_PROFILER
-	TArray<TUniquePtr<UE::RHI::GPUProfiler::FEvent>> Events;
+	UE::RHI::GPUProfiler::FEventStream EventStream;
 #endif
 
 	virtual ~FD3D12PayloadBase();
diff --git a/Engine/Source/Runtime/RHI/Private/GPUProfiler.cpp b/Engine/Source/Runtime/RHI/Private/GPUProfiler.cpp
--- a/Engine/Source/Runtime/RHI/Private/GPUProfiler.cpp
+++ b/Engine/Source/Runtime/RHI/Private/GPUProfiler.cpp
@@ -687,13 +687,13 @@ namespace UE::RHI::GPUProfiler
 		GetSinks().RemoveSingle(this);
 	}
 
-	void ProcessEvents(FQueue Queue, TConstArrayView<TUniquePtr<UE::RHI::GPUProfiler::FEvent>> Events)
+	void ProcessEvents(FQueue Queue, FEventStream const& EventStream)
 	{
-		if (!Events.IsEmpty())
+		if (!EventStream.IsEmpty())
 		{
 			for (FEventSink* Sink : GetSinks())
 			{
-				Sink->ProcessEvents(Queue, Events);
+				Sink->ProcessEvents(Queue, EventStream);
 			}
 		}
 	}
@@ -742,11 +742,11 @@ namespace UE::RHI::GPUProfiler
 			}
 		}
 
-		void ProcessEvents(FQueue Queue, TConstArrayView<TUniquePtr<UE::RHI::GPUProfiler::FEvent>> Events) override
+		void ProcessEvents(FQueue Queue, FEventStream const& EventStream) override
 		{
 			FQueueState& QueueState = QueueStates.FindChecked(Queue);
 
-			for (auto const& Event : Events)
+			for (FEvent const* Event : EventStream)
 			{
 				switch (Event->GetType())
 				{
@@ -755,7 +755,7 @@ namespace UE::RHI::GPUProfiler
 						check(!QueueState.bBusy);
 						QueueState.bBusy = true;
 						uint64 Value = Event->Value.Get<FEvent::FBeginWork>().GPUTimestampTOP;
-							
+
 						if (!QueueState.Timestamps.IsEmpty() && Value <= QueueState.Timestamps.Last())
 						{
 							//
@@ -787,7 +787,7 @@ namespace UE::RHI::GPUProfiler
 				case FEvent::EType::FrameBoundary:
 					{
 						check(!QueueState.bBusy);
-						auto const& FrameBoundary = Event->Value.Get<FEvent::FFrameBoundary>();
+						FEvent::FFrameBoundary const& FrameBoundary = Event->Value.Get<FEvent::FFrameBoundary>();
 
 						FFrameState& FrameState = Frames.FindOrAdd(FrameBoundary.FrameNumber);
 						FrameState.Emplace(Queue, MoveTemp(QueueState.Timestamps));
@@ -935,7 +935,7 @@ namespace UE::RHI::GPUProfiler
 				Current = Current->Parent;
 			}
 
-			void ProcessEvents(TConstArrayView<TUniquePtr<UE::RHI::GPUProfiler::FEvent>> Events);
+			void ProcessEvents(FEventStream const& EventStream);
 
 			void LogTree(uint32 FrameNumber);
 		};
@@ -951,9 +951,9 @@ namespace UE::RHI::GPUProfiler
 			}
 		}
 
-		void ProcessEvents(FQueue Queue, TConstArrayView<TUniquePtr<UE::RHI::GPUProfiler::FEvent>> Events) override
+		void ProcessEvents(FQueue Queue, FEventStream const& EventStream) override
 		{
-			QueueStates.FindChecked(Queue)->ProcessEvents(Events);
+			QueueStates.FindChecked(Queue)->ProcessEvents(EventStream);
 		}
 
 		void ProfileNextFrame()
@@ -966,9 +966,11 @@ namespace UE::RHI::GPUProfiler
 
 	} GGPUProfilerSink_ProfileGPU;
 
-	void FGPUProfilerSink_ProfileGPU::FQueueState::ProcessEvents(TConstArrayView<TUniquePtr<UE::RHI::GPUProfiler::FEvent>> Events)
+	void FGPUProfilerSink_ProfileGPU::FQueueState::ProcessEvents(FEventStream const& EventStream)
 	{
-		int32 Index = 0;
+		auto Iterator = std::begin(EventStream);
+		auto End = std::end(EventStream);
+
 	Restart:
 		if (State == EState::Idle && bTriggerProfile.exchange(false))
 		{
@@ -978,14 +980,14 @@ namespace UE::RHI::GPUProfiler
 		if (State == EState::WaitingFrame)
 		{
 			// Discard all received events until we reach a FrameBoundary event
-			for (; Index < Events.Num(); ++Index)
+			for (; Iterator != End; ++Iterator)
 			{
-				if (Events[Index]->GetType() == FEvent::EType::FrameBoundary)
+				if ((*Iterator)->GetType() == FEvent::EType::FrameBoundary)
 				{
 					// Start profiling until we receive another FrameBoundary event
 					State = EState::Active;
 
-					FEvent::FFrameBoundary& FrameBoundary = Events[Index]->Value.Get<FEvent::FFrameBoundary>();
+					FEvent::FFrameBoundary const& FrameBoundary = (*Iterator)->Value.Get<FEvent::FFrameBoundary>();
 
 					// Build the node tree 
 					PushNode(TEXT("<root>"));
@@ -1006,7 +1008,7 @@ namespace UE::RHI::GPUProfiler
 					Recurse(Recurse, FrameBoundary.Breadcrumb);
 				#endif // WITH_RHI_BREADCRUMBS
 
-					++Index;
+					++Iterator;
 					break;
 				}
 			}
@@ -1014,9 +1016,9 @@ namespace UE::RHI::GPUProfiler
 
 		if (State == EState::Active)
 		{
-			for (; Index < Events.Num(); ++Index)
+			for (; Iterator != End; ++Iterator)
 			{
-				auto const& Event = Events[Index];
+				FEvent const* Event = *Iterator;
 
 				switch (Event->GetType())
 				{
@@ -1043,7 +1045,7 @@ namespace UE::RHI::GPUProfiler
 			#if WITH_RHI_BREADCRUMBS
 				case FEvent::EType::BeginBreadcrumb:
 					{
-						auto& BeginBreadcrumb = Event->Value.Get<FEvent::FBeginBreadcrumb>();
+						FEvent::FBeginBreadcrumb const& BeginBreadcrumb = Event->Value.Get<FEvent::FBeginBreadcrumb>();
 
 						// Push a new node
 						FRHIBreadcrumb::FBuffer Buffer;
@@ -1055,7 +1057,7 @@ namespace UE::RHI::GPUProfiler
 
 				case FEvent::EType::EndBreadcrumb:
 					{
-						auto& EndBreadcrumb = Event->Value.Get<FEvent::FEndBreadcrumb>();
+						FEvent::FEndBreadcrumb const& EndBreadcrumb = Event->Value.Get<FEvent::FEndBreadcrumb>();
 						Current->Timestamps.Add(EndBreadcrumb.GPUTimestampBOP);
 
 						PopNode();
@@ -1065,7 +1067,7 @@ namespace UE::RHI::GPUProfiler
 
 				case FEvent::EType::Stats:
 					{
-						auto& Stats = Event->Value.Get<FEvent::FStats>();
+						FEvent::FStats const& Stats = Event->Value.Get<FEvent::FStats>();
 						Current->NumDraws += Stats.NumDraws;
 						Current->NumPrimitives += Stats.NumPrimitives;
 					}
@@ -1073,7 +1075,7 @@ namespace UE::RHI::GPUProfiler
 
 				case FEvent::EType::FrameBoundary:
 					{
-						FEvent::FFrameBoundary& FrameBoundary = Event->Value.Get<FEvent::FFrameBoundary>();
+						FEvent::FFrameBoundary const& FrameBoundary = Event->Value.Get<FEvent::FFrameBoundary>();
 						LogTree(FrameBoundary.FrameNumber);
 
 						// Reset tracking
@@ -1188,6 +1190,8 @@ namespace UE::RHI::GPUProfiler
 		{
 			GGPUProfilerSink_ProfileGPU.ProfileNextFrame();
 		}));
+
+	TLockFreePointerListUnordered<void, PLATFORM_CACHE_LINE_SIZE> FEventStream::FChunk::MemoryPool;
 }
 
 #endif // RHI_NEW_GPU_PROFILER
diff --git a/Engine/Source/Runtime/RHI/Public/GPUProfiler.h b/Engine/Source/Runtime/RHI/Public/GPUProfiler.h
--- a/Engine/Source/Runtime/RHI/Public/GPUProfiler.h
+++ b/Engine/Source/Runtime/RHI/Public/GPUProfiler.h
@@ -89,25 +89,44 @@ namespace UE::RHI::GPUProfiler
 			// Very first frame of the engine is frame 0 (from boot to first call to RHIEndFrame).
 			uint32 FrameNumber;
 
-#if WITH_RHI_BREADCRUMBS
+		#if WITH_RHI_BREADCRUMBS
 			// The RHI breadcrumb currently at the top of the stack at the frame boundary.
 			FRHIBreadcrumbNode* Breadcrumb;
-#endif
+		#endif
+
+			FFrameBoundary(uint32 FrameNumber
+			#if WITH_RHI_BREADCRUMBS
+				, FRHIBreadcrumbNode* Breadcrumb
+			#endif
+				)
+				: FrameNumber(FrameNumber)
+			#if WITH_RHI_BREADCRUMBS
+				, Breadcrumb(Breadcrumb)
+			#endif
+			{}
 		};
 
-#if WITH_RHI_BREADCRUMBS
+	#if WITH_RHI_BREADCRUMBS
 		struct FBeginBreadcrumb
 		{
 			FRHIBreadcrumbNode* Breadcrumb;
-			uint64 GPUTimestampTOP;
+			uint64 GPUTimestampTOP = 0;
+
+			FBeginBreadcrumb(FRHIBreadcrumbNode* Breadcrumb)
+				: Breadcrumb(Breadcrumb)
+			{}
 		};
 
 		struct FEndBreadcrumb
 		{
 			FRHIBreadcrumbNode* Breadcrumb;
-			uint64 GPUTimestampBOP;
+			uint64 GPUTimestampBOP = 0;
+
+			FEndBreadcrumb(FRHIBreadcrumbNode* Breadcrumb)
+				: Breadcrumb(Breadcrumb)
+			{}
 		};
-#endif // WITH_RHI_BREADCRUMBS
+	#endif
 
 		// Inserted when the GPU starts work on a queue.
 		struct FBeginWork
@@ -116,13 +135,17 @@ namespace UE::RHI::GPUProfiler
 			uint64 CPUTimestamp;
 
 			// TOP timestamp of when the work actually started on the GPU.
-			uint64 GPUTimestampTOP;
+			uint64 GPUTimestampTOP = 0;
+
+			FBeginWork(uint64 CPUTimestamp)
+				: CPUTimestamp(CPUTimestamp)
+			{}
 		};
 
 		// Inserted when the GPU completes work on a queue and goes idle.
 		struct FEndWork
 		{
-			uint64 GPUTimestampBOP;
+			uint64 GPUTimestampBOP = 0;
 		};
 
 		struct FStats
@@ -160,10 +183,10 @@ namespace UE::RHI::GPUProfiler
 		
 		using FStorage = TVariant<
 			  FFrameBoundary
-#if WITH_RHI_BREADCRUMBS
+		#if WITH_RHI_BREADCRUMBS
 			, FBeginBreadcrumb
 			, FEndBreadcrumb
-#endif
+		#endif
 			, FBeginWork
 			, FEndWork
 			, FStats
@@ -176,10 +199,10 @@ namespace UE::RHI::GPUProfiler
 		enum class EType
 		{
 			FrameBoundary   = FStorage::IndexOfType<FFrameBoundary  >(),
-#if WITH_RHI_BREADCRUMBS
+		#if WITH_RHI_BREADCRUMBS
 			BeginBreadcrumb = FStorage::IndexOfType<FBeginBreadcrumb>(),
 			EndBreadcrumb   = FStorage::IndexOfType<FEndBreadcrumb  >(),
-#endif
+		#endif
 			BeginWork       = FStorage::IndexOfType<FBeginWork      >(),
 			EndWork         = FStorage::IndexOfType<FEndWork        >(),
 			Stats           = FStorage::IndexOfType<FStats          >(),
@@ -200,6 +223,174 @@ namespace UE::RHI::GPUProfiler
 		FEvent(T const& Value)
 			: Value(TInPlaceType<T>(), Value)
 		{}
+
+		FEvent(FEvent const&) = delete;
+		FEvent(FEvent&&) = delete;
+	};
+
+	class FEventStream
+	{
+	private:
+		struct FChunk
+		{
+			struct FHeader
+			{
+				FChunk* Next = nullptr;
+				uint32 Num = 0;
+			} Header;
+
+			static constexpr uint32 ChunkSizeInBytes = 16 * 1024;
+			static constexpr uint32 RemainingBytes = ChunkSizeInBytes - Align<uint32>(sizeof(FHeader), alignof(FHeader));
+			static constexpr uint32 MaxEventsPerChunk = RemainingBytes / Align<uint32>(sizeof(FEvent), alignof(FEvent));
+
+			TStaticArray<TTypeCompatibleBytes<FEvent>, MaxEventsPerChunk> Elements;
+
+			static RHI_API TLockFreePointerListUnordered<void, PLATFORM_CACHE_LINE_SIZE> MemoryPool;
+
+			void* operator new(size_t Size)
+			{
+				check(Size == sizeof(FChunk));
+
+				void* Memory = MemoryPool.Pop();
+				if (!Memory)
+				{
+					Memory = FMemory::Malloc(sizeof(FChunk), alignof(FChunk));
+				}
+				return Memory;
+			}
+
+			void operator delete(void* Pointer)
+			{
+				MemoryPool.Push(Pointer);
+			}
+
+			FEvent* GetElement(uint32 Index)
+			{
+				return Elements[Index].GetTypedPtr();
+			}
+		};
+
+		static_assert(sizeof(FChunk) == FChunk::ChunkSizeInBytes, "Incorrect FChunk size.");
+
+		FChunk* First = nullptr;
+		FChunk* Current = nullptr;
+
+	public:
+		FEventStream() = default;
+		FEventStream(FEventStream const&) = delete;
+
+		FEventStream(FEventStream&& Other)
+			: First(Other.First)
+			, Current(Other.Current)
+		{
+			Other.First = nullptr;
+			Other.Current = nullptr;
+		}
+
+		~FEventStream()
+		{
+			while (First)
+			{
+				FChunk* Next = First->Header.Next;
+				delete First;
+				First = Next;
+			}
+		}
+
+		template <typename TEventType, typename... TArgs>
+		TEventType& Emplace(TArgs&&... Args)
+		{
+			static_assert(TIsTrivial<TEventType>::Value, "Destructors are not called on GPU profiler events, so the types must be trivial.");
+
+			if (!Current)
+			{
+				Current = new FChunk;
+				if (!First)
+				{
+					First = Current;
+				}
+			}
+
+			if (Current->Header.Num >= FChunk::MaxEventsPerChunk)
+			{
+				FChunk* NewChunk = new FChunk;
+				Current->Header.Next = NewChunk;
+				Current = NewChunk;
+			}
+
+			FEvent* Event = Current->GetElement(Current->Header.Num++);
+			new (Event) FEvent(TEventType(Forward<TArgs>(Args)...));
+
+			return Event->Value.Get<TEventType>();
+		}
+
+		bool IsEmpty() const
+		{
+			return First == nullptr;
+		}
+
+		void Append(FEventStream&& Other)
+		{
+			if (IsEmpty())
+			{
+				Current = Other.Current;
+				First = Other.First;
+			}
+			else if (!Other.IsEmpty())
+			{
+				Current->Header.Next = Other.First;
+				Current = Other.Current;
+			}
+
+			Other.Current = nullptr;
+			Other.First = nullptr;
+		}
+
+		auto begin() const
+		{
+			class FIterator
+			{
+				friend FEventStream;
+
+				FChunk* Current;
+				uint32 Index = 0;
+
+				FIterator(FChunk* Current)
+					: Current(Current)
+				{}
+
+			public:
+				FIterator& operator++()
+				{
+					++Index;
+
+					while (Current && Index >= Current->Header.Num)
+					{
+						Current = Current->Header.Next;
+						Index = 0;
+					}
+
+					return *this;
+				}
+
+				bool operator != (std::nullptr_t) const
+				{
+					return Current != nullptr;
+				}
+
+				FEvent const* operator*() const
+				{
+					return Current->GetElement(Index);
+				}
+			};
+
+			return FIterator(First);
+		}
+
+		std::nullptr_t end() const
+		{
+			return nullptr;
+		}
 	};
 
 	struct FEventSink
@@ -212,11 +403,11 @@ namespace UE::RHI::GPUProfiler
 		FEventSink(FEventSink&&) = delete;
 
 	public:
-		virtual void ProcessEvents(FQueue Queue, TConstArrayView<TUniquePtr<UE::RHI::GPUProfiler::FEvent>> Events) = 0;
+		virtual void ProcessEvents(FQueue Queue, FEventStream const& EventStream) = 0;
 		virtual void InitializeQueues(TConstArrayView<FQueue> Queues) = 0;
 	};
 
-	RHI_API void ProcessEvents(FQueue Queue, TConstArrayView<TUniquePtr<UE::RHI::GPUProfiler::FEvent>> Events);
+	RHI_API void ProcessEvents(FQueue Queue, FEventStream const& EventStream);
 	RHI_API void InitializeQueues(TConstArrayView<FQueue> Queues);
 }
 
