PipelineStateCache Asynchronous Consolidation

#tests ReplayRun in Dev-PerfB-Main

diff --git a/Engine/Source/Runtime/Core/Private/HAL/ThreadingBase.cpp b/Engine/Source/Runtime/Core/Private/HAL/ThreadingBase.cpp
--- a/Engine/Source/Runtime/Core/Private/HAL/ThreadingBase.cpp
+++ b/Engine/Source/Runtime/Core/Private/HAL/ThreadingBase.cpp
@@ -335,6 +335,12 @@ CORE_API bool IsInRHIThread()
 {
 	return FTaskTagScope::IsCurrentTag(ETaskTag::ERhiThread);
 }
+
+CORE_API bool IsInParallelRHIThread()
+{
+	return (FTaskTagScope::GetCurrentTag() & ETaskTag::ERhiThread) == ETaskTag::ERhiThread;
+}
+
 // Fake threads
 
 // Core version of IsInAsyncLoadingThread
diff --git a/Engine/Source/Runtime/Core/Public/CoreGlobals.h b/Engine/Source/Runtime/Core/Public/CoreGlobals.h
--- a/Engine/Source/Runtime/Core/Public/CoreGlobals.h
+++ b/Engine/Source/Runtime/Core/Public/CoreGlobals.h
@@ -774,6 +774,9 @@ extern CORE_API bool IsRHIThreadRunning();
 /** @return True if called from the RHI thread, or if called from ANY thread during single threaded rendering */
 extern CORE_API bool IsInRHIThread();
 
+/** @return True if called from a thread tagged as an RHI task. */
+extern CORE_API bool IsInParallelRHIThread();
+
 extern CORE_API bool GIsRunningRHIInSeparateThread_InternalUseOnly;
 extern CORE_API bool GIsRunningRHIInDedicatedThread_InternalUseOnly;
 extern CORE_API bool GIsRunningRHIInTaskThread_InternalUseOnly;
diff --git a/Engine/Source/Runtime/RHI/Private/PipelineStateCache.cpp b/Engine/Source/Runtime/RHI/Private/PipelineStateCache.cpp
--- a/Engine/Source/Runtime/RHI/Private/PipelineStateCache.cpp
+++ b/Engine/Source/Runtime/RHI/Private/PipelineStateCache.cpp
@@ -148,6 +148,15 @@ static FAutoConsoleVariableRef CVarCreatePSOsOnRHIThread(
 	ECVF_RenderThreadSafe
 );
 
+bool GEnablePSOAsyncCacheConsolidation = true;
+static FAutoConsoleVariableRef CVarEnablePSOAsyncCacheConsolidation(
+	TEXT("r.pso.EnableAsyncCacheConsolidation"),
+	GEnablePSOAsyncCacheConsolidation,
+	TEXT("0: Require Render Thread and RHI Thread to synchronize before flushing the PSO cache.")
+	TEXT("1: Flush the PSO cache without synchronizing the Render Thread with the RHI Thread.\n"),
+	ECVF_RenderThreadSafe
+);
+
 static TAutoConsoleVariable<int32> CVarPSOEvictionTime(
 	TEXT("r.pso.evictiontime"),
 	60,
@@ -611,48 +620,10 @@ public:
 		InitStats();
 	}
 
-	virtual ~FPipelineState() 
-	{
-		check(IsComplete());
-		verify(!WaitCompletion());
-		check(!PrecompileTask.IsValid());
-	}
+	virtual ~FPipelineState() = default;
 
 	virtual bool IsCompute() const = 0;
 
-	FGraphEventRef CompletionEvent;
-	TUniquePtr<FPSOPrecacheAsyncTask> PrecompileTask;
-
-	bool IsComplete()
-	{
-		return (!CompletionEvent.IsValid() || CompletionEvent->IsComplete()) && (!PrecompileTask.IsValid() || PrecompileTask->IsDone());
-	}
-
-	// return true if we actually waited on the task
-	bool WaitCompletion()
-	{
-		bool bNeedsToWait = false;
-		if(CompletionEvent.IsValid() && !CompletionEvent->IsComplete())
-		{
-			QUICK_SCOPE_CYCLE_COUNTER(STAT_FPipelineState_WaitCompletion);
-#if PSO_TRACK_CACHE_STATS
-			UE_LOG(LogRHI, Log, TEXT("FTaskGraphInterface Waiting on FPipelineState completionEvent"));
-#endif
-			bNeedsToWait = true;
-			FTaskGraphInterface::Get().WaitUntilTaskCompletes( CompletionEvent );
-		}
-		CompletionEvent = nullptr;
-
-		if (PrecompileTask.IsValid())
-		{
-			bNeedsToWait = bNeedsToWait || !PrecompileTask->IsDone();
-			PrecompileTask->EnsureCompletion();
-			PrecompileTask = nullptr;
-		}
-
-		return bNeedsToWait;
-	}
-
 	inline void AddUse()
 	{
 		FPipelineStateStats::UpdateStats(Stats);
@@ -694,8 +665,187 @@ public:
 	FPipelineStateStats* Stats;
 };
 
+namespace PipelineStateCache
+{
+	constexpr int32 RenderThreadIndex = 0;
+	constexpr int32 RHIThreadIndex = 1;
+
+	int32 GetCacheIndexForCurrentThread()
+	{
+		return IsInParallelRHIThread() || !GEnablePSOAsyncCacheConsolidation;
+	}
+}
+
+/**
+ * Base class for pipeline state intended to be stored in a TSharedPipelineStateCache,
+ * with state double buffering for Render and RHI Threads.
+ */
+class FPipelineStateAsync : public FPipelineState
+{
+private:
+
+	struct FCompletionState
+	{
+		FGraphEventRef CompletionEvent;
+		TUniquePtr<FPSOPrecacheAsyncTask> PrecompileTask;
+	};
+	TSharedPtr<FCompletionState> CompletionStates[2];
+
+	// GetCompletionState is thread safe on Render or RHI Threads.
+	const TSharedPtr<FCompletionState>& GetCompletionState() const
+	{
+		return CompletionStates[PipelineStateCache::GetCacheIndexForCurrentThread()];
+	}
+
+	// MakeCompletionState is not thread safe.
+	const TSharedPtr<FCompletionState>& MakeCompletionState()
+	{
+		const TSharedPtr<FCompletionState>& CompletionState = GetCompletionState();
+		if (!CompletionState)
+		{
+			CompletionStates[1] = CompletionStates[0] = MakeShared<FCompletionState>();
+		}
+		return CompletionState;
+	}
+
+	// ClearCompletionState can be called safely by a Render and RHI task when there are no parallel tasks of the same type.
+	void ClearCompletionState()
+	{
+		int32 CacheIndex = PipelineStateCache::GetCacheIndexForCurrentThread();
+		CompletionStates[CacheIndex] = nullptr;
+
+		// Clear both references if asynchronous pipeline state cache is disabled.
+		if (!IsInParallelRHIThread())
+		{
+			// Accessing GEnablePSOAsyncCacheConsolidation is safe on the Render Thread
+			if (!GEnablePSOAsyncCacheConsolidation)
+			{
+				CompletionStates[!CacheIndex] = nullptr;
+			}
+		}
+	}
+
+public:
+
+	virtual ~FPipelineStateAsync() override
+	{
+		check(IsComplete());
+		verify(!WaitCompletion());
+	}
+
+	// GetCompletionEvent is thread safe on Render or RHI Threads as long as ClearCompletionState is not called on a parallel task.
+	FGraphEvent* GetCompletionEvent() const
+	{
+		const TSharedPtr<FCompletionState>& CompletionState = GetCompletionState();
+		return CompletionState ? GetCompletionState()->CompletionEvent.GetReference() : nullptr;
+	}
+
+	// SetCompletionEvent is not thread safe and should be called before adding this state to the cache.
+	void SetCompletionEvent(FGraphEventRef InCompletionEvent)
+	{
+		MakeCompletionState()->CompletionEvent = MoveTemp(InCompletionEvent);
+	}
+
+	// GetPrecompileTask is thread safe on Render or RHI Threads as long as ClearCompletionState is not called on a parallel task.
+	FPSOPrecacheAsyncTask* GetPrecompileTask() const
+	{
+		const TSharedPtr<FCompletionState>& CompletionState = GetCompletionState();
+		return CompletionState ? CompletionState->PrecompileTask.Get() : nullptr;
+	}
+
+	// SetPrecompileTask is not thread safe and should be called before adding this state to the cache.
+	void SetPrecompileTask(TUniquePtr<FPSOPrecacheAsyncTask> InPrecompileTask)
+	{
+		MakeCompletionState()->PrecompileTask = MoveTemp(InPrecompileTask);
+	}
+
+	// IsComplete is thread safe on Render or RHI Threads as long as ClearCompletionState is not called on a parallel task.
+	bool IsComplete()
+	{
+		const TSharedPtr<FCompletionState>& CompletionState = GetCompletionState();
+		return CompletionState == nullptr || ((!CompletionState->CompletionEvent.IsValid() || CompletionState->CompletionEvent->IsComplete()) && (!CompletionState->PrecompileTask.IsValid() || CompletionState->PrecompileTask->IsDone()));
+	}
+
+	// WaitCompletion can be called safely by a Render and RHI task when there are no parallel tasks of the same type.
+	// return true if we actually waited on the task
+	bool WaitCompletion()
+	{
+		bool bNeedsToWait = false;
+		FGraphEvent* CompletionEvent = GetCompletionEvent();
+		if (CompletionEvent && !CompletionEvent->IsComplete())
+		{
+			QUICK_SCOPE_CYCLE_COUNTER(STAT_FPipelineState_WaitCompletion);
+#if PSO_TRACK_CACHE_STATS
+			UE_LOG(LogRHI, Log, TEXT("FTaskGraphInterface Waiting on FPipelineState completionEvent"));
+#endif
+			bNeedsToWait = true;
+			FTaskGraphInterface::Get().WaitUntilTaskCompletes(CompletionEvent);
+		}
+
+		FPSOPrecacheAsyncTask* PrecompileTask = GetPrecompileTask();
+		if (PrecompileTask)
+		{
+			bNeedsToWait = bNeedsToWait || !PrecompileTask->IsDone();
+			PrecompileTask->EnsureCompletion();
+		}
+
+		ClearCompletionState();
+		return bNeedsToWait;
+	}
+};
+
+/**
+ * Base class for pipeline state that doesn't need state double buffering.
+ */
+class FPipelineStateSync : public FPipelineState
+{
+public:
+
+	virtual ~FPipelineStateSync() 
+	{
+		check(IsComplete());
+		verify(!WaitCompletion());
+		check(!PrecompileTask.IsValid());
+	}
+
+	virtual bool IsCompute() const = 0;
+
+	FGraphEventRef CompletionEvent;
+	TUniquePtr<FPSOPrecacheAsyncTask> PrecompileTask;
+
+	bool IsComplete()
+	{
+		return (!CompletionEvent.IsValid() || CompletionEvent->IsComplete()) && (!PrecompileTask.IsValid() || PrecompileTask->IsDone());
+	}
+
+	// return true if we actually waited on the task
+	bool WaitCompletion()
+	{
+		bool bNeedsToWait = false;
+		if(CompletionEvent.IsValid() && !CompletionEvent->IsComplete())
+		{
+			QUICK_SCOPE_CYCLE_COUNTER(STAT_FPipelineState_WaitCompletion);
+#if PSO_TRACK_CACHE_STATS
+			UE_LOG(LogRHI, Log, TEXT("FTaskGraphInterface Waiting on FPipelineState completionEvent"));
+#endif
+			bNeedsToWait = true;
+			FTaskGraphInterface::Get().WaitUntilTaskCompletes( CompletionEvent );
+		}
+		CompletionEvent = nullptr;
+
+		if (PrecompileTask.IsValid())
+		{
+			bNeedsToWait = bNeedsToWait || !PrecompileTask->IsDone();
+			PrecompileTask->EnsureCompletion();
+			PrecompileTask = nullptr;
+		}
+
+		return bNeedsToWait;
+	}
+};
+
 /* State for compute  */
-class FComputePipelineState : public FPipelineState
+class FComputePipelineState : public FPipelineStateAsync
 {
 public:
 	FComputePipelineState(FRHIComputeShader* InComputeShader, const TCHAR* InName)
@@ -746,7 +896,7 @@ public:
 };
 
 /* State for work graphs  */
-class FWorkGraphPipelineState : public FPipelineState
+class FWorkGraphPipelineState : public FPipelineStateAsync
 {
 public:
 	FWorkGraphPipelineState(FRHIWorkGraphShader* InWorkGraphShader)
@@ -767,7 +917,8 @@ public:
 
 	bool IsCompilationComplete() const
 	{
-		return !CompletionEvent.IsValid() || CompletionEvent->IsComplete();
+		FGraphEvent* CompletionEvent = GetCompletionEvent();
+		return !CompletionEvent || CompletionEvent->IsComplete();
 	}
 
 	inline void Verify_IncUse()
@@ -801,7 +952,7 @@ public:
 };
 
 /* State for graphics */
-class FGraphicsPipelineState : public FPipelineState
+class FGraphicsPipelineState : public FPipelineStateAsync
 {
 public:
 	FGraphicsPipelineState() 
@@ -846,19 +997,17 @@ public:
 FRHIComputePipelineState* GetRHIComputePipelineState(FComputePipelineState* PipelineState)
 {
 	ensure(PipelineState->RHIPipeline);
-	PipelineState->CompletionEvent = nullptr;
 	return PipelineState->RHIPipeline;
 }
 
 FRHIWorkGraphPipelineState* GetRHIWorkGraphPipelineState(FWorkGraphPipelineState* PipelineState)
 {
 	ensure(PipelineState->RHIPipeline);
-	PipelineState->CompletionEvent = nullptr;
 	return PipelineState->RHIPipeline;
 }
 
 /* State for ray tracing */
-class FRayTracingPipelineState : public FPipelineState
+class FRayTracingPipelineState : public FPipelineStateSync
 {
 public:
 	FRayTracingPipelineState(const FRayTracingPipelineStateInitializer& Initializer)
@@ -901,7 +1050,7 @@ public:
 			HitsAcrossFrames++;
 		}
 
-		FPipelineState::AddHit();
+		FPipelineStateSync::AddHit();
 	}
 
 	bool operator < (const FRayTracingPipelineState& Other)
@@ -1076,24 +1225,43 @@ private:
 
 	TMap<TMyKey, TMyValue>& GetLocalCache()
 	{
-		void* TLSValue = FPlatformTLS::GetTlsValue(TLSSlot);
-		if (TLSValue == nullptr)
+		// Find or create storage for two PipelineStateCacheTypes for this thread.
+		TOptional<FPipelineStateCacheType>* PipelineStateCaches = static_cast<TOptional<FPipelineStateCacheType>*>(FPlatformTLS::GetTlsValue(TLSSlot));
+		if (!PipelineStateCaches)
+		{
+			PipelineStateCaches = new TOptional<FPipelineStateCacheType>[2];
+			FPlatformTLS::SetTlsValue(TLSSlot, PipelineStateCaches);
+		}
+
+		// Select the cache to use, based on whether or not this thread is processing RHI tasks.
+		int32 CacheIndex = PipelineStateCache::GetCacheIndexForCurrentThread();
+		TOptional<FPipelineStateCacheType>& PipelineStateCache = PipelineStateCaches[CacheIndex];
+		if (!PipelineStateCache)
 		{
-			FPipelineStateCacheType* PipelineStateCache = new FPipelineStateCacheType;
-			FPlatformTLS::SetTlsValue(TLSSlot, (void*)(PipelineStateCache) );
+			// If the cache doesn't exist, create it and register it with the appropriate cache directories.
+			PipelineStateCache.Emplace();
 
 			FScopeLock S(&AllThreadsLock);
-			AllThreadsPipelineStateCache.Add(PipelineStateCache);
-			return *PipelineStateCache;
+			AllThreadsPipelineStateCache.Add(&PipelineStateCache.GetValue());
+			if (CacheIndex == PipelineStateCache::RHIThreadIndex)
+			{
+				RHIThreadsPipelineStateCache.Add(&PipelineStateCache.GetValue());
+			}
+			else
+			{
+				RenderThreadsPipelineStateCache.Add(&PipelineStateCache.GetValue());
+			}
 		}
-		return *((FPipelineStateCacheType*)TLSValue);
+
+		// Return the selected cache.
+		return PipelineStateCache.GetValue();
 	}
 
 #if PIPELINESTATECACHE_VERIFYTHREADSAFE
 	struct FScopeVerifyIncrement
 	{
-		volatile int32 &VerifyMutex;
-		FScopeVerifyIncrement(volatile int32& InVerifyMutex) : VerifyMutex(InVerifyMutex)
+		volatile int32& VerifyMutex;
+		FScopeVerifyIncrement(volatile int32(&InVerifyMutex)[2]) : VerifyMutex(InVerifyMutex[PipelineStateCache::GetCacheIndexForCurrentThread()])
 		{
 			int32 Result = FPlatformAtomics::InterlockedIncrement(&VerifyMutex);
 			if (Result <= 0)
@@ -1114,8 +1282,8 @@ private:
 
 	struct FScopeVerifyDecrement
 	{
-		volatile int32 &VerifyMutex;
-		FScopeVerifyDecrement(volatile int32& InVerifyMutex) : VerifyMutex(InVerifyMutex)
+		volatile int32& VerifyMutex;
+		FScopeVerifyDecrement(volatile int32(&InVerifyMutex)[2]) : VerifyMutex(InVerifyMutex[PipelineStateCache::GetCacheIndexForCurrentThread()])
 		{
 			int32 Result = FPlatformAtomics::InterlockedDecrement(&VerifyMutex);
 			if (Result >= 0)
@@ -1136,32 +1304,39 @@ private:
 #endif
 
 public:
-	typedef TMap<TMyKey,TMyValue> FPipelineStateCacheType;
+	typedef TMap<TMyKey, TMyValue> FPipelineStateCacheType;
 
-	TSharedPipelineStateCache()
-	{
-		CurrentMap = &Map1;
-		BackfillMap = &Map2;
-		DuplicateStateGenerated = 0;
-		TLSSlot = FPlatformTLS::AllocTlsSlot();
-	}
+	TSharedPipelineStateCache() = default;
 
-	bool Find( const TMyKey& InKey, TMyValue& OutResult )
+	bool Find(const TMyKey& InKey, TMyValue& OutResult)
 	{
 #if PIPELINESTATECACHE_VERIFYTHREADSAFE
 		FScopeVerifyIncrement S(VerifyMutex);
 #endif
+
+		// Select the maps to use.
+		FPipelineStateCacheType* LocalCurrentMap = CurrentMap;
+		FPipelineStateCacheType* LocalBackfillMap = BackfillMap;
+		if (!IsInParallelRHIThread() && !IsInRHIThread())
+		{
+			if (GEnablePSOAsyncCacheConsolidation)
+			{
+				LocalCurrentMap = CurrentMap_RenderThread;
+				LocalBackfillMap = BackfillMap_RenderThread;
+			}
+		}
+
 		// safe because we only ever find when we don't add
-		TMyValue* Result = CurrentMap->Find(InKey);
+		TMyValue* Result = LocalCurrentMap->Find(InKey);
 
-		if ( Result )
+		if (Result)
 		{
 			OutResult = *Result;
 			return true;
 		}
 
 		// check the local cahce which is safe because only this thread adds to it
-		TMap<TMyKey, TMyValue> &LocalCache = GetLocalCache();
+		TMap<TMyKey, TMyValue>& LocalCache = GetLocalCache();
 		// if it's not in the local cache then it will rebuild
 		Result = LocalCache.Find(InKey);
 		if (Result)
@@ -1170,9 +1345,9 @@ public:
 			return true;
 		}
 
-		Result = BackfillMap->Find(InKey);
+		Result = LocalBackfillMap->Find(InKey);
 
-		if ( Result )
+		if (Result)
 		{
 			LocalCache.Add(InKey, *Result);
 			OutResult = *Result;
@@ -1181,11 +1356,7 @@ public:
 
 
 		return false;
-		
-		
 	}
-	// The list of tasks that are still in progress.
-	TArray<TTuple<TMyKey,TMyValue>> Uncompleted;
 
 	bool Add(const TMyKey& InKey, const TMyValue& InValue)
 	{
@@ -1193,27 +1364,169 @@ public:
 		FScopeVerifyIncrement S(VerifyMutex);
 #endif
 		// everything is added to the local cache then at end of frame we consolidate them all
-		TMap<TMyKey, TMyValue> &LocalCache = GetLocalCache();
+		TMap<TMyKey, TMyValue>& LocalCache = GetLocalCache();
 
-		check( LocalCache.Contains(InKey) == false );
+		check(LocalCache.Contains(InKey) == false);
 		LocalCache.Add(InKey, InValue);
 		checkfSlow(LocalCache.Contains(InKey), TEXT("PSO not found immediately after adding.  Likely cause is an uninitialized field in a constructor or copy constructor"));
 		return true;
 	}
 
-	void ConsolidateThreadedCaches()
+	// Call from the Render Thread
+	void FlushResources(bool bInDiscardAndSwap)
 	{
+		SCOPED_NAMED_EVENT(ConsolidateThreadedCaches, FColor::Turquoise);
 
-		SCOPE_TIME_GUARD_MS(TEXT("ConsolidatePipelineCache"), 0.1);
-		check(IsInRenderingThread());
+		bPendingDiscardAndSwap |= bInDiscardAndSwap;
+
+#if PIPELINESTATECACHE_VERIFYTHREADSAFE
+		FScopeVerifyDecrement S(VerifyMutex);
+#endif
+
+		// Determine if the RHI THread is still consolidating its cache.
+		if (RHICompletionEvent)
+		{
+			if (!RHICompletionEvent->IsComplete())
+			{
+				return;
+			}
+			RHICompletionEvent = nullptr;
+
+			// Finish asynchronous cache consolidation.
+			FinishAsyncCacheConsolidation();
+		}
+
+		bDiscardAndSwap = bPendingDiscardAndSwap;
+		bPendingDiscardAndSwap = false;
+
+		if (GEnablePSOAsyncCacheConsolidation)
+		{
+			// Determine if asynchronous cache consolidation was just enabled.
+			if (CurrentMap_RenderThread->IsEmpty() && BackfillMap_RenderThread->IsEmpty())
+			{
+				// If the maps just happen to be empty, this will be cheap.
+				OnAsyncConsolidationEnabled();
+			}
+
+			// Initiate an asynchronous cache consolidation.
+			StartAsyncCacheConsolidation();
+		}
+		else
+		{
+			// Determine if asynchronous cache consolidation was just disabled.
+			if (!CurrentMap_RenderThread->IsEmpty() || !BackfillMap_RenderThread->IsEmpty())
+			{
+				OnAsyncConsolidationDisabled();
+			}
+
+			// Synchronously consolidate all caches.
+			ConsolidateThreadedCaches();
+			ProcessDelayedCleanup();
+			ReleasedEntries = 0;
+			if (bDiscardAndSwap)
+			{
+				ReleasedEntries = DiscardAndSwap(CurrentMap, BackfillMap);
+				bDiscardAndSwap = false;
+			}
+		}
+	}
+
+	void Shutdown()
+	{
 #if PIPELINESTATECACHE_VERIFYTHREADSAFE
 		FScopeVerifyDecrement S(VerifyMutex);
 #endif
+
+		if (RHICompletionEvent)
+		{
+			RHICompletionEvent->Wait();
+			RHICompletionEvent = nullptr;
+
+			// Finish asynchronous cache consolidation.
+			FinishAsyncCacheConsolidation();
+		}
+
+		// Determine if asynchronous cache consolidation was just disabled.
+		if (!CurrentMap_RenderThread->IsEmpty() || !BackfillMap_RenderThread->IsEmpty())
+		{
+			OnAsyncConsolidationDisabled();
+		}
+
+		// Synchronously consolidate all caches.
+		ConsolidateThreadedCaches();
+		ProcessDelayedCleanup();
+
+		// call discard twice to clear both the backing and main caches
+		ReleasedEntries = DiscardAndSwap(CurrentMap, BackfillMap);
+		ReleasedEntries += DiscardAndSwap(CurrentMap, BackfillMap);
+
+		bDiscardAndSwap = false;
+	}
+
+	void WaitTasksComplete()
+	{
+		FScopeLock S(&AllThreadsLock);
+		
+		for (FPipelineStateCacheType* PipelineStateCache : AllThreadsPipelineStateCache)
+		{
+			WaitTasksComplete(PipelineStateCache);
+		}
+		
+		WaitTasksComplete(BackfillMap);
+		WaitTasksComplete(CurrentMap);
+
+		WaitTasksComplete(BackfillMap_RenderThread);
+		WaitTasksComplete(CurrentMap_RenderThread);
+	}
+
+	int32 NumReleasedEntries() const
+	{
+		return ReleasedEntries;
+	}
+
+private:
+
+	void WaitTasksComplete(FPipelineStateCacheType* PipelineStateCache)
+	{
+		FScopeLock S(&AllThreadsLock);
+		for (auto PipelineStateCacheIterator = PipelineStateCache->CreateIterator(); PipelineStateCacheIterator; ++PipelineStateCacheIterator)
+		{
+			auto PipelineState = PipelineStateCacheIterator->Value;
+			if (PipelineState != nullptr)
+			{
+				PipelineState->WaitCompletion();
+			}
+		}
+	}
+
+	void OnAsyncConsolidationEnabled()
+	{
+		*CurrentMap_RenderThread = *CurrentMap;
+		*BackfillMap_RenderThread = *BackfillMap;
+	}
+
+	void OnAsyncConsolidationDisabled()
+	{
+		// The render thread caches are the most up-to-date.
+		Swap(CurrentMap, CurrentMap_RenderThread);
+		Swap(BackfillMap, BackfillMap_RenderThread);
+		CurrentMap_RenderThread->Reset();
+		BackfillMap_RenderThread->Reset();
+
+		// New Render Thread pipeline states have already been
+		// consolidated into the Render Thread's maps.
+		NewRenderThreadPipelineStates.Reset();
+	}
+
+	void ConsolidateThreadedCaches()
+	{
+		SCOPE_TIME_GUARD_MS(TEXT("ConsolidatePipelineCache"), 0.1);
+		check(IsInRenderingThread());
 		
 		// consolidate all the local threads keys with the current thread
 		// No one is allowed to call GetLocalCache while this is running
 		// this is verified by the VerifyMutex.
-		for ( FPipelineStateCacheType* PipelineStateCache : AllThreadsPipelineStateCache)
+		for (FPipelineStateCacheType* PipelineStateCache : AllThreadsPipelineStateCache)
 		{
 			for (auto PipelineStateCacheIterator = PipelineStateCache->CreateIterator(); PipelineStateCacheIterator; ++PipelineStateCacheIterator)
 			{
@@ -1221,20 +1534,24 @@ public:
 				const TMyValue& ThreadValue = PipelineStateCacheIterator->Value;
 
 				{
-					BackfillMap->Remove(ThreadKey);
-
 					TMyValue* CurrentValue = CurrentMap->Find(ThreadKey);
 					if (CurrentValue)
 					{
+						check(!BackfillMap->Contains(ThreadKey));
+
 						// if two threads get from the backfill map then we might just be dealing with one pipelinestate, in which case we have already added it to the currentmap and don't need to do anything else
-						if ( *CurrentValue != ThreadValue )
+						if (*CurrentValue != ThreadValue)
 						{
+							// otherwise we need to discard the duplicate.
 							++DuplicateStateGenerated;
 							DeleteArray.Add(ThreadValue);
 						}
 					}
 					else
 					{
+						check(!BackfillMap->Contains(ThreadKey) || *BackfillMap->Find(ThreadKey) == ThreadValue);
+
+						BackfillMap->Remove(ThreadKey);
 						CurrentMap->Add(ThreadKey, ThreadValue);
 						Uncompleted.Add(TTuple<TMyKey, TMyValue>(ThreadKey, ThreadValue));
 					}
@@ -1255,8 +1572,218 @@ public:
 		}
 	}
 
+	void StartAsyncCacheConsolidation()
+	{
+		SCOPED_NAMED_EVENT(StartAsyncCacheConsolidation, FColor::Magenta);
+
+		// Create an event to signal when the RHI cache conslidation completes.
+		RHICompletionEvent = FGraphEvent::CreateGraphEvent();
+		RHICompletionEvent->SetDebugName(TEXT("AsyncCacheConsolidation"));
+
+		// Add the completion of the RHI cache consolidation as a prerequisite for the next RHI dispatch.
+		GRHICommandList.AddNextDispatchPrerequisite(RHICompletionEvent);
+
+		// Flush Render Thread local caches and consolidate them into a single map.
+		ConsolidatePipelineStates(NewRenderThreadPipelineStates, RenderThreadsPipelineStateCache);
+
+		// Add new Render Thread pipeline states to the consolidated maps for the Render Thread.
+		ConsolidateThreadCache(*CurrentMap_RenderThread, *BackfillMap_RenderThread, NewRenderThreadPipelineStates, true);
+
+		// Enqueue an RHI cache consolidation task to execute when the last RHI submit completes.
+		FFunctionGraphTask::CreateAndDispatchWhenReady(
+			[this]()
+			{
+				FTaskTagScope TaskTagScope(ETaskTag::EParallelRhiThread);
+
+				RHIAsyncCacheConsolidation();
+			},
+			TStatId(),
+			GRHICommandList.GetCompletionEvent(),
+			ENamedThreads::AnyHiPriThreadHiPriTask);
+	}
+
+	void RHIAsyncCacheConsolidation()
+	{
+		SCOPED_NAMED_EVENT(RHIAsyncCacheConsolidation, FColor::Purple);
+
+#if PIPELINESTATECACHE_VERIFYTHREADSAFE
+		FScopeVerifyDecrement S(VerifyMutex);
+#endif
+		
+		// Add new Render Thread pipeline states to the consolidated maps for the RHI Thread.
+		ConsolidateThreadCache(*CurrentMap, *BackfillMap, NewRenderThreadPipelineStates, false);
+
+		// New Render Thread pipeline states have already been consolidated on the Render Thread.
+		NewRenderThreadPipelineStates.Reset();
+
+		// Flush RHI Thread local caches and consolidate them into a single map.
+		ConsolidatePipelineStates(NewRHIThreadPipelineStates, RHIThreadsPipelineStateCache);
+
+		// Add new RHI Thread pipeline states to the consolidated maps for the RHI Thread.
+		ConsolidateThreadCache(*CurrentMap, *BackfillMap, NewRHIThreadPipelineStates, true);
+
+		// Check for completed tasks.
+		ManageIncompleteTasks();
+
+		if (bDiscardAndSwap)
+		{
+			// The Render Thread will discard the contents of the backfill map.
+			BackfillMap->Reset();
+
+			DiscardAndSwap(CurrentMap, BackfillMap);
+		}
+
+		// Signal that the RHI cache consolidation is complete.
+		RHICompletionEvent->DispatchSubsequents();
+	}
+
+	void FinishAsyncCacheConsolidation()
+	{
+		SCOPED_NAMED_EVENT(FinishAsyncCacheConsolidation, FColor::Orange);
+
+		// Add new RHI Thread pipeline states to the consolidated maps for the Render Thread.
+		ConsolidateThreadCache(*CurrentMap_RenderThread, *BackfillMap_RenderThread, NewRHIThreadPipelineStates, false);
+
+		// New RHI Thread pipeline states have already been consolidated on the RHI Thread.
+		NewRHIThreadPipelineStates.Reset();
+
+		// Flush Render Thread local caches and consolidate them into a single map.
+		ConsolidatePipelineStates(NewRenderThreadPipelineStates, RenderThreadsPipelineStateCache);
+
+		// Add new Render Thread pipeline states to the consolidated maps for the Render Thread.
+		ConsolidateThreadCache(*CurrentMap_RenderThread, *BackfillMap_RenderThread, NewRenderThreadPipelineStates, true);
+
+		// Check for completed tasks.
+		ManageCompleteTasks();
+
+		// Clean up duplicate tasks.
+		ProcessDelayedCleanup();
+
+		ReleasedEntries = 0;
+		if (bDiscardAndSwap)
+		{
+			ReleasedEntries = DiscardAndSwap(CurrentMap_RenderThread, BackfillMap_RenderThread);
+
+			bDiscardAndSwap = false;
+		}
+	}
+
+	void ConsolidatePipelineStates(FPipelineStateCacheType& PipelineStates, TArray<FPipelineStateCacheType*>& ThreadsPipelineStateCache)
+	{
+		SCOPE_TIME_GUARD_MS(TEXT("ConsolidatePipelineStateCache"), 0.1);
+
+		// Gather pipeline states generated in Render Thread tasks into a single map.
+		// No Render Thread task is allowed to call GetLocalCache while this is running
+		// this is verified by the VerifyMutex.
+		for (FPipelineStateCacheType* PipelineStateCache : ThreadsPipelineStateCache)
+		{
+			for (auto PipelineStateCacheIterator = PipelineStateCache->CreateIterator(); PipelineStateCacheIterator; ++PipelineStateCacheIterator)
+			{
+				const TMyKey& ThreadKey = PipelineStateCacheIterator->Key;
+				const TMyValue& ThreadValue = PipelineStateCacheIterator->Value;
+
+				{
+					TMyValue* CurrentValue = PipelineStates.Find(ThreadKey);
+					if (CurrentValue)
+					{
+						// if two threads get from the backfill map then we might just be dealing with one pipelinestate,
+						// in which case we have already added it to the map and don't need to do anything else
+						if (*CurrentValue != ThreadValue)
+						{
+							// otherwise we need to discard the duplicate.
+							++DuplicateStateGenerated;
+							DeleteArray.Add(ThreadValue);
+						}
+					}
+					else
+					{
+						PipelineStates.Add(ThreadKey, ThreadValue);
+					}
+
+					PipelineStateCacheIterator.RemoveCurrent();
+				}
+			}
+		}
+	}
+
+	void ConsolidateThreadCache(FPipelineStateCacheType& CurrentPipelineStateMap, FPipelineStateCacheType& BackfillPipelineStateMap, FPipelineStateCacheType& NewPipelineStates, bool bCacheNewTasks)
+	{
+		SCOPE_TIME_GUARD_MS(TEXT("ConsolidateThreadCache"), 0.1);
+
+		// consolidate all the new pipeline states with the state maps.
+		// No one is allowed to call Add or Find while this is running
+		// this is verified by the VerifyMutex.
+		{
+			for (auto PipelineStateCacheIterator = NewPipelineStates.CreateIterator(); PipelineStateCacheIterator; ++PipelineStateCacheIterator)
+			{
+				const TMyKey& ThreadKey = PipelineStateCacheIterator->Key;
+				const TMyValue& ThreadValue = PipelineStateCacheIterator->Value;
+
+				{
+					TMyValue* CurrentValue = CurrentPipelineStateMap.Find(ThreadKey);
+					if (CurrentValue)
+					{
+						check(!BackfillPipelineStateMap.Contains(ThreadKey));
+
+						// if two threads get from the backfill map then we might just be dealing with one pipelinestate, in which case we have already added it to the currentmap and don't need to do anything else
+						if (*CurrentValue != ThreadValue)
+						{
+							// otherwise we need to discard the duplicate.
+							++DuplicateStateGenerated;
+							DeleteArray.Add(ThreadValue);
+							PipelineStateCacheIterator.RemoveCurrent();
+						}
+					}
+					else
+					{
+						check(!BackfillPipelineStateMap.Contains(ThreadKey) || *BackfillPipelineStateMap.Find(ThreadKey) == ThreadValue);
+
+						CurrentPipelineStateMap.Add(ThreadKey, ThreadValue);
+						int32 Removed = BackfillPipelineStateMap.Remove(ThreadKey);
+						if (Removed == 0 && bCacheNewTasks)
+						{
+							Uncompleted.Add(*PipelineStateCacheIterator);
+						}
+					}
+				}
+			}
+		}
+	}
+
+	void ManageIncompleteTasks()
+	{
+		// tick and complete any uncompleted PSO tasks (we free up precompile tasks here).
+		for (int32 i = Uncompleted.Num() - 1; i >= 0; i--)
+		{
+			checkSlow(CurrentMap->Find(Uncompleted[i].Key));
+			if (Uncompleted[i].Value->IsComplete())
+			{
+				Uncompleted[i].Value->WaitCompletion();
+				Completed.Add(Uncompleted[i]); // WaitCompletion must also be called on the Render Thread to ensure the CompletionState is destroyed.
+				Uncompleted.RemoveAtSwap(i, EAllowShrinking::No);
+			}
+		}
+	}
+
+	void ManageCompleteTasks()
+	{
+		// tick Completed PSO tasks (we free up precompile tasks here).
+		for (int32 i = Completed.Num() - 1; i >= 0; i--)
+		{
+			checkSlow(CurrentMap_RenderThread->Find(Completed[i].Key));
+			checkSlow(Completed[i].Value->IsComplete());
+			Completed[i].Value->WaitCompletion();
+		}
+		Completed.Reset();
+	}
+
 	void ProcessDelayedCleanup()
 	{
+		if (DeleteArray.IsEmpty())
+		{
+			return;
+		}
+
 		FRHICommandListImmediate::Get().EnqueueLambda([DeleteArray = MoveTemp(DeleteArray)](FRHICommandListImmediate& RHICmdList) mutable
 		{
 			for (TMyValue& OldPipelineState : DeleteArray)
@@ -1270,104 +1797,91 @@ public:
 				UE_CLOG(bWaited, LogRHI, Log, TEXT("Waited on a pipeline compile task while discarding duplicate."));
 				delete OldPipelineState;
 			}
-			DeleteArray.Empty();
 		});
 	}
 
-	int32 DiscardAndSwap()
+	int32 DiscardAndSwap(FPipelineStateCacheType*& InOutCurrentMap, FPipelineStateCacheType*& InOutBackfillMap)
 	{
+		// This should be very fast, if not it's likely eviction time is too high and too 
+		// many items are building up.
+		SCOPE_TIME_GUARD_MS(TEXT("TrimPiplelineCache"), 0.1);
+
 		// the consolidate should always be run before the DiscardAndSwap.
 		// there should be no inuse pipeline states in the backfill map (because they should have been moved into the CurrentMap).
-		int32 Discarded = BackfillMap->Num();
-
-		FRHICommandListImmediate::Get().EnqueueLambda([BackfillMap = MoveTemp(*BackfillMap)](FRHICommandListImmediate& RHICmdList) mutable
+		int32 Discarded = InOutBackfillMap->Num();
+		if (Discarded > 0)
 		{
-			for (const auto& DiscardIterator :  BackfillMap)
+			FRHICommandListImmediate::Get().EnqueueLambda([DiscardMap = MoveTemp(*InOutBackfillMap)](FRHICommandListImmediate& RHICmdList) mutable
 			{
-				DiscardIterator.Value->Verify_NoUse();
-
-				// Incomplete tasks should be put back to the current map. There should be no incomplete tasks encountered here.
-				bool bWaited = DiscardIterator.Value->WaitCompletion();
-				UE_CLOG(bWaited, LogRHI, Error, TEXT("Waited on a pipeline compile task while discarding retired PSOs."));
-				delete DiscardIterator.Value;
-			}
-		});
+				for (const auto& DiscardIterator : DiscardMap)
+				{
+					DiscardIterator.Value->Verify_NoUse();
 
-		if ( CurrentMap == &Map1 )
-		{
-			CurrentMap = &Map2;
-			BackfillMap = &Map1;
-		}
-		else
-		{
-			CurrentMap = &Map1;
-			BackfillMap = &Map2;
+					// Incomplete tasks should be put back to the current map. There should be no incomplete tasks encountered here.
+					bool bWaited = DiscardIterator.Value->WaitCompletion();
+					UE_CLOG(bWaited, LogRHI, Error, TEXT("Waited on a pipeline compile task while discarding retired PSOs."));
+					delete DiscardIterator.Value;
+				}
+			});
 		}
 
+		Swap(InOutBackfillMap, InOutCurrentMap);
+
 		// keep alive incomplete tasks by moving them back to the current map.
 		for (int32 i = Uncompleted.Num() - 1; i >= 0; i--)
 		{
-			int32 Removed = BackfillMap->Remove(Uncompleted[i].Key);
+			int32 Removed = InOutBackfillMap->Remove(Uncompleted[i].Key);
 			checkSlow(Removed);
 			if (Removed)
 			{
-				CurrentMap->Add(Uncompleted[i].Key, Uncompleted[i].Value);
+				InOutCurrentMap->Add(Uncompleted[i].Key, Uncompleted[i].Value);
 			}
 		}
 
 		return Discarded;
 	}
-	
-	void WaitTasksComplete()
-	{
-		FScopeLock S(&AllThreadsLock);
-		
-		for ( FPipelineStateCacheType* PipelineStateCache : AllThreadsPipelineStateCache )
-		{
-			WaitTasksComplete(PipelineStateCache);
-		}
-		
-		WaitTasksComplete(BackfillMap);
-		WaitTasksComplete(CurrentMap);
-	}
 
 private:
-
-	void WaitTasksComplete(FPipelineStateCacheType* PipelineStateCache)
-	{
-		FScopeLock S(&AllThreadsLock);
-		for (auto PipelineStateCacheIterator = PipelineStateCache->CreateIterator(); PipelineStateCacheIterator; ++PipelineStateCacheIterator)
-		{
-			auto PipelineState = PipelineStateCacheIterator->Value;
-			if (PipelineState != nullptr)
-			{
-				PipelineState->WaitCompletion();
-			}
-		}
-	}
 	
-private:
-	uint32 TLSSlot;
-	FPipelineStateCacheType *CurrentMap;
-	FPipelineStateCacheType *BackfillMap;
+	// The list of tasks that are still in progress.
+	TArray<TTuple<TMyKey, TMyValue>> Uncompleted;
+	TArray<TTuple<TMyKey, TMyValue>> Completed;
+
+	uint32 TLSSlot{ FPlatformTLS::AllocTlsSlot() };
 
-	FPipelineStateCacheType Map1;
-	FPipelineStateCacheType Map2;
+	FPipelineStateCacheType NewRenderThreadPipelineStates;
+	FPipelineStateCacheType NewRHIThreadPipelineStates;
+
+	FPipelineStateCacheType Maps[4];
+
+	FPipelineStateCacheType* CurrentMap{ &Maps[0] };
+	FPipelineStateCacheType* BackfillMap{ &Maps[1] };
+
+	FPipelineStateCacheType* CurrentMap_RenderThread{ &Maps[2] };
+	FPipelineStateCacheType* BackfillMap_RenderThread{ &Maps[3] };
 
 	TArray<TMyValue> DeleteArray;
 
 	FCriticalSection AllThreadsLock;
 	TArray<FPipelineStateCacheType*> AllThreadsPipelineStateCache;
+	TArray<FPipelineStateCacheType*> RenderThreadsPipelineStateCache;
+	TArray<FPipelineStateCacheType*> RHIThreadsPipelineStateCache;
+
+	FGraphEventRef RHICompletionEvent;
+
+	int32 ReleasedEntries = 0;
+	uint32 DuplicateStateGenerated = 0;
+
+	bool bPendingDiscardAndSwap = false;
+	bool bDiscardAndSwap = false;
 
-	uint32 DuplicateStateGenerated;
 #if PSO_TRACK_CACHE_STATS
 	friend void DumpPipelineCacheStats();
 #endif
 
 #if PIPELINESTATECACHE_VERIFYTHREADSAFE
-	volatile int32 VerifyMutex;
+	volatile int32 VerifyMutex[2];
 #endif
-
 };
 
 // Request state
@@ -1411,13 +1925,13 @@ protected:
 
 		if (FPSOPrecacheThreadPool::UsePool())
 		{
-			check(PipelineState->PrecompileTask);
+			check(PipelineState->GetPrecompileTask());
 			EQueuedWorkPriority NewPriority = bHighestPriority ? EQueuedWorkPriority::Highest : EQueuedWorkPriority::High;
-			if(PipelineState->PrecompileTask)
+			if(PipelineState->GetPrecompileTask())
 			{
-				EQueuedWorkPriority PrevPriority = PipelineState->PrecompileTask->GetPriority();
+				EQueuedWorkPriority PrevPriority = PipelineState->GetPrecompileTask()->GetPriority();
 				check(PrevPriority > NewPriority);
-				PipelineState->PrecompileTask->Reschedule(&GPSOPrecacheThreadPool.Get(), NewPriority);
+				PipelineState->GetPrecompileTask()->Reschedule(&GPSOPrecacheThreadPool.Get(), NewPriority);
 			}
 		}
 
@@ -1476,8 +1990,8 @@ protected:
 				// Assign the event at this point because we need to release the lock before calling OnNewPipelineStateCreated which
 				// might call PrecacheFinished directly (The background task might get abandoned) and FRWLock can't be acquired recursively
 				// Note that calling IsComplete will return false until we link it somehow like we do below
-				NewPipelineState->CompletionEvent = FGraphEvent::CreateGraphEvent();
-				Result.AsyncCompileEvent = NewPipelineState->CompletionEvent;
+				NewPipelineState->SetCompletionEvent(FGraphEvent::CreateGraphEvent());
+				Result.AsyncCompileEvent = NewPipelineState->GetCompletionEvent();
 
 				UpdateActiveCompileCount(true /*Increment*/);
 			}
@@ -1683,7 +2197,7 @@ protected:
 			if (!IsCompilationDone(FindResult->ReadPSOPrecacheState()))
 			{
 				Result.RequestID = FindResult->RequestID;
-				Result.AsyncCompileEvent = FindResult->PipelineState->CompletionEvent;
+				Result.AsyncCompileEvent = FindResult->PipelineState->GetCompletionEvent();
 				check(Result.RequestID.IsValid());
 			}
 			return true;
@@ -2123,7 +2637,7 @@ int32 PipelineStateCache::GetNumActivePipelinePrecompileTasks()
 class FCompilePipelineStateTask
 {
 public:
-	FPipelineState* Pipeline;
+	FPipelineStateAsync* Pipeline;
 	FGraphicsPipelineStateInitializer Initializer;
 	EPSOPrecacheResult PSOPrecacheResult;
 	bool bInImmediateCmdList;
@@ -2131,7 +2645,7 @@ public:
 
 	// InInitializer is only used for non-compute tasks, a default can just be used otherwise
 	FCompilePipelineStateTask(
-		FPipelineState* InPipeline, 
+		FPipelineStateAsync* InPipeline, 
 		const FGraphicsPipelineStateInitializer& InInitializer, 
 		EPSOPrecacheResult InPSOPrecacheResult, 
 		bool InbInImmediateCmdList,
@@ -2142,7 +2656,7 @@ public:
 		, bInImmediateCmdList(InbInImmediateCmdList)
 		, PSOCompilationDebugData(InPSOCompilationDebugData)
 	{
-		ensure(Pipeline->CompletionEvent != nullptr);
+		ensure(Pipeline->GetCompletionEvent() != nullptr);
 		if(Initializer.bFromPSOFileCache)
 		{
 			GPipelinePrecompileTasksInFlight++;
@@ -2432,9 +2946,9 @@ public:
 #endif // WITH_RHI_BREADCRUMBS
 		
 		// We kicked a task: the event really should be there
-		if (ensure(Pipeline->CompletionEvent))
+		if (ensure(Pipeline->GetCompletionEvent()))
 		{
-			Pipeline->CompletionEvent->DispatchSubsequents();
+			Pipeline->GetCompletionEvent()->DispatchSubsequents();
 			// At this point, it's not safe to use Pipeline anymore, as it might get picked up by ProcessDelayedCleanup and deleted
 			Pipeline = nullptr;
 		}
@@ -2469,14 +2983,25 @@ void PipelineStateCache::FlushResources()
 {
 	check(IsInRenderingThread());
 
-	GComputePipelineCache.ConsolidateThreadedCaches();
-	GComputePipelineCache.ProcessDelayedCleanup();
+	static double LastEvictionTime = FPlatformTime::Seconds();
+	double CurrentTime = FPlatformTime::Seconds();
 
-	GWorkGraphPipelineCache.ConsolidateThreadedCaches();
-	GWorkGraphPipelineCache.ProcessDelayedCleanup();
+#if PSO_DO_CACHE_EVICT_EACH_FRAME
+	LastEvictionTime = 0;
+#endif
 
-	GGraphicsPipelineCache.ConsolidateThreadedCaches();
-	GGraphicsPipelineCache.ProcessDelayedCleanup();
+	// because it takes two cycles for an object to move from main->backfill->gone we check
+	// at half the desired eviction time
+	int32 EvictionPeriod = CVarPSOEvictionTime.GetValueOnAnyThread();
+	bool bDiscardAndSwap = !(EvictionPeriod == 0 || CurrentTime - LastEvictionTime < EvictionPeriod);
+	if (bDiscardAndSwap)
+	{
+		LastEvictionTime = CurrentTime;
+	}
+
+	GComputePipelineCache.FlushResources(bDiscardAndSwap);
+	GWorkGraphPipelineCache.FlushResources(bDiscardAndSwap);
+	GGraphicsPipelineCache.FlushResources(bDiscardAndSwap);
 
 	GPrecacheGraphicsPipelineCache.ProcessDelayedCleanup();
 	GPrecacheComputePipelineCache.ProcessDelayedCleanup();
@@ -2521,46 +3046,20 @@ void PipelineStateCache::FlushResources()
 	GraphicsPipelineCacheMisses = 0;
 	ComputePipelineCacheMisses = 0;
 
-	static double LastEvictionTime = FPlatformTime::Seconds();
-	double CurrentTime = FPlatformTime::Seconds();
-
-#if PSO_DO_CACHE_EVICT_EACH_FRAME
-	LastEvictionTime = 0;
-#endif
-	
-	// because it takes two cycles for an object to move from main->backfill->gone we check
-	// at half the desired eviction time
-	int32 EvictionPeriod = CVarPSOEvictionTime.GetValueOnAnyThread();
-
-	if (EvictionPeriod == 0 || CurrentTime - LastEvictionTime < EvictionPeriod)
-	{
-		return;
-	}
-
-	// This should be very fast, if not it's likely eviction time is too high and too 
-	// many items are building up.
-	SCOPE_TIME_GUARD_MS(TEXT("TrimPiplelineCache"), 0.1);
-
 #if PSO_TRACK_CACHE_STATS
 	DumpPipelineCacheStats();
-#endif
-
-	LastEvictionTime = CurrentTime;
-
-	int32 ReleasedComputeEntries = 0;
-	int32 ReleasedGraphicsEntries = 0;
-	int32 ReleasedWorkGraphEntries = 0;
 
-	ReleasedComputeEntries  =  GComputePipelineCache.DiscardAndSwap();
-	ReleasedGraphicsEntries = GGraphicsPipelineCache.DiscardAndSwap();
-	ReleasedWorkGraphEntries = GWorkGraphPipelineCache.DiscardAndSwap();
+	int32 ReleasedComputeEntries = GComputePipelineCache.NumReleasedEntries();
+	int32 ReleasedGraphicsEntries = GGraphicsPipelineCache.NumReleasedEntries();
+	int32 ReleasedWorkGraphEntries = GWorkGraphPipelineCache.NumReleasedEntries();
 
-#if PSO_TRACK_CACHE_STATS
-	UE_LOG(LogRHI, Log, TEXT("Cleared state cache in %.02f ms. %d ComputeEntries, %d GraphicsEntries, %d WorkGraphEntries")
-		, (FPlatformTime::Seconds() - CurrentTime) / 1000
-		, ReleasedComputeEntries, ReleasedGraphicsEntries, ReleasedWorkGraphEntries);
+	if (ReleasedComputeEntries > 0 || ReleasedGraphicsEntries > 0 || ReleasedWorkGraphEntries > 0)
+	{
+		UE_LOG(LogRHI, Log, TEXT("Cleared state cache in %.02f ms. %d ComputeEntries, %d GraphicsEntries, %d WorkGraphEntries")
+			, (FPlatformTime::Seconds() - CurrentTime) / 1000
+			, ReleasedComputeEntries, ReleasedGraphicsEntries, ReleasedWorkGraphEntries);
+	}
 #endif // PSO_TRACK_CACHE_STATS
-
 }
 
 static bool IsAsyncCompilationAllowed(FRHIComputeCommandList& RHICmdList, bool bIsPrecompileRequest)
@@ -2583,7 +3082,7 @@ uint64 PipelineStateCache::RetrieveGraphicsPipelineStateSortKey(const FGraphicsP
 
 static void InternalCreateComputePipelineState(FRHIComputeShader* ComputeShader, bool bDoAsyncCompile, bool bFromPSOFileCache, EPSOPrecacheResult PSOPrecacheResult, FComputePipelineState* CachedState, const FPSOCompilationDebugData& PSOCompilationDebugData, bool bInImmediateCmdList)
 {
-	FGraphEventRef GraphEvent = CachedState->CompletionEvent;
+	FGraphEventRef GraphEvent = CachedState->GetCompletionEvent();
 
 	// create a compilation task, or just do it now...
 	if (bDoAsyncCompile)
@@ -2631,7 +3130,7 @@ FComputePipelineState* PipelineStateCache::GetAndOrCreateComputePipelineState(FR
 		OutCachedState->Stats = FPipelineFileCacheManager::RegisterPSOStats(GetTypeHash(ComputeShader));
 		if (DoAsyncCompile)
 		{
-			OutCachedState->CompletionEvent = FGraphEvent::CreateGraphEvent();
+			OutCachedState->SetCompletionEvent(FGraphEvent::CreateGraphEvent());
 		}
 
 		if (!bFromFileCache)
@@ -2650,7 +3149,7 @@ FComputePipelineState* PipelineStateCache::GetAndOrCreateComputePipelineState(FR
 		PSOCompilationDebugData.BreadcrumbNode = DoAsyncCompile ? RHICmdList.GetCurrentBreadcrumbRef() : nullptr;
 #endif // WITH_RHI_BREADCRUMBS
 		
-		FGraphEventRef GraphEvent = OutCachedState->CompletionEvent;
+		FGraphEventRef GraphEvent = OutCachedState->GetCompletionEvent();
 		InternalCreateComputePipelineState(ComputeShader, DoAsyncCompile, bFromFileCache, PSOPrecacheResult, OutCachedState, PSOCompilationDebugData, RHICmdList.IsImmediate());
 
 		if (GraphEvent.IsValid())
@@ -2665,7 +3164,7 @@ FComputePipelineState* PipelineStateCache::GetAndOrCreateComputePipelineState(FR
 	{
 		if (!bFromFileCache && !OutCachedState->IsComplete())
 		{
-			RHICmdList.AddDispatchPrerequisite(OutCachedState->CompletionEvent);
+			RHICmdList.AddDispatchPrerequisite(OutCachedState->GetCompletionEvent());
 		}
 
 	#if PSO_TRACK_CACHE_STATS
@@ -2714,12 +3213,12 @@ public:
 
 	UE_NONCOPYABLE(FCompileRayTracingPipelineStateTask)
 
-	FPipelineState* Pipeline;
+	FPipelineStateSync* Pipeline;
 
 	FRayTracingPipelineStateInitializer Initializer;
 	const bool bBackgroundTask;
 
-	FCompileRayTracingPipelineStateTask(FPipelineState* InPipeline, const FRayTracingPipelineStateInitializer& InInitializer, bool bInBackgroundTask)
+	FCompileRayTracingPipelineStateTask(FPipelineStateSync* InPipeline, const FRayTracingPipelineStateInitializer& InInitializer, bool bInBackgroundTask)
 		: Pipeline(InPipeline)
 		, Initializer(InInitializer)
 		, bBackgroundTask(bInBackgroundTask)
@@ -3015,7 +3514,7 @@ inline void ValidateGraphicsPipelineStateInitializer(const FGraphicsPipelineStat
 
 static void InternalCreateGraphicsPipelineState(const FGraphicsPipelineStateInitializer& Initializer, EPSOPrecacheResult PSOPrecacheResult, bool bDoAsyncCompile, bool bPSOPrecache, FGraphicsPipelineState* CachedState, const FPSOCompilationDebugData& PSOCompilationDebugData, bool bInImmediateCmdList)
 {
-	FGraphEventRef GraphEvent = CachedState->CompletionEvent;
+	FGraphEventRef GraphEvent = CachedState->GetCompletionEvent();
 
 	// create a compilation task, or just do it now...
 	if (bDoAsyncCompile)
@@ -3031,7 +3530,7 @@ static void InternalCreateGraphicsPipelineState(const FGraphicsPipelineStateInit
 			// Here, PSO precompiles use a separate thread pool.
 			// Note that we do not add precompile tasks as cmdlist prerequisites.
 			TUniquePtr<FCompilePipelineStateTask> ThreadPoolTask = MakeUnique<FCompilePipelineStateTask>(CachedState, Initializer, PSOPrecacheResult, bInImmediateCmdList, PSOCompilationDebugData);
-			CachedState->PrecompileTask = MakeUnique<FPSOPrecacheAsyncTask>(
+			CachedState->SetPrecompileTask(MakeUnique<FPSOPrecacheAsyncTask>(
 				[ThreadPoolTask = MoveTemp(ThreadPoolTask)](const FPSOPrecacheAsyncTask* ThisTask)
 				{
 					// Convert the task priority to PSO precompile priority.
@@ -3055,8 +3554,8 @@ static void InternalCreateGraphicsPipelineState(const FGraphicsPipelineStateInit
 					}
 					ThreadPoolTask->CompilePSO(&PriOverride);
 				}
-			);
-			CachedState->PrecompileTask->StartBackgroundTask(&GPSOPrecacheThreadPool.Get(), Initializer.bFromPSOFileCache ? EQueuedWorkPriority::Normal : EQueuedWorkPriority::Low);
+			));
+			CachedState->GetPrecompileTask()->StartBackgroundTask(&GPSOPrecacheThreadPool.Get(), Initializer.bFromPSOFileCache ? EQueuedWorkPriority::Normal : EQueuedWorkPriority::Low);
 		}
 	}
 	else
@@ -3151,7 +3650,7 @@ FGraphicsPipelineState* PipelineStateCache::GetAndOrCreateGraphicsPipelineState(
 		OutCachedState->Stats = FPipelineFileCacheManager::RegisterPSOStats(GetTypeHash(Initializer));
 		if (DoAsyncCompile)
 		{
-			OutCachedState->CompletionEvent = FGraphEvent::CreateGraphEvent();
+			OutCachedState->SetCompletionEvent(FGraphEvent::CreateGraphEvent());
 		}
 
 		if (!Initializer.bFromPSOFileCache)
@@ -3171,7 +3670,7 @@ FGraphicsPipelineState* PipelineStateCache::GetAndOrCreateGraphicsPipelineState(
 #endif // WITH_RHI_BREADCRUMBS
 
 		bool bPSOPrecache = Initializer.bFromPSOFileCache;
-		FGraphEventRef GraphEvent = OutCachedState->CompletionEvent;
+		FGraphEventRef GraphEvent = OutCachedState->GetCompletionEvent();
 		InternalCreateGraphicsPipelineState(Initializer, PSOPrecacheResult, DoAsyncCompile, bPSOPrecache, OutCachedState, PSOCompilationDebugData, RHICmdList.IsImmediate());
 
 		// Add dispatch pre requisite for non precaching jobs only
@@ -3188,16 +3687,16 @@ FGraphicsPipelineState* PipelineStateCache::GetAndOrCreateGraphicsPipelineState(
 	{
 		if (!Initializer.bFromPSOFileCache && !OutCachedState->IsComplete())
 		{
-			if (OutCachedState->PrecompileTask)
+			if (OutCachedState->GetPrecompileTask())
 			{
 				// if this is an in-progress threadpool precompile task then it could be seconds away in the queue.
 				// Reissue this task so that it jumps the precompile queue.
-				OutCachedState->PrecompileTask->Reschedule(&GPSOPrecacheThreadPool.Get(), EQueuedWorkPriority::Highest);
+				OutCachedState->GetPrecompileTask()->Reschedule(&GPSOPrecacheThreadPool.Get(), EQueuedWorkPriority::Highest);
 #if PSO_TRACK_CACHE_STATS
 		UE_LOG(LogRHI, Log, TEXT("An incomplete precompile task was required for rendering!"));
 #endif
 			}
-			RHICmdList.AddDispatchPrerequisite(OutCachedState->CompletionEvent);
+			RHICmdList.AddDispatchPrerequisite(OutCachedState->GetCompletionEvent());
 		}
 
 #if PSO_TRACK_CACHE_STATS
@@ -3310,7 +3809,7 @@ void FPrecacheComputePipelineCache::OnNewPipelineStateCreated(const FPrecacheCom
 	// create a compilation task, or just do it now...
 	if (bDoAsyncCompile)
 	{
-		check(CachedState->CompletionEvent != nullptr);
+		check(CachedState->GetCompletionEvent() != nullptr);
 		FGraphicsPipelineStateInitializer GraphicsPipelineStateInitializer;
 		GraphicsPipelineStateInitializer.bPSOPrecache = true;
 		GraphicsPipelineStateInitializer.SetPSOPrecacheCompileType(FGraphicsPipelineStateInitializer::EPSOPrecacheCompileType::NormalPri);
@@ -3328,18 +3827,18 @@ void FPrecacheComputePipelineCache::OnNewPipelineStateCreated(const FPrecacheCom
 			// Here, PSO precompiles use a separate thread pool.
 			// Note that we do not add precompile tasks as cmdlist prerequisites.
 			TUniquePtr<FCompilePipelineStateTask> ThreadPoolTask = MakeUnique<FCompilePipelineStateTask>(CachedState, GraphicsPipelineStateInitializer, EPSOPrecacheResult::Active, false, PSOCompilationDebugData);
-			CachedState->PrecompileTask = MakeUnique<FPSOPrecacheAsyncTask>(
+			CachedState->SetPrecompileTask(MakeUnique<FPSOPrecacheAsyncTask>(
 				[ThreadPoolTask = MoveTemp(ThreadPoolTask)](const FPSOPrecacheAsyncTask* ThisTask)
 				{
 					ThreadPoolTask->CompilePSO();
 				}
-			);
-			CachedState->PrecompileTask->StartBackgroundTask(&GPSOPrecacheThreadPool.Get(), EQueuedWorkPriority::Normal);
+			));
+			CachedState->GetPrecompileTask()->StartBackgroundTask(&GPSOPrecacheThreadPool.Get(), EQueuedWorkPriority::Normal);
 		}
 	}
 	else
 	{
-		check(CachedState->CompletionEvent == nullptr);
+		check(CachedState->GetCompletionEvent() == nullptr);
 		CachedState->RHIPipeline = RHICreateComputePipelineState(CachedState->ComputeShader);
 		GPrecacheComputePipelineCache.PrecacheFinished(ComputeInitializer, CachedState->RHIPipeline != nullptr);
 	}
@@ -3372,7 +3871,7 @@ FPSOPrecacheRequestResult FPrecacheGraphicsPipelineCache::PrecacheGraphicsPipeli
 void FPrecacheGraphicsPipelineCache::OnNewPipelineStateCreated(const FGraphicsPipelineStateInitializer & Initializer, FGraphicsPipelineState * NewGraphicsPipelineState, const FString& PSOCompilationEventName, bool bDoAsyncCompile)
 {
 	ValidateGraphicsPipelineStateInitializer(Initializer);
-	check((NewGraphicsPipelineState->CompletionEvent != nullptr) == bDoAsyncCompile);
+	check((NewGraphicsPipelineState->GetCompletionEvent() != nullptr) == bDoAsyncCompile);
 
 	// Mark as precache so it will try and use the background thread pool if available
 	FGraphicsPipelineStateInitializer InitializerCopy(Initializer);
@@ -3620,13 +4119,10 @@ void PipelineStateCache::Shutdown()
 	GRayTracingPipelineCache.Shutdown();
 #endif
 
-	// call discard twice to clear both the backing and main caches
-	for (int32 i = 0; i < 2; i++)
-	{
-		GComputePipelineCache.DiscardAndSwap();
-		GWorkGraphPipelineCache.DiscardAndSwap();
-		GGraphicsPipelineCache.DiscardAndSwap();
-	}
+	GComputePipelineCache.Shutdown();
+	GWorkGraphPipelineCache.Shutdown();
+	GGraphicsPipelineCache.Shutdown();
+
 	FPipelineFileCacheManager::Shutdown();
 
 	for (auto Pair : GVertexDeclarationCache)
diff --git a/Engine/Source/Runtime/RHI/Private/RHICommandList.cpp b/Engine/Source/Runtime/RHI/Private/RHICommandList.cpp
--- a/Engine/Source/Runtime/RHI/Private/RHICommandList.cpp
+++ b/Engine/Source/Runtime/RHI/Private/RHICommandList.cpp
@@ -644,6 +644,12 @@ bool FRHICommandListExecutor::AllowParallel() const
 	return !Bypass() && IsRunningRHIInSeparateThread();
 }
 
+void FRHICommandListExecutor::AddNextDispatchPrerequisite(FGraphEventRef Prereq)
+{
+	check(IsInRenderingThread());
+	NextDispatchTaskPrerequisites.Add(MoveTemp(Prereq));
+}
+
 FRHICommandListExecutor::FTaskPipe* FRHICommandListExecutor::EnqueueDispatchTask(FGraphEventArray&& Prereqs, TFunction<void()>&& Lambda)
 {
 	check(IsInRenderingThread());
@@ -651,7 +657,10 @@ FRHICommandListExecutor::FTaskPipe* FRHICommandListExecutor::EnqueueDispatchTask
 		? ENamedThreads::GetRenderThread_Local()
 		: ENamedThreads::AnyHiPriThreadHiPriTask;
 
-	DispatchPipe.Enqueue(NamedThread, MoveTemp(Prereqs), MoveTemp(Lambda));
+	// Append any additional dispatch prerequisites.
+	NextDispatchTaskPrerequisites.Append(MoveTemp(Prereqs));
+
+	DispatchPipe.Enqueue(NamedThread, MoveTemp(NextDispatchTaskPrerequisites), MoveTemp(Lambda));
 	return &DispatchPipe;
 }
 
@@ -1533,16 +1542,16 @@ void FRHICommandListExecutor::WaitOnRHIThreadFence(FGraphEventRef& Fence)
 	}
 }
 
-void FRHICommandListExecutor::WaitForTasks()
+void FRHICommandListExecutor::WaitForTasks(FGraphEventArray& OutstandingTasks)
 {
 	check(IsInRenderingThread());
 
-	if (WaitOutstandingTasks.Num())
+	if (OutstandingTasks.Num())
 	{
 		bool bAny = false;
-		for (int32 Index = 0; Index < WaitOutstandingTasks.Num(); Index++)
+		for (int32 Index = 0; Index < OutstandingTasks.Num(); Index++)
 		{
-			if (!WaitOutstandingTasks[Index]->IsComplete())
+			if (!OutstandingTasks[Index]->IsComplete())
 			{
 				bAny = true;
 				break;
@@ -1554,10 +1563,10 @@ void FRHICommandListExecutor::WaitForTasks()
 			SCOPE_CYCLE_COUNTER(STAT_ExplicitWait);
 			ENamedThreads::Type RenderThread_Local = ENamedThreads::GetRenderThread_Local();
 			check(!FTaskGraphInterface::Get().IsThreadProcessingTasks(RenderThread_Local));
-			FTaskGraphInterface::Get().WaitUntilTasksComplete(WaitOutstandingTasks, RenderThread_Local);
+			FTaskGraphInterface::Get().WaitUntilTasksComplete(OutstandingTasks, RenderThread_Local);
 		}
 
-		WaitOutstandingTasks.Reset();
+		OutstandingTasks.Reset();
 	}
 }
 
diff --git a/Engine/Source/Runtime/RHI/Public/RHICommandList.h b/Engine/Source/Runtime/RHI/Public/RHICommandList.h
--- a/Engine/Source/Runtime/RHI/Public/RHICommandList.h
+++ b/Engine/Source/Runtime/RHI/Public/RHICommandList.h
@@ -5185,7 +5185,15 @@ public:
 	//
 	// Blocks the calling thread until all dispatch prerequisites of enqueued parallel command lists are completed.
 	//
-	RHI_API void WaitForTasks();
+	RHI_API void WaitForTasks()
+	{
+		WaitForTasks(WaitOutstandingTasks);
+	}
+
+	//
+	// Blocks the calling thread until all specified tasks are completed.
+	//
+	RHI_API void WaitForTasks(FGraphEventArray& OutstandingTasks);
 
 	//
 	// Blocks the calling thread until the RHI thread is idle.
@@ -5237,6 +5245,23 @@ public:
 		return !AreRHITasksActive();
 	}
 
+	//
+	// Adds a prerequisite for subsequent Submit dispatch tasks.
+	// 
+	// This function should only be called from the render thread.
+	//
+	RHI_API void AddNextDispatchPrerequisite(FGraphEventRef Prereq);
+
+	//
+	// Gets the CompletionEvent for the most recent submit to GPU
+	// 
+	// This function should only be called from the render thread.
+	//
+	const FGraphEventRef& GetCompletionEvent() const
+	{
+		return CompletionEvent;
+	}
+
 	FGraphEventArray WaitOutstandingTasks;
 
 private:
@@ -5353,6 +5378,7 @@ private:
 
 	FTaskPipe* EnqueueDispatchTask(FGraphEventArray&& Prereqs, TFunction<void()>&& Lambda);
 	FTaskPipe* EnqueueSubmitTask  (FGraphEventArray&& Prereqs, TFunction<void()>&& Lambda);
+	FGraphEventArray NextDispatchTaskPrerequisites;
 
 #if WITH_RHI_BREADCRUMBS
 
diff --git a/Engine/Source/Runtime/RenderCore/Private/RenderGraphBuilder.cpp b/Engine/Source/Runtime/RenderCore/Private/RenderGraphBuilder.cpp
--- a/Engine/Source/Runtime/RenderCore/Private/RenderGraphBuilder.cpp
+++ b/Engine/Source/Runtime/RenderCore/Private/RenderGraphBuilder.cpp
@@ -336,7 +336,18 @@ void FRDGBuilder::BeginFlushResourcesRHI()
 
 	CSV_SCOPED_TIMING_STAT_EXCLUSIVE(STAT_RDG_FlushResourcesRHI);
 	SCOPED_NAMED_EVENT(BeginFlushResourcesRHI, FColor::Emerald);
-	RHICmdList.ImmediateFlush(EImmediateFlushType::DispatchToRHIThread);
+
+	static const auto CVarEnablePSOAsyncCacheConsolidation = IConsoleManager::Get().FindConsoleVariable(TEXT("r.pso.EnableAsyncCacheConsolidation"));
+	if (CVarEnablePSOAsyncCacheConsolidation->GetBool())
+	{
+		// Cache prior tasks before enqueuing setup tasks, which can run while the pipeline state cache flushes.
+		WaitOutstandingTasks = GRHICommandList.WaitOutstandingTasks;
+	}
+	else
+	{
+		// Dispatch to RHI thread if cache consolidation is not asynchronous, so it can get some work started before blocking in EndFlushResourcesRHI.
+		RHICmdList.ImmediateFlush(EImmediateFlushType::DispatchToRHIThread);
+	}
 }
 
 void FRDGBuilder::EndFlushResourcesRHI()
@@ -349,7 +360,23 @@ void FRDGBuilder::EndFlushResourcesRHI()
 	CSV_SCOPED_TIMING_STAT_EXCLUSIVE(STAT_RDG_FlushResourcesRHI);
 	CSV_SCOPED_SET_WAIT_STAT(FlushResourcesRHI);
 	SCOPED_NAMED_EVENT(EndFlushResourcesRHI, FColor::Emerald);
-	RHICmdList.ImmediateFlush(EImmediateFlushType::FlushRHIThreadFlushResources);
+
+	static const auto CVarEnablePSOAsyncCacheConsolidation = IConsoleManager::Get().FindConsoleVariable(TEXT("r.pso.EnableAsyncCacheConsolidation"));
+	if (CVarEnablePSOAsyncCacheConsolidation->GetBool())
+	{
+		// Dispatch to RHI thread and delete resources.
+		RHICmdList.ImmediateFlush(EImmediateFlushType::DispatchToRHIThread, ERHISubmitFlags::DeleteResources);
+
+		// Wait for tasks cached in BeginFlushResourcesRHI.
+		GRHICommandList.WaitForTasks(WaitOutstandingTasks);
+	}
+	else
+	{
+		// Wait until all RHI work is complete.
+		RHICmdList.ImmediateFlush(EImmediateFlushType::FlushRHIThreadFlushResources);
+	}
+
+	// Flush the pipeline state cache.
 	PipelineStateCache::FlushResources();
 }
 
diff --git a/Engine/Source/Runtime/RenderCore/Public/RenderGraphBuilder.h b/Engine/Source/Runtime/RenderCore/Public/RenderGraphBuilder.h
--- a/Engine/Source/Runtime/RenderCore/Public/RenderGraphBuilder.h
+++ b/Engine/Source/Runtime/RenderCore/Public/RenderGraphBuilder.h
@@ -1084,6 +1084,7 @@ private:
 	/////////////////////////////////////////////////////////////////////////////
 	// Resource Deletion Flushing
 
+	FGraphEventArray WaitOutstandingTasks;
 	bool bFlushResourcesRHI = false;
 	FRHICommandListScopedExtendResourceLifetime ExtendResourceLifetimeScope;
 
