Fix verify mutex scope.

diff --git a/Engine/Source/Runtime/RHI/Private/PipelineStateCache.cpp b/Engine/Source/Runtime/RHI/Private/PipelineStateCache.cpp
--- a/Engine/Source/Runtime/RHI/Private/PipelineStateCache.cpp
+++ b/Engine/Source/Runtime/RHI/Private/PipelineStateCache.cpp
@@ -1605,32 +1605,33 @@ private:
 	void RHIAsyncCacheConsolidation()
 	{
 		SCOPED_NAMED_EVENT(RHIAsyncCacheConsolidation, FColor::Purple);
-
+		{
 #if PIPELINESTATECACHE_VERIFYTHREADSAFE
-		FScopeVerifyDecrement S(VerifyMutex);
+			FScopeVerifyDecrement S(VerifyMutex);
 #endif
-		
-		// Add new Render Thread pipeline states to the consolidated maps for the RHI Thread.
-		ConsolidateThreadCache(*CurrentMap, *BackfillMap, NewRenderThreadPipelineStates, false);
 
-		// New Render Thread pipeline states have already been consolidated on the Render Thread.
-		NewRenderThreadPipelineStates.Reset();
+			// Add new Render Thread pipeline states to the consolidated maps for the RHI Thread.
+			ConsolidateThreadCache(*CurrentMap, *BackfillMap, NewRenderThreadPipelineStates, false);
 
-		// Flush RHI Thread local caches and consolidate them into a single map.
-		ConsolidatePipelineStates(NewRHIThreadPipelineStates, RHIThreadsPipelineStateCache);
+			// New Render Thread pipeline states have already been consolidated on the Render Thread.
+			NewRenderThreadPipelineStates.Reset();
 
-		// Add new RHI Thread pipeline states to the consolidated maps for the RHI Thread.
-		ConsolidateThreadCache(*CurrentMap, *BackfillMap, NewRHIThreadPipelineStates, true);
+			// Flush RHI Thread local caches and consolidate them into a single map.
+			ConsolidatePipelineStates(NewRHIThreadPipelineStates, RHIThreadsPipelineStateCache);
 
-		// Check for completed tasks.
-		ManageIncompleteTasks();
+			// Add new RHI Thread pipeline states to the consolidated maps for the RHI Thread.
+			ConsolidateThreadCache(*CurrentMap, *BackfillMap, NewRHIThreadPipelineStates, true);
 
-		if (bDiscardAndSwap)
-		{
-			// The Render Thread will discard the contents of the backfill map.
-			BackfillMap->Reset();
+			// Check for completed tasks.
+			ManageIncompleteTasks();
+
+			if (bDiscardAndSwap)
+			{
+				// The Render Thread will discard the contents of the backfill map.
+				BackfillMap->Reset();
 
-			DiscardAndSwap(CurrentMap, BackfillMap);
+				DiscardAndSwap(CurrentMap, BackfillMap);
+			}
 		}
 
 		// Signal that the RHI cache consolidation is complete.
