Add vendor-aware PSO precaching and the ability to keep PSOs in memory.

This change:

* Introduces the concept of "vendor-aware" precaching with the ability to have the precaching system only precache PSOs that are considered "different" by the graphics driver. Every GPU vendor has a different subset of PSO state that can cause a driver cache hit or miss. This change implements this mechanism for NVIDIA and Intel Arc where testing was performed and assumptions verified. While avoiding the compilation of similar PSOs might seem unimportant (since by definition similar PSOs will cause a cache hit and will be fast), there is significant scheduling overhead and contention that can be avoided. According to my tests, this reduces the number of precached PSOs on NVIDIA by ~40% and on Intel by ~25%.

* Adds the ability to keep the precached PSOs in memory instead of deleting them immediately after creation. This can significantly help PSO re-creation performance on NVIDIA, and avoids otherwise-unavoidable small hitches on the critical path when the PSO is used for rendering. Precached PSOs are deleted once they are actually used for rendering, and CVars are provided to control how many are kept in memory at a time (to take into account all the precached PSOs that don't end up being used and avoid unbounded memory usage). According to my tests on a replay, this reduces the number of >4ms PSO creation hitches from over a thousand to less than 50.

* Moves initialization of the precaching structures to be explicitly controlled instead of relying on static initialization since that's too early for CVars to be set.

diff --git a/Engine/Source/Runtime/D3D12RHI/Private/D3D12RHIPrivate.h b/Engine/Source/Runtime/D3D12RHI/Private/D3D12RHIPrivate.h
--- a/Engine/Source/Runtime/D3D12RHI/Private/D3D12RHIPrivate.h
+++ b/Engine/Source/Runtime/D3D12RHI/Private/D3D12RHIPrivate.h
@@ -373,6 +373,7 @@ public:
 	virtual void RHIAliasTextureResources(FTextureRHIRef& DestTexture, FTextureRHIRef& SrcTexture) final override;
 	virtual FTextureRHIRef RHICreateAliasedTexture(FTextureRHIRef& SourceTexture) final override;
 	virtual void RHIGetDisplaysInformation(FDisplayInformationArray& OutDisplayInformation) final override;
+	virtual uint64 RHIComputeStatePrecachePSOHash(const FGraphicsPipelineStateInitializer& Initializer) final override;
 	virtual uint64 RHIComputePrecachePSOHash(const FGraphicsPipelineStateInitializer& Initializer) final override;
 	virtual bool RHIMatchPrecachePSOInitializers(const FGraphicsPipelineStateInitializer& LHS, const FGraphicsPipelineStateInitializer& RHS) final override;
 	virtual void RHIAdvanceFrameForGetViewportBackBuffer(FRHIViewport* Viewport) final override;
diff --git a/Engine/Source/Runtime/D3D12RHI/Private/D3D12State.cpp b/Engine/Source/Runtime/D3D12RHI/Private/D3D12State.cpp
--- a/Engine/Source/Runtime/D3D12RHI/Private/D3D12State.cpp
+++ b/Engine/Source/Runtime/D3D12RHI/Private/D3D12State.cpp
@@ -8,6 +8,14 @@
 #include "D3D12RHIPrivate.h"
 #include "RHIUtilities.h"
 
+static int32 GPSOPrecacheD3D12DriverCacheAware = 0;
+static FAutoConsoleVariableRef CVarPSOPrecacheD3D12DriverCacheAware(
+	TEXT("r.PSOPrecache.D3D12.DriverCacheAware"),
+	GPSOPrecacheD3D12DriverCacheAware,
+	TEXT("If enabled, the PSO precaching system will not precache PSOs that the D3D12 graphics driver considers similar for caching, i.e. it will not precache PSOs that while technically different will still result in a driver cache hit.\n")
+	TEXT("This is not implemented for all GPU vendors and can result in performance issues or cache misses if the heuristics the engine uses does not match the graphics driver's behavior that decides whether a PSO is in the cache or not."),
+	ECVF_ReadOnly);
+
 // MSFT: Need to make sure sampler state is thread safe
 // Cache of Sampler States; we store pointers to both as we don't want the TMap to be artificially
 // modifying ref counts if not needed; so we manage that ourselves
@@ -488,11 +496,125 @@ bool FD3D12BlendState::GetInitializer(class FBlendStateInitializerRHI& Init)
 	return true;
 }
 
+uint64 FD3D12DynamicRHI::RHIComputeStatePrecachePSOHash(const FGraphicsPipelineStateInitializer& Initializer)
+{
+	if (GPSOPrecacheD3D12DriverCacheAware)
+	{
+		if (IsRHIDeviceNVIDIA() && GMaxRHIFeatureLevel == ERHIFeatureLevel::SM6)
+		{
+			// NVIDIA drivers only care about the shaders for PSO caching.
+			struct FHashKey
+			{
+				uint32 VertexShader;
+				uint32 PixelShader;
+#if PLATFORM_SUPPORTS_GEOMETRY_SHADERS
+				uint32 GeometryShader;
+#endif // PLATFORM_SUPPORTS_GEOMETRY_SHADERS
+#if PLATFORM_SUPPORTS_MESH_SHADERS
+				uint32 MeshShader;
+#endif // PLATFORM_SUPPORTS_MESH_SHADERS
+			} HashKey;
+
+			FMemory::Memzero(&HashKey, sizeof(FHashKey));
+
+			HashKey.VertexShader = Initializer.BoundShaderState.GetVertexShader() ? GetTypeHash(Initializer.BoundShaderState.GetVertexShader()->GetHash()) : 0;
+			HashKey.PixelShader = Initializer.BoundShaderState.GetPixelShader() ? GetTypeHash(Initializer.BoundShaderState.GetPixelShader()->GetHash()) : 0;
+#if PLATFORM_SUPPORTS_GEOMETRY_SHADERS
+			HashKey.GeometryShader = Initializer.BoundShaderState.GetGeometryShader() ? GetTypeHash(Initializer.BoundShaderState.GetGeometryShader()->GetHash()) : 0;
+#endif
+#if PLATFORM_SUPPORTS_MESH_SHADERS
+			HashKey.MeshShader = Initializer.BoundShaderState.GetMeshShader() ? GetTypeHash(Initializer.BoundShaderState.GetMeshShader()->GetHash()) : 0;
+#endif
+			return CityHash64((const char*)&HashKey, sizeof(FHashKey));
+		}
+		else if (IsRHIDeviceIntel() && GMaxRHIFeatureLevel == ERHIFeatureLevel::SM6)
+		{
+			// Intel drivers have a few elements on top of the shaders that can cause PSO recompilation.
+			struct FHashKey
+			{
+				uint32 VertexShader;
+				uint32 PixelShader;
+#if PLATFORM_SUPPORTS_GEOMETRY_SHADERS
+				uint32 GeometryShader;
+#endif // PLATFORM_SUPPORTS_GEOMETRY_SHADERS
+#if PLATFORM_SUPPORTS_MESH_SHADERS
+				uint32 MeshShader;
+#endif // PLATFORM_SUPPORTS_MESH_SHADERS
+
+				uint32 BlendState;
+				uint8  MultisamplingEnabled;
+			} HashKey;
+
+			FMemory::Memzero(&HashKey, sizeof(FHashKey));
+
+			HashKey.VertexShader = Initializer.BoundShaderState.GetVertexShader() ? GetTypeHash(Initializer.BoundShaderState.GetVertexShader()->GetHash()) : 0;
+			HashKey.PixelShader = Initializer.BoundShaderState.GetPixelShader() ? GetTypeHash(Initializer.BoundShaderState.GetPixelShader()->GetHash()) : 0;
+#if PLATFORM_SUPPORTS_GEOMETRY_SHADERS
+			HashKey.GeometryShader = Initializer.BoundShaderState.GetGeometryShader() ? GetTypeHash(Initializer.BoundShaderState.GetGeometryShader()->GetHash()) : 0;
+#endif
+#if PLATFORM_SUPPORTS_MESH_SHADERS
+			HashKey.MeshShader = Initializer.BoundShaderState.GetMeshShader() ? GetTypeHash(Initializer.BoundShaderState.GetMeshShader()->GetHash()) : 0;
+#endif
+			FBlendStateInitializerRHI BlendStateInitializerRHI;
+			if (Initializer.BlendState && Initializer.BlendState->GetInitializer(BlendStateInitializerRHI))
+			{
+				HashKey.BlendState = GetTypeHash(BlendStateInitializerRHI);
+			}
+
+			// The only rasterizer state that matters is whether multisampling is enabled.
+			FRasterizerStateInitializerRHI RasterizerStateInitializerRHI;
+			if (Initializer.RasterizerState && Initializer.RasterizerState->GetInitializer(RasterizerStateInitializerRHI))
+			{
+				HashKey.MultisamplingEnabled = RasterizerStateInitializerRHI.bAllowMSAA;
+			}
+
+			return CityHash64((const char*)&HashKey, sizeof(FHashKey));
+		}
+	}
+
+	return FDynamicRHI::RHIComputeStatePrecachePSOHash(Initializer);
+}
+
 uint64 FD3D12DynamicRHI::RHIComputePrecachePSOHash(const FGraphicsPipelineStateInitializer& Initializer)
 {
 	// When compute precache PSO hash we assume a valid state precache PSO hash is already provided
 	checkf(Initializer.StatePrecachePSOHash != 0, TEXT("Initializer should have a valid state precache PSO hash set when computing the full initializer PSO hash"));
 
+	if (GPSOPrecacheD3D12DriverCacheAware)
+	{
+		if (IsRHIDeviceNVIDIA() && GMaxRHIFeatureLevel == ERHIFeatureLevel::SM6)
+		{
+			// We already hashed everything we needed.
+			return Initializer.StatePrecachePSOHash;
+		}
+		else if (IsRHIDeviceIntel() && GMaxRHIFeatureLevel == ERHIFeatureLevel::SM6)
+		{
+			// On top of the state already hashed, Intel drivers care about multisampling and render target count/format.
+			struct FHashKey
+			{
+				uint64 StatePrecachePSOHash;
+
+				uint8 NumSamples;
+				uint8 NumRenderTargets;
+				FGraphicsPipelineStateInitializer::TRenderTargetFormats	RenderTargetFormats;
+				FGraphicsPipelineStateInitializer::TRenderTargetFlags   RenderTargetFlags;
+			} HashKey;
+
+			FMemory::Memzero(&HashKey, sizeof(FHashKey));
+
+			HashKey.StatePrecachePSOHash = Initializer.StatePrecachePSOHash;
+			HashKey.NumSamples = Initializer.NumSamples;
+			HashKey.NumRenderTargets = Initializer.RenderTargetsEnabled;
+			HashKey.RenderTargetFormats = Initializer.RenderTargetFormats;
+			for (uint32 Index = 0; Index < HashKey.NumRenderTargets; ++Index)
+			{
+				HashKey.RenderTargetFlags[Index] = Initializer.RenderTargetFlags[Index] & FGraphicsPipelineStateInitializer::RelevantRenderTargetFlagMask;
+			}
+
+			return CityHash64((const char*)&HashKey, sizeof(FHashKey));
+		}
+	}
+
 	// All members which are not part of the state objects and influence the PSO on D3D12
 	struct FNonStateHashKey
 	{
diff --git a/Engine/Source/Runtime/Launch/Private/LaunchEngineLoop.cpp b/Engine/Source/Runtime/Launch/Private/LaunchEngineLoop.cpp
--- a/Engine/Source/Runtime/Launch/Private/LaunchEngineLoop.cpp
+++ b/Engine/Source/Runtime/Launch/Private/LaunchEngineLoop.cpp
@@ -3275,6 +3275,11 @@ int32 FEngineLoop::PreInitPreStartupScreen(const TCHAR* CmdLine)
 		RHIInit(bHasEditorToken);
 	}
 
+	{
+		SCOPED_BOOT_TIMING("PipelineStateCacheInit");
+		PipelineStateCache::Init();
+	}
+
 	{
 		UE_SCOPED_ENGINE_ACTIVITY(TEXT("Initializing Render Settings"));
 		SCOPED_BOOT_TIMING("RenderUtilsInit");
diff --git a/Engine/Source/Runtime/RHI/Private/PipelineStateCache.cpp b/Engine/Source/Runtime/RHI/Private/PipelineStateCache.cpp
--- a/Engine/Source/Runtime/RHI/Private/PipelineStateCache.cpp
+++ b/Engine/Source/Runtime/RHI/Private/PipelineStateCache.cpp
@@ -54,6 +54,7 @@ DECLARE_DWORD_ACCUMULATOR_STAT(TEXT("High Priority Compute PSO Precache Requests
 
 DECLARE_DWORD_ACCUMULATOR_STAT(TEXT("Highest Priority Graphics PSO Precache Requests"), STAT_HighestPriorityGraphicsPSOPrecacheRequests, STATGROUP_PipelineStateCache);
 DECLARE_DWORD_ACCUMULATOR_STAT(TEXT("Highest Priority Compute PSO Precache Requests"), STAT_HighestPriorityComputePSOPrecacheRequests, STATGROUP_PipelineStateCache);
+DECLARE_DWORD_ACCUMULATOR_STAT(TEXT("Precached PSOs Kept In Memory"), STAT_InMemoryPrecachedPSOCount, STATGROUP_PipelineStateCache);
 
 static inline uint32 GetTypeHash(const FBoundShaderStateInput& Input)
 {
@@ -313,6 +314,31 @@ static FAutoConsoleVariableRef CVarPrecompileThreadPoolThreadPriority(
 	TEXT("Thread priority for the PSO precompile pool"),
 	ECVF_RenderThreadSafe);
 
+int32 GPSOPrecacheKeepInMemoryUntilUsed = 0;
+static FAutoConsoleVariableRef CVarPSOPrecacheKeepInMemoryUntilUsed(
+	TEXT("r.PSOPrecache.KeepInMemoryUntilUsed"),
+	GPSOPrecacheKeepInMemoryUntilUsed,
+	TEXT("If enabled, precached PSOs will be kept in memory instead of being deleted immediately after creation, and will only be deleted once they are actually used for rendering.\n")
+	TEXT("This can speed up the re-creation of precached PSOs for certain drivers and avoid small hitches, at the cost of memory.\n")
+	TEXT("It's recommended to set r.PSOPrecache.KeepInMemoryGraphicsMaxNum and r.PSOPrecache.KeepInMemoryComputeMaxNum to a non-zero value to ensure the number of in-memory PSOs is bounded."),
+	ECVF_ReadOnly);
+
+int32 GPSOPrecacheKeepInMemoryGraphicsMaxNum = 2000;
+static FAutoConsoleVariableRef CVarPSOPrecacheKeepInMemoryGraphicsMaxNum(
+	TEXT("r.PSOPrecache.KeepInMemoryGraphicsMaxNum"),
+	GPSOPrecacheKeepInMemoryGraphicsMaxNum,
+	TEXT("If r.PSOPrecache.KeepInMemoryUntilUsed is enabled, this value will control the maximum number of precached graphics PSOs that are kept in memory at a time.\n")
+	TEXT("If set to 0, no limit will be applied (not recommended outside of testing, as it can cause unbounded memory usage)."),
+	ECVF_ReadOnly);
+
+int32 GPSOPrecacheKeepInMemoryComputeMaxNum = 200;
+static FAutoConsoleVariableRef CVarPSOPrecacheKeepInMemoryComputeMaxNum(
+	TEXT("r.PSOPrecache.KeepInMemoryComputeMaxNum"),
+	GPSOPrecacheKeepInMemoryComputeMaxNum,
+	TEXT("If r.PSOPrecache.KeepInMemoryUntilUsed is enabled, this value will control the maximum number of precached compute PSOs that are kept in memory at a time.\n")
+	TEXT("If set to 0, no limit will be applied (not recommended outside of testing, as it can cause unbounded memory usage)."),
+	ECVF_ReadOnly);
+
 const TCHAR* LexToString(EPSOPrecacheResult Result)
 {
 	switch (Result)
@@ -1894,6 +1920,7 @@ enum class EPSOPrecacheStateMask : uint8
 	Failed = 1 << 2,    // once compilation is finished
 	Boosted = 1 << 3,
 	HighestPri = 1 << 4,
+	UsedForRendering = 1 << 5, // Set once this precached PSO is actually needed and used for rendering.
 };
 
 ENUM_CLASS_FLAGS(EPSOPrecacheStateMask)
@@ -1903,12 +1930,19 @@ template<class TPrecachePipelineCacheDerived, class TPrecachedPSOInitializer, cl
 class TPrecachePipelineCacheBase
 {
 public:
-
-	TPrecachePipelineCacheBase(FPSOPrecacheRequestID::EType InType) : PSOType(InType) {}
+	TPrecachePipelineCacheBase(FPSOPrecacheRequestID::EType InType, uint32 InMaxInMemoryPSOs)
+		: PSOType(InType),
+		MaxInMemoryPSOs(InMaxInMemoryPSOs)
+	{
+		if (GPSOPrecacheKeepInMemoryUntilUsed && InMaxInMemoryPSOs > 0)
+		{
+			InMemoryPSOIndices.Init(kInvalidPSOIndex, InMaxInMemoryPSOs);
+		}
+	}
 
 	~TPrecachePipelineCacheBase()
 	{
-		// Wait for all precache tasks to finished
+		// Wait for all precache tasks to have finished.
 		WaitTasksComplete();
 	}
 	
@@ -1996,6 +2030,29 @@ protected:
 
 				UpdateActiveCompileCount(true /*Increment*/);
 			}
+
+			if (GPSOPrecacheKeepInMemoryUntilUsed)
+			{
+				if (MaxInMemoryPSOs > 0)
+				{
+					uint32& PSOIndex = InMemoryPSOIndices[CurrentInMemoryPSOIndex];
+
+					// Check if we're evicting a valid index. In that case, enqueue the corresponding PSO for cleanup.
+					if (PSOIndex != kInvalidPSOIndex)
+					{
+						PrecachedPSOsToCleanup.Add(PrecachedPSOInitializers[PSOIndex]);
+					}
+					PSOIndex = Result.RequestID.RequestID;
+
+					// Manage the circular buffer index.
+					CurrentInMemoryPSOIndex++;
+					if (CurrentInMemoryPSOIndex >= MaxInMemoryPSOs)
+					{
+						CurrentInMemoryPSOIndex = 0;
+					}
+				}
+				INC_DWORD_STAT(STAT_InMemoryPrecachedPSOCount);
+			}
 		}
 
 		TPrecachePipelineCacheDerived::OnNewPipelineStateCreated(Initializer, NewPipelineState, PSOCompilationEventName, bDoAsyncCompile);
@@ -2051,7 +2108,7 @@ public:
 				}
 				if (!bTasksWaiting)
 				{
-					PrecachedPSOs.Empty();
+					PrecachedPSOsToCleanup.Empty();
 				}
 			}
 			if (bTasksWaiting)
@@ -2137,7 +2194,10 @@ public:
 			PreviousStateMask = FindResult->AddPSOPrecacheState(CompleteStateMask);
 
 			// Add to array of precached PSOs so it can be cleaned up
-			PrecachedPSOs.Add(InitializerHash);
+			if (!GPSOPrecacheKeepInMemoryUntilUsed)
+			{
+				PrecachedPSOsToCleanup.Add(InitializerHash);
+			}
 		}
 
         // Need to ensure that the boost request was actually executed: if only it was asked by BoostPriority, but not requested (ie TryAddNewState has not set the Compiling bit
@@ -2168,26 +2228,64 @@ public:
 		SET_DWORD_STAT_FName(TPrecachePipelineCacheDerived::GetHighestPriorityCompileStatName(), HighestPriorityCompileCount);
 		
 		FRWScopeLock WriteLock(PrecachePSOsRWLock, SLT_Write);
-		for (int32 Index = 0; Index < PrecachedPSOs.Num(); ++Index)		
+		for (int32 Index = 0; Index < PrecachedPSOsToCleanup.Num(); ++Index)
 		{
-			uint64 InitializerHash = PrecachedPSOs[Index];
+			uint64 InitializerHash = PrecachedPSOsToCleanup[Index];
 
 			FPrecacheTask* FindResult = PrecachedPSOInitializerData.Find(InitializerHash);
-			check(FindResult && IsCompilationDone((FindResult->ReadPSOPrecacheState())));
-			if (FindResult->PipelineState->IsComplete())
+			check(FindResult && (GPSOPrecacheKeepInMemoryUntilUsed || IsCompilationDone((FindResult->ReadPSOPrecacheState()))));
+			if (!FindResult || !FindResult->PipelineState)
+			{
+				// Was already cleaned up (can happen if it was marked as used).
+				PrecachedPSOsToCleanup.RemoveAtSwap(Index);
+				Index--;
+			}
+			else if (FindResult->PipelineState->IsComplete())
 			{
 				// This is needed to cleanup the members - bit strange because it's complete already
 				verify(!FindResult->PipelineState->WaitCompletion());
-
 				delete FindResult->PipelineState;
 				FindResult->PipelineState = nullptr;
 
-				PrecachedPSOs.RemoveAtSwap(Index);
+				if (GPSOPrecacheKeepInMemoryUntilUsed)
+				{
+					DEC_DWORD_STAT(STAT_InMemoryPrecachedPSOCount);
+				}
+
+				PrecachedPSOsToCleanup.RemoveAtSwap(Index);
 				Index--;
 			}
 		}
 	}
 
+	// Marks a PSO as used for rendering so that it will be enqueued for cleanup if it's currently kept in memory.
+	// Does not do anything if PSOs are not kept in memory, or if a bound on the number of PSOs kept in memory is set.
+	void MarkPSOAsUsed(const TPrecachedPSOInitializer& Initializer)
+	{
+		if (!GPSOPrecacheKeepInMemoryUntilUsed || MaxInMemoryPSOs > 0)
+		{
+			return;
+		}
+
+		uint64 InitializerHash = TPrecachePipelineCacheDerived::PipelineStateInitializerHash(Initializer);
+		bool bMarkedAsUsed = false;
+		{
+			FRWScopeLock ReadOnlyLock(PrecachePSOsRWLock, SLT_ReadOnly);
+			FPrecacheTask* FindResult = PrecachedPSOInitializerData.Find(InitializerHash);
+			if (FindResult)
+			{
+				EPSOPrecacheStateMask PreviousState = FindResult->AddPSOPrecacheState(EPSOPrecacheStateMask::UsedForRendering);
+				bMarkedAsUsed = !EnumHasAnyFlags(PreviousState, EPSOPrecacheStateMask::UsedForRendering);
+			}
+		}
+
+		if (bMarkedAsUsed)
+		{
+			FRWScopeLock WriteLock(PrecachePSOsRWLock, SLT_Write);
+			PrecachedPSOsToCleanup.Add(InitializerHash);
+		}
+	}
+
 protected:
 	bool HasPSOBeenRequested(const TPrecachedPSOInitializer& Initializer, uint64 InitializerHash, FPSOPrecacheRequestResult& Result)
 	{
@@ -2313,7 +2411,13 @@ protected:
 	volatile int32 HighestPriorityCompileCount = 0;
 
 	// Finished Precached PSOs which can be garbage collected
-	TArray<uint64> PrecachedPSOs;
+	TArray<uint64> PrecachedPSOsToCleanup;
+
+	// Circular buffer of indices of PSOs that are kept in memory if MaxInMemoryPSOs is non-zero.
+	TArray<uint32> InMemoryPSOIndices;
+	uint32 MaxInMemoryPSOs = 0;
+	uint32 CurrentInMemoryPSOIndex = 0;
+	static constexpr uint32 kInvalidPSOIndex = TNumericLimits<uint32>::Max();
 };
 
 struct FPrecacheComputeInitializer
@@ -2339,7 +2443,7 @@ public:
 		return new FComputePipelineState((FRHIComputeShader*)ComputeShaderInitializer.RHIComputeShaderAsU64, ComputeShaderInitializer.Name);
 	}
 
-	FPrecacheComputePipelineCache() : TPrecachePipelineCacheBase(FPSOPrecacheRequestID::EType::Compute) {}
+	FPrecacheComputePipelineCache(uint32 InMemoryPSOsMaxNum) : TPrecachePipelineCacheBase(FPSOPrecacheRequestID::EType::Compute, InMemoryPSOsMaxNum) {}
 	FPSOPrecacheRequestResult PrecacheComputePipelineState(FRHIComputeShader* ComputeShader, const TCHAR* Name, bool bForcePrecache);
 	static void OnNewPipelineStateCreated(const FPrecacheComputeInitializer& ComputeInitializer, FComputePipelineState* NewComputePipelineState, const FString& PSOCompilationEventName, bool bDoAsyncCompile);
  
@@ -2360,7 +2464,6 @@ public:
 	{
 		return Key.RHIComputeShaderAsU64;
 	}
-
 };
 
 class FPrecacheGraphicsPipelineCache : public TPrecachePipelineCacheBase<FPrecacheGraphicsPipelineCache, FGraphicsPipelineStateInitializer, FGraphicsPipelineState>
@@ -2389,7 +2492,7 @@ public:
 	{
 		return GET_STATFNAME(STAT_HighestPriorityGraphicsPSOPrecacheRequests);
 	}
-	FPrecacheGraphicsPipelineCache() : TPrecachePipelineCacheBase(FPSOPrecacheRequestID::EType::Graphics) {}
+	FPrecacheGraphicsPipelineCache(uint32 InMemoryPSOsMaxNum) : TPrecachePipelineCacheBase(FPSOPrecacheRequestID::EType::Graphics, InMemoryPSOsMaxNum) {}
 	FPSOPrecacheRequestResult PrecacheGraphicsPipelineState(const FGraphicsPipelineStateInitializer& Initializer);
 
 	static void OnNewPipelineStateCreated(const FGraphicsPipelineStateInitializer& Initializer, FGraphicsPipelineState* NewGraphicsPipelineState, const FString& PSOCompilationEventName, bool bDoAsyncCompile);
@@ -2405,8 +2508,9 @@ typedef TSharedPipelineStateCache<FGraphicsPipelineStateInitializer, FGraphicsPi
 FComputePipelineCache GComputePipelineCache;
 FWorkGraphPipelineCache GWorkGraphPipelineCache;
 FGraphicsPipelineCache GGraphicsPipelineCache;
-FPrecacheGraphicsPipelineCache GPrecacheGraphicsPipelineCache;
-FPrecacheComputePipelineCache GPrecacheComputePipelineCache;
+TUniquePtr<FPrecacheGraphicsPipelineCache> GPrecacheGraphicsPipelineCache;
+TUniquePtr<FPrecacheComputePipelineCache> GPrecacheComputePipelineCache;
+
 
 FAutoConsoleTaskPriority CPrio_FCompilePipelineStateTask(
 	TEXT("TaskGraph.TaskPriorities.CompilePipelineStateTask"),
@@ -2809,10 +2913,15 @@ public:
 					HandlePipelineCreationFailure(ComputePipeline->ComputeShader, Initializer.bFromPSOFileCache || Initializer.bPSOPrecache);
 				}
 
+				FPrecacheComputeInitializer PrecacheInitializer{ ComputePipeline->ComputeShader, ComputePipeline->Name };
 				if (Initializer.bPSOPrecache)
 				{
 					bool bCSValid = ComputePipeline->RHIPipeline != nullptr && ComputePipeline->RHIPipeline->IsValid();
-					GPrecacheComputePipelineCache.PrecacheFinished(FPrecacheComputeInitializer(ComputePipeline->ComputeShader, ComputePipeline->Name), bCSValid);
+					GPrecacheComputePipelineCache->PrecacheFinished(PrecacheInitializer, bCSValid);
+				}
+				else
+				{
+					GPrecacheComputePipelineCache->MarkPSOAsUsed(PrecacheInitializer);
 				}
 				
 #if PLATFORM_WINDOWS
@@ -2875,7 +2984,11 @@ public:
 				// Mark as finished when it's a precaching job
 				if (Initializer.bPSOPrecache)
 				{
-					GPrecacheGraphicsPipelineCache.PrecacheFinished(Initializer, GfxPipeline->RHIPipeline != nullptr);
+					GPrecacheGraphicsPipelineCache->PrecacheFinished(Initializer, GfxPipeline->RHIPipeline != nullptr);
+				}
+				else
+				{
+					GPrecacheGraphicsPipelineCache->MarkPSOAsUsed(Initializer);
 				}
 				
 #if PLATFORM_WINDOWS
@@ -3004,8 +3117,8 @@ void PipelineStateCache::FlushResources()
 	GWorkGraphPipelineCache.FlushResources(bDiscardAndSwap);
 	GGraphicsPipelineCache.FlushResources(bDiscardAndSwap);
 
-	GPrecacheGraphicsPipelineCache.ProcessDelayedCleanup();
-	GPrecacheComputePipelineCache.ProcessDelayedCleanup();
+	GPrecacheGraphicsPipelineCache->ProcessDelayedCleanup();
+	GPrecacheComputePipelineCache->ProcessDelayedCleanup();
 
 	FPipelineFileCacheManager::BroadcastNewPSOsDelegate();
 
@@ -3570,7 +3683,7 @@ static void InternalCreateGraphicsPipelineState(const FGraphicsPipelineStateInit
 
 		if (Initializer.bPSOPrecache)
 		{
-			GPrecacheGraphicsPipelineCache.PrecacheFinished(Initializer, CachedState->RHIPipeline != nullptr);
+			GPrecacheGraphicsPipelineCache->PrecacheFinished(Initializer, CachedState->RHIPipeline != nullptr);
 		}
 
 		if (CachedState->RHIPipeline)
@@ -3802,7 +3915,7 @@ FPSOPrecacheRequestResult FPrecacheComputePipelineCache::PrecacheComputePipeline
 
 	FPrecacheComputeInitializer PrecacheComputeInitializer(ComputeShader, Name);
 	static bool bDoAsyncCompile = FApp::ShouldUseThreadingForPerformance();
-	return GPrecacheComputePipelineCache.TryAddNewState(PrecacheComputeInitializer, Name, bDoAsyncCompile);
+	return TryAddNewState(PrecacheComputeInitializer, Name, bDoAsyncCompile);
 }
 
 void FPrecacheComputePipelineCache::OnNewPipelineStateCreated(const FPrecacheComputeInitializer& ComputeInitializer, FComputePipelineState* CachedState, const FString& PSOCompilationEventName, bool bDoAsyncCompile)
@@ -3841,13 +3954,13 @@ void FPrecacheComputePipelineCache::OnNewPipelineStateCreated(const FPrecacheCom
 	{
 		check(CachedState->GetCompletionEvent() == nullptr);
 		CachedState->RHIPipeline = RHICreateComputePipelineState(CachedState->ComputeShader);
-		GPrecacheComputePipelineCache.PrecacheFinished(ComputeInitializer, CachedState->RHIPipeline != nullptr);
+		GPrecacheComputePipelineCache->PrecacheFinished(ComputeInitializer, CachedState->RHIPipeline != nullptr);
 	}
 }
 
 FPSOPrecacheRequestResult PipelineStateCache::PrecacheComputePipelineState(FRHIComputeShader* ComputeShader, const TCHAR* Name, bool bForcePrecache)
 {
-	return GPrecacheComputePipelineCache.PrecacheComputePipelineState(ComputeShader, Name, bForcePrecache);
+	return GPrecacheComputePipelineCache->PrecacheComputePipelineState(ComputeShader, Name, bForcePrecache);
 }
 
 FPSOPrecacheRequestResult FPrecacheGraphicsPipelineCache::PrecacheGraphicsPipelineState(const FGraphicsPipelineStateInitializer& Initializer)
@@ -3866,7 +3979,7 @@ FPSOPrecacheRequestResult FPrecacheGraphicsPipelineCache::PrecacheGraphicsPipeli
 	FString PSOCompilationEventName;
 
 	// try and create new graphics state
-	return GPrecacheGraphicsPipelineCache.TryAddNewState(Initializer, PSOCompilationEventName, bDoAsyncCompile);
+	return TryAddNewState(Initializer, PSOCompilationEventName, bDoAsyncCompile);
 }
 	
 void FPrecacheGraphicsPipelineCache::OnNewPipelineStateCreated(const FGraphicsPipelineStateInitializer & Initializer, FGraphicsPipelineState * NewGraphicsPipelineState, const FString& PSOCompilationEventName, bool bDoAsyncCompile)
@@ -3888,7 +4001,7 @@ void FPrecacheGraphicsPipelineCache::OnNewPipelineStateCreated(const FGraphicsPi
 
 FPSOPrecacheRequestResult PipelineStateCache::PrecacheGraphicsPipelineState(const FGraphicsPipelineStateInitializer& Initializer)
 {
-	return GPrecacheGraphicsPipelineCache.PrecacheGraphicsPipelineState(Initializer);
+	return GPrecacheGraphicsPipelineCache->PrecacheGraphicsPipelineState(Initializer);
 }
 
 EPSOPrecacheResult PipelineStateCache::CheckPipelineStateInCache(const FGraphicsPipelineStateInitializer& PipelineStateInitializer)
@@ -3898,7 +4011,7 @@ EPSOPrecacheResult PipelineStateCache::CheckPipelineStateInCache(const FGraphics
 		return EPSOPrecacheResult::Unknown;
 	}
 
-	return GPrecacheGraphicsPipelineCache.GetPrecachingState(PipelineStateInitializer);
+	return GPrecacheGraphicsPipelineCache->GetPrecachingState(PipelineStateInitializer);
 }
 
 EPSOPrecacheResult PipelineStateCache::CheckPipelineStateInCache(FRHIComputeShader* ComputeShader)
@@ -3908,7 +4021,7 @@ EPSOPrecacheResult PipelineStateCache::CheckPipelineStateInCache(FRHIComputeShad
 		return EPSOPrecacheResult::Unknown;
 	}
 
-	return GPrecacheComputePipelineCache.GetPrecachingState(FPrecacheComputeInitializer(ComputeShader, nullptr));
+	return GPrecacheComputePipelineCache->GetPrecachingState(FPrecacheComputeInitializer(ComputeShader, nullptr));
 }
 
 bool PipelineStateCache::IsPrecaching(const FPSOPrecacheRequestID& PSOPrecacheRequestID)
@@ -3921,11 +4034,11 @@ bool PipelineStateCache::IsPrecaching(const FPSOPrecacheRequestID& PSOPrecacheRe
 	EPSOPrecacheResult PrecacheResult = EPSOPrecacheResult::Unknown;
 	if (PSOPrecacheRequestID.GetType() == FPSOPrecacheRequestID::EType::Graphics)
 	{
-		PrecacheResult = GPrecacheGraphicsPipelineCache.GetPrecachingState(PSOPrecacheRequestID);
+		PrecacheResult = GPrecacheGraphicsPipelineCache->GetPrecachingState(PSOPrecacheRequestID);
 	}
 	else
 	{
-		PrecacheResult = GPrecacheComputePipelineCache.GetPrecachingState(PSOPrecacheRequestID);
+		PrecacheResult = GPrecacheComputePipelineCache->GetPrecachingState(PSOPrecacheRequestID);
 	}
 	return PrecacheResult == EPSOPrecacheResult::Active;
 }
@@ -3937,7 +4050,7 @@ bool PipelineStateCache::IsPrecaching(const FGraphicsPipelineStateInitializer& P
 		return false;
 	}
 
-	return GPrecacheGraphicsPipelineCache.GetPrecachingState(PipelineStateInitializer) == EPSOPrecacheResult::Active;
+	return GPrecacheGraphicsPipelineCache->GetPrecachingState(PipelineStateInitializer) == EPSOPrecacheResult::Active;
 }
 
 bool PipelineStateCache::IsPrecaching(FRHIComputeShader* ComputeShader)
@@ -3947,7 +4060,7 @@ bool PipelineStateCache::IsPrecaching(FRHIComputeShader* ComputeShader)
 		return false;
 	}
 
-	return GPrecacheComputePipelineCache.GetPrecachingState(FPrecacheComputeInitializer(ComputeShader, nullptr)) == EPSOPrecacheResult::Active;
+	return GPrecacheComputePipelineCache->GetPrecachingState(FPrecacheComputeInitializer(ComputeShader, nullptr)) == EPSOPrecacheResult::Active;
 }
 
 bool PipelineStateCache::IsPrecaching()
@@ -3957,7 +4070,7 @@ bool PipelineStateCache::IsPrecaching()
 		return false;
 	}
 
-	return GPrecacheGraphicsPipelineCache.IsPrecaching() || GPrecacheComputePipelineCache.IsPrecaching();
+	return GPrecacheGraphicsPipelineCache->IsPrecaching() || GPrecacheComputePipelineCache->IsPrecaching();
 }
 
 void PipelineStateCache::BoostPrecachePriority(EPSOPrecachePriority PSOPrecachePriority, const FPSOPrecacheRequestID& PSOPrecacheRequestID)
@@ -3968,11 +4081,11 @@ void PipelineStateCache::BoostPrecachePriority(EPSOPrecachePriority PSOPrecacheP
 
 		if (PSOPrecacheRequestID.GetType() == FPSOPrecacheRequestID::EType::Graphics)
 		{
-			GPrecacheGraphicsPipelineCache.BoostPriority(PSOPrecachePriority, PSOPrecacheRequestID);
+			GPrecacheGraphicsPipelineCache->BoostPriority(PSOPrecachePriority, PSOPrecacheRequestID);
 		}
 		else
 		{
-			GPrecacheComputePipelineCache.BoostPriority(PSOPrecachePriority, PSOPrecacheRequestID);
+			GPrecacheComputePipelineCache->BoostPriority(PSOPrecachePriority, PSOPrecacheRequestID);
 		}
 	}
 }
@@ -3994,7 +4107,7 @@ uint32 PipelineStateCache::NumActivePrecacheRequests()
 		return 0;
 	}
 
-	return GPrecacheGraphicsPipelineCache.NumActivePrecacheRequests() + GPrecacheComputePipelineCache.NumActivePrecacheRequests();
+	return GPrecacheGraphicsPipelineCache->NumActivePrecacheRequests() + GPrecacheComputePipelineCache->NumActivePrecacheRequests();
 }
 
 void PipelineStateCache::ResetPSOHitchTrackingStats()
@@ -4108,8 +4221,14 @@ void PipelineStateCache::WaitForAllTasks()
 	GComputePipelineCache.WaitTasksComplete();
 	GWorkGraphPipelineCache.WaitTasksComplete();
 	GGraphicsPipelineCache.WaitTasksComplete();
-	GPrecacheGraphicsPipelineCache.WaitTasksComplete();
-	GPrecacheComputePipelineCache.WaitTasksComplete();
+	GPrecacheGraphicsPipelineCache->WaitTasksComplete();
+	GPrecacheComputePipelineCache->WaitTasksComplete();
+}
+
+void PipelineStateCache::Init()
+{
+	GPrecacheGraphicsPipelineCache = MakeUnique<FPrecacheGraphicsPipelineCache>(GPSOPrecacheKeepInMemoryGraphicsMaxNum);
+	GPrecacheComputePipelineCache = MakeUnique<FPrecacheComputePipelineCache>(GPSOPrecacheKeepInMemoryComputeMaxNum);
 }
 
 void PipelineStateCache::Shutdown()
@@ -4133,6 +4252,9 @@ void PipelineStateCache::Shutdown()
 	GVertexDeclarationCache.Empty();
 
 	GPSOPrecacheThreadPool.ShutdownThreadPool();
+
+	GPrecacheGraphicsPipelineCache.Reset();
+	GPrecacheComputePipelineCache.Reset();
 }
 
 FRHIVertexDeclaration*	PipelineStateCache::GetOrCreateVertexDeclaration(const FVertexDeclarationElementList& Elements)
diff --git a/Engine/Source/Runtime/RHI/Public/PipelineStateCache.h b/Engine/Source/Runtime/RHI/Public/PipelineStateCache.h
--- a/Engine/Source/Runtime/RHI/Public/PipelineStateCache.h
+++ b/Engine/Source/Runtime/RHI/Public/PipelineStateCache.h
@@ -150,6 +150,9 @@ namespace PipelineStateCache
 	// Waits for any pending tasks to complete.
 	extern RHI_API void WaitForAllTasks();
 
+	// Initializes any required component.
+	extern RHI_API void Init();
+
 	/* Clears all pipeline cached state. Called on shutdown, calling GetAndOrCreate after this will recreate state */
 	extern RHI_API void Shutdown();
 
