New GPU profiler improvements. The TStatId on FRHIBreadcrumb has been replaced with a FRHIBreadcrumbData struct that holds additional profiling related data, which includes:
 - The TStatId for "stat gpu" stats.
 - The FName required by the CSV profiler for GPU stats.
 - The source file and line number to allow breadcrumbs shown in Insights to link back to their original source location.

Additional changes:
 - Added temporary support for the Insights GPU track. This is guarded by RHI_TEMP_USE_GPU_TRACE until we have a newer, more capable API.
 - Simplified FMeshDrawEvent into a standard RHI breadcrumb in FMeshDrawCommand::SubmitDraw().
 - Moved "r.GPUCsvStatsEnabled" cvar into GPUProfiler.cpp, so it is accessible to both old and new profilers.

diff --git a/Engine/Source/Runtime/Engine/Private/KismetRenderingLibrary.cpp b/Engine/Source/Runtime/Engine/Private/KismetRenderingLibrary.cpp
--- a/Engine/Source/Runtime/Engine/Private/KismetRenderingLibrary.cpp
+++ b/Engine/Source/Runtime/Engine/Private/KismetRenderingLibrary.cpp
@@ -803,7 +803,7 @@ void UKismetRenderingLibrary::BeginDrawCanvasToRenderTarget(UObject* WorldContex
 				](FRHICommandListImmediate& RHICmdList)
 			{
 #if  WANTS_DRAW_MESH_EVENTS
-				Breadcrumb->Emplace(RHICmdList, TStatId(), TEXT("DrawCanvasToTarget: %s"), Name);
+				Breadcrumb->Emplace(RHICmdList, FRHIBreadcrumbData(__FILE__, __LINE__, TStatId(), NAME_None), TEXT("DrawCanvasToTarget: %s"), Name);
 #endif
 				RenderTargetResource->FlushDeferredResourceUpdate(RHICmdList);
 			});
diff --git a/Engine/Source/Runtime/Launch/Private/LaunchEngineLoop.cpp b/Engine/Source/Runtime/Launch/Private/LaunchEngineLoop.cpp
--- a/Engine/Source/Runtime/Launch/Private/LaunchEngineLoop.cpp
+++ b/Engine/Source/Runtime/Launch/Private/LaunchEngineLoop.cpp
@@ -5493,12 +5493,12 @@ static inline void BeginFrameRenderThread(FRHICommandListImmediate& RHICmdList,
 #if CSV_PROFILER
 	if (FCsvProfiler::Get()->IsCapturing_Renderthread())
 	{
-		GRHIFrameBreadcrumb.Emplace(RHICmdList, TStatId(), TEXT("CsvFrame %d"), FCsvProfiler::Get()->GetCaptureFrameNumberRT());
+		GRHIFrameBreadcrumb.Emplace(RHICmdList, FRHIBreadcrumbData(__FILE__, __LINE__, TStatId(), NAME_None), TEXT("CsvFrame %d"), FCsvProfiler::Get()->GetCaptureFrameNumberRT());
 	}
 	else
 #endif
 	{
-		GRHIFrameBreadcrumb.Emplace(RHICmdList, TStatId(), TEXT("Frame %d"), CurrentFrameCounter);
+		GRHIFrameBreadcrumb.Emplace(RHICmdList, FRHIBreadcrumbData(__FILE__, __LINE__, TStatId(), NAME_None), TEXT("Frame %d"), CurrentFrameCounter);
 	}
 #endif
 
diff --git a/Engine/Source/Runtime/RHI/Private/GPUProfiler.cpp b/Engine/Source/Runtime/RHI/Private/GPUProfiler.cpp
--- a/Engine/Source/Runtime/RHI/Private/GPUProfiler.cpp
+++ b/Engine/Source/Runtime/RHI/Private/GPUProfiler.cpp
@@ -8,6 +8,7 @@
 #include "Async/TaskGraphInterfaces.h"
 #include "Misc/WildcardString.h"
 #include "RHI.h"
+#include "GpuProfilerTrace.h"
 
 #if !UE_BUILD_SHIPPING
 #include "VisualizerEvents.h"
@@ -15,6 +16,34 @@
 #include "Modules/ModuleManager.h"
 #endif
 
+#if HAS_GPU_STATS
+CSV_DEFINE_CATEGORY_MODULE(RHI_API, GPU, true);
+DECLARE_FLOAT_COUNTER_STAT(TEXT("[TOTAL]"), Stat_GPU_Total, STATGROUP_GPU);
+CSV_DEFINE_STAT(GPU, Total);
+#endif
+
+// Temporary function to resolve link issues with the above "Total" GPU stat moving from RenderCore to RHI.
+// We can remove this once the old GPU profiler code has been deleted, and RHI_NEW_GPU_PROFILER is set to 1 permanently.
+extern RHI_API void RHISetGPUStatTotals(bool bCsvStatsEnabled, double TotalMs)
+{
+#if STATS
+	FThreadStats::AddMessage(GET_STATFNAME(Stat_GPU_Total), EStatOperation::Set, TotalMs);
+	TRACE_STAT_SET(GET_STATFNAME(Stat_GPU_Total), TotalMs);
+#endif
+
+#if CSV_PROFILER_STATS
+	if (bCsvStatsEnabled)
+	{
+		FCsvProfiler::Get()->RecordCustomStat(CSV_STAT_FNAME(Total), CSV_CATEGORY_INDEX(GPU), TotalMs, ECsvCustomStatOp::Set);
+	}
+#endif
+}
+
+static TAutoConsoleVariable<int> CVarGPUCsvStatsEnabled(
+	TEXT("r.GPUCsvStatsEnabled"),
+	0,
+	TEXT("Enables or disables GPU stat recording to CSVs"));
+
 #define LOCTEXT_NAMESPACE "GpuProfiler"
 
 static TAutoConsoleVariable<FString> GProfileGPUPatternCVar(
@@ -669,6 +698,10 @@ bool FGPUTiming::GAreGlobalsInitialized = false;
 
 #else
 
+// Temporary. Adds Insights markers for the 0th GPU graphics queue
+// until we have a new API that is capable of displaying more info.
+#define RHI_TEMP_USE_GPU_TRACE (1 && GPUPROFILERTRACE_ENABLED)
+
 namespace UE::RHI::GPUProfiler
 {
 	static TArray<FEventSink*>& GetSinks()
@@ -816,28 +849,55 @@ namespace UE::RHI::GPUProfiler
 			}
 		};
 
+		struct FQueueTimestamps
+		{
+			FTimestampStream Queue;
+
+		#if WITH_RHI_BREADCRUMBS
+			TMap<FRHIBreadcrumbData_Stats, FTimestampStream> Stats;
+		#endif
+		};
+
 		struct FQueueState
 		{
 			bool bBusy = false;
-			FTimestampStream QueueTimestamps;
+			FQueueTimestamps Timestamps;
 
-		#if WITH_RHI_BREADCRUMBS && HAS_GPU_STATS
-			TMap<TStatId, FTimestampStream> StatTimestamps;
-			TMap<TStatId, int32> ActiveStats;
+		#if WITH_RHI_BREADCRUMBS
+			TMap<FRHIBreadcrumbData_Stats, int32> ActiveStats;
 		#endif
 		};
 
-		using FFrameState = TMap<FQueue, FQueueState>;
+		using FFrameState = TMap<FQueue, FQueueTimestamps>;
 
 		TMap<FQueue, FQueueState> QueueStates;
 		TMap<uint32, FFrameState> Frames;
 
+	#if RHI_TEMP_USE_GPU_TRACE
+		uint32 FrameNumber = 0;
+		bool ShouldEmitGPU(FQueue const& Queue)
+		{
+			return Queue.Type == FQueue::EType::Graphics && Queue.GPU == 0 && Queue.Index == 0;
+		}
+	#endif
+
 		void InitializeQueues(TConstArrayView<FQueue> Queues) override
 		{
 			for (FQueue const& Queue : Queues)
 			{
 				check(QueueStates.Find(Queue) == nullptr);
 				QueueStates.Add(Queue);
+
+			#if RHI_TEMP_USE_GPU_TRACE
+				// Start the first Insights GPU frame
+				if (ShouldEmitGPU(Queue))
+				{
+					// Use 1,1 calibration to disable any adjustments Insights makes.
+					// The timestamps we use in the GPU event stream are already in the CPU clock domain.
+					FGPUTimingCalibrationTimestamp Calibration{ 1, 1 };
+					FGpuProfilerTrace::BeginFrame(Calibration);
+				}
+			#endif
 			}
 		}
 
@@ -853,14 +913,15 @@ namespace UE::RHI::GPUProfiler
 					{
 						check(!QueueState.bBusy);
 						QueueState.bBusy = true;
-						uint64 Value = Event->Value.Get<FEvent::FBeginWork>().GPUTimestampTOP;
-						QueueState.QueueTimestamps.AddTimestamp(Value, true);
 
-					#if WITH_RHI_BREADCRUMBS && HAS_GPU_STATS
+						FEvent::FBeginWork const& BeginWork = Event->Value.Get<FEvent::FBeginWork>();
+						QueueState.Timestamps.Queue.AddTimestamp(BeginWork.GPUTimestampTOP, true);
+
+					#if WITH_RHI_BREADCRUMBS
 						// Apply the timestamp to all active stats
-						for (auto const& [StatId, RefCount] : QueueState.ActiveStats)
+						for (auto const& [Stat, RefCount] : QueueState.ActiveStats)
 						{
-							QueueState.StatTimestamps.FindChecked(StatId).AddTimestamp(Value, true);
+							QueueState.Timestamps.Stats.FindChecked(Stat).AddTimestamp(BeginWork.GPUTimestampTOP, true);
 						}
 					#endif
 					}
@@ -872,51 +933,44 @@ namespace UE::RHI::GPUProfiler
 						QueueState.bBusy = false;
 
 						uint64 Value = Event->Value.Get<FEvent::FEndWork>().GPUTimestampBOP;
-						QueueState.QueueTimestamps.AddTimestamp(Value, false);
+						QueueState.Timestamps.Queue.AddTimestamp(Value, false);
 
-					#if WITH_RHI_BREADCRUMBS && HAS_GPU_STATS
+					#if WITH_RHI_BREADCRUMBS
 						// Apply the timestamp to all active stats
-						for (auto const& [StatId, RefCount] : QueueState.ActiveStats)
+						for (auto const& [Stat, RefCount] : QueueState.ActiveStats)
 						{
-							QueueState.StatTimestamps.FindChecked(StatId).AddTimestamp(Value, false);
+							QueueState.Timestamps.Stats.FindChecked(Stat).AddTimestamp(Value, false);
 						}
 					#endif
 					}
 					break;
 
-				case FEvent::EType::FrameBoundary:
+			#if WITH_RHI_BREADCRUMBS
+				case FEvent::EType::BeginBreadcrumb:
 					{
-						check(!QueueState.bBusy);
-						FEvent::FFrameBoundary const& FrameBoundary = Event->Value.Get<FEvent::FFrameBoundary>();
-
-						FFrameState& FrameState = Frames.FindOrAdd(FrameBoundary.FrameNumber);
-						FrameState.Emplace(Queue, MoveTemp(QueueState));
+						FEvent::FBeginBreadcrumb const& BeginBreadcrumb = Event->Value.Get<FEvent::FBeginBreadcrumb>();
+						FRHIBreadcrumbData_Stats const& Stat = BeginBreadcrumb.Breadcrumb->Name.Data;
 
-						if (FrameState.Num() == QueueStates.Num())
+					#if RHI_TEMP_USE_GPU_TRACE
+						if (ShouldEmitGPU(Queue))
 						{
-							// All registered queues have reported their frame boundary event.
-							// We have a full set of data to compute the total frame GPU stats.
-							ProcessFrame(FrameState);
+							FRHIBreadcrumb::FBuffer Buffer;
+							TCHAR const* Str = BeginBreadcrumb.Breadcrumb->Name.GetTCHAR(Buffer);
+							FName Name(Str);
 
-							Frames.Remove(FrameBoundary.FrameNumber);
+							FGpuProfilerTrace::SpecifyEventByName(Name);
+							FGpuProfilerTrace::BeginEventByName(Name, FrameNumber, (uint64)(FPlatformTime::ToMilliseconds64(BeginBreadcrumb.GPUTimestampTOP) * 1000.0));
 						}
-					}
-					break;
-
-			#if WITH_RHI_BREADCRUMBS && HAS_GPU_STATS
-				case FEvent::EType::BeginBreadcrumb:
-					{
-						FEvent::FBeginBreadcrumb const& BeginBreadcrumb = Event->Value.Get<FEvent::FBeginBreadcrumb>();
-						TStatId StatId = BeginBreadcrumb.Breadcrumb->Name.StatId;
+					#endif
 
-						if (StatId.IsValidStat())
+						if (Stat.ShouldComputeStat())
 						{
 							// Disregard the stat if it is nested within itself (i.e. its already in the ActiveStats map with a non-zero ref count).
 							// Only the outermost stat will count the busy time, otherwise we'd be double-counting the nested time.
-							int32 RefCount = QueueState.ActiveStats.FindOrAdd(StatId)++;
+							int32 RefCount = QueueState.ActiveStats.FindOrAdd(Stat)++;
 							if (RefCount == 0)
 							{
-								QueueState.StatTimestamps.FindOrAdd(StatId).AddTimestamp(BeginBreadcrumb.GPUTimestampTOP, true);
+								QueueState.Timestamps.Stats.FindOrAdd(Stat).AddTimestamp(BeginBreadcrumb.GPUTimestampTOP, true);
 							}
 						}
 					}
@@ -925,22 +979,69 @@ namespace UE::RHI::GPUProfiler
 				case FEvent::EType::EndBreadcrumb:
 					{
 						FEvent::FEndBreadcrumb const& EndBreadcrumb = Event->Value.Get<FEvent::FEndBreadcrumb>();
-						TStatId StatId = EndBreadcrumb.Breadcrumb->Name.StatId;
+						FRHIBreadcrumbData_Stats const& Stat = EndBreadcrumb.Breadcrumb->Name.Data;
 
-						if (StatId.IsValidStat())
+					#if RHI_TEMP_USE_GPU_TRACE
+						if (ShouldEmitGPU(Queue))
+						{
+							FGpuProfilerTrace::EndEvent((uint64)(FPlatformTime::ToMilliseconds64(EndBreadcrumb.GPUTimestampBOP) * 1000.0));
+						}
+					#endif
+
+						if (Stat.ShouldComputeStat())
 						{
 							// Pop the stat when the refcount hits zero.
-							int32 RefCount = --QueueState.ActiveStats.FindChecked(StatId);
+							int32 RefCount = --QueueState.ActiveStats.FindChecked(Stat);
 							if (RefCount == 0)
 							{
-								QueueState.StatTimestamps.FindChecked(StatId).AddTimestamp(EndBreadcrumb.GPUTimestampBOP, false);
-								QueueState.ActiveStats.FindAndRemoveChecked(StatId);
+								QueueState.Timestamps.Stats.FindChecked(Stat).AddTimestamp(EndBreadcrumb.GPUTimestampBOP, false);
+								QueueState.ActiveStats.FindAndRemoveChecked(Stat);
 							}							
 						}
 					}
 					break;
-			#endif // WITH_RHI_BREADCRUMBS && HAS_GPU_STATS
+			#endif // WITH_RHI_BREADCRUMBS
 
+				case FEvent::EType::FrameBoundary:
+					{
+						check(!QueueState.bBusy);
+						FEvent::FFrameBoundary const& FrameBoundary = Event->Value.Get<FEvent::FFrameBoundary>();
+
+						FFrameState& FrameState = Frames.FindOrAdd(FrameBoundary.FrameNumber);
+						FrameState.Emplace(Queue, MoveTemp(QueueState.Timestamps));
+
+						// Reinsert timestamp streams for the current active stats on 
+						// this queue, since these got moved into the frame state.
+						for (auto& [Stat, RefCount] : QueueState.ActiveStats)
+						{
+							QueueState.Timestamps.Stats.FindOrAdd(Stat);
+						}
+
+						if (FrameState.Num() == QueueStates.Num())
+						{
+							// All registered queues have reported their frame boundary event.
+							// We have a full set of data to compute the total frame GPU stats.
+							ProcessFrame(FrameState);
+
+							Frames.Remove(FrameBoundary.FrameNumber);
+						}
+
+					#if RHI_TEMP_USE_GPU_TRACE
+						// End the current Insights GPU frame + start the next one
+						if (ShouldEmitGPU(Queue))
+						{
+							check(FrameNumber == FrameBoundary.FrameNumber);
+							FGpuProfilerTrace::EndFrame(0);
+							FrameNumber++;
+
+							// Use 1,1 calibration to disable any adjustments Insights makes.
+							// The timestamps we use in the GPU event stream are already in the CPU clock domain.
+							FGPUTimingCalibrationTimestamp Calibration{ 1, 1 };
+							FGpuProfilerTrace::BeginFrame(Calibration);
+						}
+					#endif
+					}
+					break;
 				}
 			}
 		}
@@ -949,23 +1050,31 @@ namespace UE::RHI::GPUProfiler
 		{
 			TArray<FTimestampStream::FState, TInlineAllocator<GetRHIPipelineCount() * MAX_NUM_GPUS>> StreamPointers;
 
+		#if CSV_PROFILER_STATS
+			const bool bCsvStatsEnabled = !!CVarGPUCsvStatsEnabled.GetValueOnAnyThread();
+			FCsvProfiler* CsvProfiler = bCsvStatsEnabled ? FCsvProfiler::Get() : nullptr;
+		#else
+			const bool bCsvStatsEnabled = false;
+		#endif
+
 			// Compute the individual GPU stats
-		#if WITH_RHI_BREADCRUMBS && HAS_GPU_STATS
-			TSet<TStatId> UniqueStats;
+		#if WITH_RHI_BREADCRUMBS
+
+			TSet<FRHIBreadcrumbData_Stats> UniqueStats;
 			for (auto const& [Queue, State] : FrameState)
 			{
-				for (auto const& [StatId, Timestamps] : State.StatTimestamps)
+				for (auto const& [Stat, Timestamps] : State.Stats)
 				{
-					UniqueStats.Add(StatId);
+					UniqueStats.Add(Stat);
 				}
 			}
 
-			for (TStatId StatId : UniqueStats)
+			for (FRHIBreadcrumbData_Stats const& Stat : UniqueStats)
 			{
 				StreamPointers.Reset();
 				for (auto const& [Queue, State] : FrameState)
 				{
-					FTimestampStream const* Stream = State.StatTimestamps.Find(StatId);
+					FTimestampStream const* Stream = State.Stats.Find(Stat);
 					if (Stream)
 					{
 						StreamPointers.Emplace(*Stream);
@@ -973,22 +1082,32 @@ namespace UE::RHI::GPUProfiler
 				}
 
 				uint64 Union = FTimestampStream::ComputeUnion(StreamPointers);
-				double Time = FPlatformTime::ToMilliseconds64(Union);
+				double Milliseconds = FPlatformTime::ToMilliseconds64(Union);
 
-				SET_FLOAT_STAT_FName(StatId.GetName(), Time);
+				SET_FLOAT_STAT_FName(Stat.StatId.GetName(), Milliseconds);
+
+			#if CSV_PROFILER_STATS
+				if (CsvProfiler)
+				{
+					CsvProfiler->RecordCustomStat(Stat.CsvStat, CSV_CATEGORY_INDEX(GPU), Milliseconds, ECsvCustomStatOp::Set);
+				}
+			#endif
 			}
-		#endif // WITH_RHI_BREADCRUMBS && HAS_GPU_STATS
+
+		#endif // WITH_RHI_BREADCRUMBS
 
 			// Compute the whole-frame total GPU time.
 			StreamPointers.Reset();
 			for (auto const& [Queue, State] : FrameState)
 			{
-				StreamPointers.Emplace(State.QueueTimestamps);
+				StreamPointers.Emplace(State.Queue);
 			}
 			uint64 WholeFrameUnion = FTimestampStream::ComputeUnion(StreamPointers);
 
 			// Update the global GPU frame time stats - need to convert to Cycles32 rather than Cycles64.
 			GGPUFrameTime = FPlatformMath::TruncToInt(FPlatformTime::ToSeconds64(WholeFrameUnion) / FPlatformTime::GetSecondsPerCycle());
+
+			RHISetGPUStatTotals(bCsvStatsEnabled, FPlatformTime::ToMilliseconds64(WholeFrameUnion));
 		}
 
 	} GGPUProfilerSink_StatSystem;
diff --git a/Engine/Source/Runtime/RHI/Private/GpuProfilerTrace.cpp b/Engine/Source/Runtime/RHI/Private/GpuProfilerTrace.cpp
--- a/Engine/Source/Runtime/RHI/Private/GpuProfilerTrace.cpp
+++ b/Engine/Source/Runtime/RHI/Private/GpuProfilerTrace.cpp
@@ -78,7 +78,7 @@ void FGpuProfilerTrace::BeginFrame(FGPUTimingCalibrationTimestamp& Calibration)
 	GCurrentFrame.EventBufferSize = 0;
 	GCurrentFrame.bActive = true;
 
-	int32 NeededSize = CVarGpuProfilerMaxEventBufferSizeKB.GetValueOnRenderThread() * 1024;
+	int32 NeededSize = CVarGpuProfilerMaxEventBufferSizeKB.GetValueOnAnyThread() * 1024;
 	if ((GCurrentFrame.MaxEventBufferSize != NeededSize) && (NeededSize > 0))
 	{
 		FMemory::Free(GCurrentFrame.EventBuffer);
diff --git a/Engine/Source/Runtime/RHI/Private/RHIBreadcrumbs.cpp b/Engine/Source/Runtime/RHI/Private/RHIBreadcrumbs.cpp
--- a/Engine/Source/Runtime/RHI/Private/RHIBreadcrumbs.cpp
+++ b/Engine/Source/Runtime/RHI/Private/RHIBreadcrumbs.cpp
@@ -12,7 +12,7 @@
 FRHIBreadcrumbNode::FRHIBreadcrumbNode()
 	: Name([]() -> FRHIBreadcrumb const&
 	{
-		static TRHIBreadcrumb<> Breadcrumb(TStatId(), TEXT("Sentinel"));
+		static TRHIBreadcrumb<> Breadcrumb(FRHIBreadcrumbData(__FILE__, __LINE__, TStatId(), NAME_None), TEXT("Sentinel"));
 		return Breadcrumb;
 	}())
 {}
diff --git a/Engine/Source/Runtime/RHI/Public/GPUProfiler.h b/Engine/Source/Runtime/RHI/Public/GPUProfiler.h
--- a/Engine/Source/Runtime/RHI/Public/GPUProfiler.h
+++ b/Engine/Source/Runtime/RHI/Public/GPUProfiler.h
@@ -8,9 +8,14 @@
 
 #include "CoreMinimal.h"
 #include "Misc/TVariant.h"
+#include "ProfilingDebugging/CsvProfiler.h"
 
 #include "RHIBreadcrumbs.h"
 
+#if HAS_GPU_STATS
+CSV_DECLARE_CATEGORY_MODULE_EXTERN(RHI_API, GPU);
+#endif
+
 #if RHI_NEW_GPU_PROFILER
 
 namespace UE::RHI::GPUProfiler
diff --git a/Engine/Source/Runtime/RHI/Public/GpuProfilerTrace.h b/Engine/Source/Runtime/RHI/Public/GpuProfilerTrace.h
--- a/Engine/Source/Runtime/RHI/Public/GpuProfilerTrace.h
+++ b/Engine/Source/Runtime/RHI/Public/GpuProfilerTrace.h
@@ -6,7 +6,7 @@
 #include "Trace/Config.h"
 
 #if !defined(GPUPROFILERTRACE_ENABLED)
-#if UE_TRACE_ENABLED && !UE_BUILD_SHIPPING && (RHI_NEW_GPU_PROFILER == 0)
+#if UE_TRACE_ENABLED && !UE_BUILD_SHIPPING
 #define GPUPROFILERTRACE_ENABLED 1
 #else
 #define GPUPROFILERTRACE_ENABLED 0
@@ -15,6 +15,17 @@
 
 #if GPUPROFILERTRACE_ENABLED
 
+#if RHI_NEW_GPU_PROFILER
+// Define this structure here when the new GPU profiler is enabled so we can still build the old trace API.
+// @todo - remove this. GPU timestamp calibration is no longer necessary with the new GPU profiler, as the
+// platform RHIs are expected to translate timestamps from GPU to CPU clock domain before they reach the profiler.
+struct FGPUTimingCalibrationTimestamp
+{
+	uint64 GPUMicroseconds = 0;
+	uint64 CPUMicroseconds = 0;
+};
+#endif
+
 class FName;
 
 struct FGpuProfilerTrace
diff --git a/Engine/Source/Runtime/RHI/Public/RHIBreadcrumbs.h b/Engine/Source/Runtime/RHI/Public/RHIBreadcrumbs.h
--- a/Engine/Source/Runtime/RHI/Public/RHIBreadcrumbs.h
+++ b/Engine/Source/Runtime/RHI/Public/RHIBreadcrumbs.h
@@ -6,6 +6,7 @@
 #include "Misc/MemStack.h"
 #include "HAL/Platform.h"
 #include "Templates/SharedPointer.h"
+#include "ProfilingDebugging/CsvProfiler.h"
 #include "GenericPlatform/GenericPlatformCrashContext.h"
 
 #include "RHIFwd.h"
@@ -19,8 +20,101 @@
 // Whether to emit Unreal Insights breadcrumb events on threads involved in RHI command list recording and execution.
 #define RHI_BREADCRUMBS_EMIT_CPU (WITH_RHI_BREADCRUMBS && CPUPROFILERTRACE_ENABLED && 1)
 
+// Whether to store the filename and line number of each RHI breadcrumb and emit this data to Insights.
+#define RHI_BREADCRUMBS_EMIT_LOCATION (WITH_RHI_BREADCRUMBS && (CPUPROFILERTRACE_ENABLED || GPUPROFILERTRACE_ENABLED) && 1)
+
 #if WITH_RHI_BREADCRUMBS
 
+	//
+	// Holds the filename and line number location of the RHI breadcrumb in source.
+	//
+	struct FRHIBreadcrumbData_Location
+	{
+#if RHI_BREADCRUMBS_EMIT_LOCATION
+		ANSICHAR const* File;
+		uint32 Line;
+#endif
+
+		FRHIBreadcrumbData_Location(ANSICHAR const* File, uint32 Line)
+#if RHI_BREADCRUMBS_EMIT_LOCATION
+			: File(File)
+			, Line(Line)
+#endif
+		{}
+	};
+
+	//
+	// Holds both a stats system ID, and a CSV profiler ID.
+	// The computed stat value is emitted to both "stat gpu" and the CSV profiler.
+	//
+	struct FRHIBreadcrumbData_Stats
+	{
+#if STATS
+		TStatId StatId {};
+#endif
+#if CSV_PROFILER_STATS
+		FName CsvStat = NAME_None;
+#endif
+
+		FRHIBreadcrumbData_Stats(TStatId InStatId, FName InCsvStat)
+		{
+#if STATS
+			StatId = InStatId;
+#endif
+#if CSV_PROFILER_STATS
+			CsvStat = InCsvStat;
+#endif
+		}
+
+		bool ShouldComputeStat() const
+		{
+#if STATS
+			return StatId.IsValidStat();
+#elif CSV_PROFILER_STATS
+			return CsvStat != NAME_None;
+#else
+			return false;
+#endif
+		}
+
+		bool operator == (FRHIBreadcrumbData_Stats const& RHS) const
+		{
+#if STATS
+			return StatId == RHS.StatId;
+#elif CSV_PROFILER_STATS
+			return CsvStat == RHS.CsvStat;
+#else
+			return true;
+#endif
+		}
+
+		friend uint32 GetTypeHash(FRHIBreadcrumbData_Stats const& Stats)
+		{
+#if STATS
+			return GetTypeHash(Stats.StatId);
+#elif CSV_PROFILER_STATS
+			return GetTypeHash(Stats.CsvStat);
+#else
+			return 0;
+#endif
+		}
+	};
+
+	//
+	// Container for extra profiling-related data for each RHI breadcrumb.
+	//
+	class FRHIBreadcrumbData
+		// Use inheritance for empty-base-optimization.
+		: public FRHIBreadcrumbData_Location
+		, public FRHIBreadcrumbData_Stats
+	{
+	public:
+		FRHIBreadcrumbData(ANSICHAR const* File, uint32 Line, TStatId StatId, FName CsvStat)
+			: FRHIBreadcrumbData_Location(File, Line)
+			, FRHIBreadcrumbData_Stats(StatId, CsvStat)
+		{}
+	};
+
 	class FRHIBreadcrumb;
 	class FRHIBreadcrumbAllocator;
 	struct FRHIBreadcrumbRange;
@@ -156,7 +250,7 @@
 		FRHIBreadcrumbAllocatorArray const& GetParents() const { return Parents; }
 
 		template<size_t N, typename... TArgs>
-		inline FRHIBreadcrumbNode* AllocBreadcrumb(TStatId StatId, TCHAR const(&FormatString)[N], TArgs&&... Args);
+		inline FRHIBreadcrumbNode* AllocBreadcrumb(FRHIBreadcrumbData&& Stat, TCHAR const(&FormatString)[N], TArgs&&... Args);
 
 		inline void* Alloc(uint32 Size, uint32 Align)
 		{
@@ -433,9 +527,9 @@
 		RHI_API static std::atomic<uint32> NextID;
 
 	protected:
-		FRHIBreadcrumb(TStatId StatId)
+		FRHIBreadcrumb(FRHIBreadcrumbData&& Data)
 			: ID(NextID.fetch_add(1, std::memory_order_relaxed) | 0x80000000) // Set the top bit to avoid collision with zero (i.e. "no breadcrumb")
-			, StatId(StatId)
+			, Data(Data)
 		{}
 
 		FRHIBreadcrumb(FRHIBreadcrumb const&) = delete;
@@ -443,7 +537,7 @@
 
 	public:
 		uint32 const ID;
-		TStatId const StatId;
+		FRHIBreadcrumbData const Data;
 
 #if RHI_BREADCRUMBS_EMIT_CPU
 		uint32 const CPUTraceMarkerID = 0;
@@ -456,7 +550,11 @@
 			{
 				FBuffer Buffer;
 				TCHAR const* Str = GetTCHAR(Buffer);
-				const_cast<uint32&>(CPUTraceMarkerID) = FCpuProfilerTrace::OutputDynamicEventType(Str, __FILE__, __LINE__);
+				const_cast<uint32&>(CPUTraceMarkerID) = FCpuProfilerTrace::OutputDynamicEventType(Str
+			#if RHI_BREADCRUMBS_EMIT_LOCATION
+					, Data.File, Data.Line
+			#endif
+				);
 			}
 #endif
 		}
@@ -657,8 +755,8 @@
 		TRHIBreadcrumb(TRHIBreadcrumb const& Other) = delete;
 
 	public:
-		TRHIBreadcrumb(TStatId StatId, TCHAR const* FormatString, TArgs const&... Args)
-			: FRHIBreadcrumb(StatId)
+		TRHIBreadcrumb(FRHIBreadcrumbData&& Data, TCHAR const* FormatString, TArgs const&... Args)
+			: FRHIBreadcrumb(MoveTemp(Data))
 			, FormatString(FormatString)
 			, Values(Args...)
 		{
@@ -688,8 +786,8 @@
 		TRHIBreadcrumb(TRHIBreadcrumb const& Other) = delete;
 
 	public:
-		TRHIBreadcrumb(TStatId StatId, TCHAR const* StringLiteral)
-			: FRHIBreadcrumb(StatId)
+		TRHIBreadcrumb(FRHIBreadcrumbData&& Data, TCHAR const* StringLiteral)
+			: FRHIBreadcrumb(MoveTemp(Data))
 			, StringLiteral(StringLiteral)
 		{
 			CreateTraceMarkers();
@@ -707,20 +805,20 @@
 	};
 
 	template<size_t N, typename... TArgs>
-	inline FRHIBreadcrumbNode* FRHIBreadcrumbAllocator::AllocBreadcrumb(TStatId StatId, TCHAR const(&FormatString)[N], TArgs&&... Args)
+	inline FRHIBreadcrumbNode* FRHIBreadcrumbAllocator::AllocBreadcrumb(FRHIBreadcrumbData&& Data, TCHAR const(&FormatString)[N], TArgs&&... Args)
 	{
 		struct FStorage
 		{
 			TRHIBreadcrumb<std::decay_t<TArgs>...> Name;
 			FRHIBreadcrumbNode Node;
 
-			FStorage(FRHIBreadcrumbAllocator& Allocator, TStatId StatId, TCHAR const(&FormatString)[N], TArgs&&... Args)
-				: Name(StatId, FormatString, Forward<TArgs>(Args)...)
+			FStorage(FRHIBreadcrumbAllocator& Allocator, FRHIBreadcrumbData&& Data, TCHAR const(&FormatString)[N], TArgs&&... Args)
+				: Name(MoveTemp(Data), FormatString, Forward<TArgs>(Args)...)
 				, Node(Allocator, Name)
 			{}
 		};
 
-		return &Alloc<FStorage>(*this, StatId, FormatString, Forward<TArgs>(Args)...)->Node;
+		return &Alloc<FStorage>(*this, MoveTemp(Data), FormatString, Forward<TArgs>(Args)...)->Node;
 	}
 
 	class FRHIBreadcrumbNodeRef
@@ -763,7 +861,7 @@
 
 	public:
 		template<size_t N, typename... TArgs>
-		inline FRHIBreadcrumbEventManual(FRHIComputeCommandList& RHICmdList, TStatId StatId, TCHAR const(&FormatString)[N], TArgs&&... Args);
+		inline FRHIBreadcrumbEventManual(FRHIComputeCommandList& RHICmdList, FRHIBreadcrumbData&& Data, TCHAR const(&FormatString)[N], TArgs&&... Args);
 
 		inline void End(FRHIComputeCommandList& RHICmdList);
 
@@ -783,33 +881,33 @@
 		FRHIBreadcrumbEventScope(FRHIBreadcrumbEventScope&&) = delete;
 
 		template<size_t N, typename... TArgs>
-		inline FRHIBreadcrumbEventScope(FRHIComputeCommandList& InRHICmdList, TStatId StatId, ERHIPipeline InPipeline, bool bCondition, TCHAR const(&FormatString)[N], TArgs&&... Args);
+		inline FRHIBreadcrumbEventScope(FRHIComputeCommandList& InRHICmdList, FRHIBreadcrumbData&& Data, ERHIPipeline InPipeline, bool bCondition, TCHAR const(&FormatString)[N], TArgs&&... Args);
 
 	public:
 		// Top-of-pipe breadcrumb event scope for RHI command lists
 		template<size_t N, typename... TArgs>
-		inline FRHIBreadcrumbEventScope(FRHIComputeCommandList& InRHICmdList, TStatId StatId, bool bCondition, TCHAR const(&FormatString)[N], TArgs&&... Args);
+		inline FRHIBreadcrumbEventScope(FRHIComputeCommandList& InRHICmdList, FRHIBreadcrumbData&& Data, bool bCondition, TCHAR const(&FormatString)[N], TArgs&&... Args);
 
 		// Bottom-of-pipe breadcrumb event scope for RHI contexts
 		template<size_t N, typename... TArgs>
-		inline FRHIBreadcrumbEventScope(IRHIComputeContext& InRHIContext, TStatId StatId, bool bCondition, TCHAR const(&FormatString)[N], TArgs&&... Args);
+		inline FRHIBreadcrumbEventScope(IRHIComputeContext& InRHIContext, FRHIBreadcrumbData&& Data, bool bCondition, TCHAR const(&FormatString)[N], TArgs&&... Args);
 
 		inline ~FRHIBreadcrumbEventScope();
 	};
 
-	#define RHI_BREADCRUMB_EVENT(                 RHICmdList_Or_RHIContext,                  Format, ...) FRHIBreadcrumbEventScope ANONYMOUS_VARIABLE(BreadcrumbEvent)(RHICmdList_Or_RHIContext, TStatId()                  , true     , TEXT(Format), ## __VA_ARGS__)
-	#define RHI_BREADCRUMB_EVENT_CONDITIONAL(     RHICmdList_Or_RHIContext,       Condition, Format, ...) FRHIBreadcrumbEventScope ANONYMOUS_VARIABLE(BreadcrumbEvent)(RHICmdList_Or_RHIContext, TStatId()                  , Condition, TEXT(Format), ## __VA_ARGS__)
+	#define RHI_BREADCRUMB_EVENT(                 RHICmdList_Or_RHIContext,                  Format, ...) FRHIBreadcrumbEventScope ANONYMOUS_VARIABLE(BreadcrumbEvent)(RHICmdList_Or_RHIContext, FRHIBreadcrumbData(__FILE__, __LINE__, TStatId(), NAME_None), true     , TEXT(Format), ## __VA_ARGS__)
+	#define RHI_BREADCRUMB_EVENT_CONDITIONAL(     RHICmdList_Or_RHIContext,       Condition, Format, ...) FRHIBreadcrumbEventScope ANONYMOUS_VARIABLE(BreadcrumbEvent)(RHICmdList_Or_RHIContext, FRHIBreadcrumbData(__FILE__, __LINE__, TStatId(), NAME_None), Condition, TEXT(Format), ## __VA_ARGS__)
 #if HAS_GPU_STATS
-	#define RHI_BREADCRUMB_EVENT_STAT(            RHICmdList_Or_RHIContext, Stat,            Format, ...) FRHIBreadcrumbEventScope ANONYMOUS_VARIABLE(BreadcrumbEvent)(RHICmdList_Or_RHIContext, GET_STATID(Stat_GPU_##Stat), true     , TEXT(Format), ## __VA_ARGS__)
-	#define RHI_BREADCRUMB_EVENT_CONDITIONAL_STAT(RHICmdList_Or_RHIContext, Stat, Condition, Format, ...) FRHIBreadcrumbEventScope ANONYMOUS_VARIABLE(BreadcrumbEvent)(RHICmdList_Or_RHIContext, GET_STATID(Stat_GPU_##Stat), Condition, TEXT(Format), ## __VA_ARGS__)
+	#define RHI_BREADCRUMB_EVENT_STAT(            RHICmdList_Or_RHIContext, Stat,            Format, ...) FRHIBreadcrumbEventScope ANONYMOUS_VARIABLE(BreadcrumbEvent)(RHICmdList_Or_RHIContext, FRHIBreadcrumbData(__FILE__, __LINE__, GET_STATID(Stat_GPU_##Stat), CSV_STAT_FNAME(Stat)), true     , TEXT(Format), ## __VA_ARGS__)
+	#define RHI_BREADCRUMB_EVENT_CONDITIONAL_STAT(RHICmdList_Or_RHIContext, Stat, Condition, Format, ...) FRHIBreadcrumbEventScope ANONYMOUS_VARIABLE(BreadcrumbEvent)(RHICmdList_Or_RHIContext, FRHIBreadcrumbData(__FILE__, __LINE__, GET_STATID(Stat_GPU_##Stat), CSV_STAT_FNAME(Stat)), Condition, TEXT(Format), ## __VA_ARGS__)
 #else
-	#define RHI_BREADCRUMB_EVENT_STAT(            RHICmdList_Or_RHIContext, Stat,            Format, ...) FRHIBreadcrumbEventScope ANONYMOUS_VARIABLE(BreadcrumbEvent)(RHICmdList_Or_RHIContext, TStatId()                  , true     , TEXT(Format), ## __VA_ARGS__)
-	#define RHI_BREADCRUMB_EVENT_CONDITIONAL_STAT(RHICmdList_Or_RHIContext, Stat, Condition, Format, ...) FRHIBreadcrumbEventScope ANONYMOUS_VARIABLE(BreadcrumbEvent)(RHICmdList_Or_RHIContext, TStatId()                  , Condition, TEXT(Format), ## __VA_ARGS__)
+	#define RHI_BREADCRUMB_EVENT_STAT(            RHICmdList_Or_RHIContext, Stat,            Format, ...) FRHIBreadcrumbEventScope ANONYMOUS_VARIABLE(BreadcrumbEvent)(RHICmdList_Or_RHIContext, FRHIBreadcrumbData(__FILE__, __LINE__, TStatId(), NAME_None), true     , TEXT(Format), ## __VA_ARGS__)
+	#define RHI_BREADCRUMB_EVENT_CONDITIONAL_STAT(RHICmdList_Or_RHIContext, Stat, Condition, Format, ...) FRHIBreadcrumbEventScope ANONYMOUS_VARIABLE(BreadcrumbEvent)(RHICmdList_Or_RHIContext, FRHIBreadcrumbData(__FILE__, __LINE__, TStatId(), NAME_None), Condition, TEXT(Format), ## __VA_ARGS__)
 #endif
 
 	// Used only for back compat with SCOPED_DRAW_EVENTF
-	#define RHI_BREADCRUMB_EVENT_STR_DEPRECATED(            RHICmdList_Or_RHIContext,            Format, ...) FRHIBreadcrumbEventScope ANONYMOUS_VARIABLE(BreadcrumbEvent)(RHICmdList_Or_RHIContext, TStatId(), true     , Format, ## __VA_ARGS__)
-	#define RHI_BREADCRUMB_EVENT_CONDITIONAL_STR_DEPRECATED(RHICmdList_Or_RHIContext, Condition, Format, ...) FRHIBreadcrumbEventScope ANONYMOUS_VARIABLE(BreadcrumbEvent)(RHICmdList_Or_RHIContext, TStatId(), Condition, Format, ## __VA_ARGS__)
+	#define RHI_BREADCRUMB_EVENT_STR_DEPRECATED(            RHICmdList_Or_RHIContext,            Format, ...) FRHIBreadcrumbEventScope ANONYMOUS_VARIABLE(BreadcrumbEvent)(RHICmdList_Or_RHIContext, FRHIBreadcrumbData(__FILE__, __LINE__, TStatId(), NAME_None), true     , Format, ## __VA_ARGS__)
+	#define RHI_BREADCRUMB_EVENT_CONDITIONAL_STR_DEPRECATED(RHICmdList_Or_RHIContext, Condition, Format, ...) FRHIBreadcrumbEventScope ANONYMOUS_VARIABLE(BreadcrumbEvent)(RHICmdList_Or_RHIContext, FRHIBreadcrumbData(__FILE__, __LINE__, TStatId(), NAME_None), Condition, Format, ## __VA_ARGS__)
 
 #else // WITH_RHI_BREADCRUMBS == 0
 
diff --git a/Engine/Source/Runtime/RHI/Public/RHICommandList.inl b/Engine/Source/Runtime/RHI/Public/RHICommandList.inl
--- a/Engine/Source/Runtime/RHI/Public/RHICommandList.inl
+++ b/Engine/Source/Runtime/RHI/Public/RHICommandList.inl
@@ -138,20 +138,20 @@ inline FRHIResourceReplaceBatcher::~FRHIResourceReplaceBatcher()
 
 	// Top-of-pipe breadcrumb event scope for RHI command lists
 	template<size_t N, typename... TArgs>
-	inline FRHIBreadcrumbEventScope::FRHIBreadcrumbEventScope(FRHIComputeCommandList& InRHICmdList, TStatId StatId, bool bCondition, TCHAR const(&FormatString)[N], TArgs&&... Args)
-		: FRHIBreadcrumbEventScope(InRHICmdList, StatId, InRHICmdList.GetPipeline(), bCondition, FormatString, Forward<TArgs>(Args)...)
+	inline FRHIBreadcrumbEventScope::FRHIBreadcrumbEventScope(FRHIComputeCommandList& InRHICmdList, FRHIBreadcrumbData&& Data, bool bCondition, TCHAR const(&FormatString)[N], TArgs&&... Args)
+		: FRHIBreadcrumbEventScope(InRHICmdList, MoveTemp(Data), InRHICmdList.GetPipeline(), bCondition, FormatString, Forward<TArgs>(Args)...)
 	{}
 
 	// Bottom-of-pipe breadcrumb event scope for RHI contexts
 	template<size_t N, typename... TArgs>
-	inline FRHIBreadcrumbEventScope::FRHIBreadcrumbEventScope(IRHIComputeContext& InRHIContext, TStatId StatId, bool bCondition, TCHAR const(&FormatString)[N], TArgs&&... Args)
-		: FRHIBreadcrumbEventScope(static_cast<FRHIComputeCommandList&>(InRHIContext.GetExecutingCommandList()), StatId, InRHIContext.GetPipeline(), bCondition, FormatString, Forward<TArgs>(Args)...)
+	inline FRHIBreadcrumbEventScope::FRHIBreadcrumbEventScope(IRHIComputeContext& InRHIContext, FRHIBreadcrumbData&& Data, bool bCondition, TCHAR const(&FormatString)[N], TArgs&&... Args)
+		: FRHIBreadcrumbEventScope(static_cast<FRHIComputeCommandList&>(InRHIContext.GetExecutingCommandList()), MoveTemp(Data), InRHIContext.GetPipeline(), bCondition, FormatString, Forward<TArgs>(Args)...)
 	{}
 
 	template<size_t N, typename... TArgs>
-	inline FRHIBreadcrumbEventScope::FRHIBreadcrumbEventScope(FRHIComputeCommandList& InRHICmdList, TStatId StatId, ERHIPipeline InPipeline, bool bCondition, TCHAR const(&FormatString)[N], TArgs&&... Args)
+	inline FRHIBreadcrumbEventScope::FRHIBreadcrumbEventScope(FRHIComputeCommandList& InRHICmdList, FRHIBreadcrumbData&& Data, ERHIPipeline InPipeline, bool bCondition, TCHAR const(&FormatString)[N], TArgs&&... Args)
 		: RHICmdList(InRHICmdList)
-		, Node(bCondition ? RHICmdList.GetBreadcrumbAllocator().AllocBreadcrumb(StatId, FormatString, Forward<TArgs>(Args)...) : nullptr)
+		, Node(bCondition ? RHICmdList.GetBreadcrumbAllocator().AllocBreadcrumb(MoveTemp(Data), FormatString, Forward<TArgs>(Args)...) : nullptr)
 		, Pipeline(InPipeline)
 	{
 		if (Node)
@@ -172,8 +172,8 @@ inline FRHIResourceReplaceBatcher::~FRHIResourceReplaceBatcher()
 	}
 
 	template<size_t N, typename... TArgs>
-	inline FRHIBreadcrumbEventManual::FRHIBreadcrumbEventManual(FRHIComputeCommandList& RHICmdList, TStatId StatId, TCHAR const(&FormatString)[N], TArgs&&... Args)
-		: Node(RHICmdList.GetBreadcrumbAllocator().AllocBreadcrumb(StatId, FormatString, Forward<TArgs>(Args)...))
+	inline FRHIBreadcrumbEventManual::FRHIBreadcrumbEventManual(FRHIComputeCommandList& RHICmdList, FRHIBreadcrumbData&& Data, TCHAR const(&FormatString)[N], TArgs&&... Args)
+		: Node(RHICmdList.GetBreadcrumbAllocator().AllocBreadcrumb(MoveTemp(Data), FormatString, Forward<TArgs>(Args)...))
 	#if DO_CHECK
 		, Pipeline(RHICmdList.GetPipeline())
 		, ThreadId(FPlatformTLS::GetCurrentThreadId())
diff --git a/Engine/Source/Runtime/RenderCore/Private/ProfilingDebugging/RealtimeGPUProfiler.cpp b/Engine/Source/Runtime/RenderCore/Private/ProfilingDebugging/RealtimeGPUProfiler.cpp
--- a/Engine/Source/Runtime/RenderCore/Private/ProfilingDebugging/RealtimeGPUProfiler.cpp
+++ b/Engine/Source/Runtime/RenderCore/Private/ProfilingDebugging/RealtimeGPUProfiler.cpp
@@ -11,26 +11,16 @@
 #if RHI_NEW_GPU_PROFILER
 
 // @todo - new gpu profiler
-CSV_DEFINE_CATEGORY_MODULE(RENDERCORE_API, GPU, true);
 
 #elif HAS_GPU_STATS
 
 typedef TArray<TCHAR, TInlineAllocator<4096u>> FDescriptionStringBuffer;
 
-CSV_DEFINE_CATEGORY_MODULE(RENDERCORE_API, GPU, true);
-
 static TAutoConsoleVariable<int> CVarGPUStatsEnabled(
 	TEXT("r.GPUStatsEnabled"),
 	1,
 	TEXT("Enables or disables GPU stat recording"));
 
-static TAutoConsoleVariable<int> CVarGPUCsvStatsEnabled(
-	TEXT("r.GPUCsvStatsEnabled"),
-	0,
-	TEXT("Enables or disables GPU stat recording to CSVs"));
-
-DECLARE_GPU_STAT_NAMED( Total, TEXT("[TOTAL]") );
-
 static TAutoConsoleVariable<int> CVarGPUTracingStatsEnabled(
 	TEXT("r.GPUTracingStatsEnabled"),
 	1,
@@ -472,8 +462,11 @@ public:
 
 		// Update the stats
 #if CSV_PROFILER_STATS
-		const bool bCsvStatsEnabled = !!CVarGPUCsvStatsEnabled.GetValueOnRenderThread();
+		static IConsoleVariable* CVarGPUCsvStatsEnabled = IConsoleManager::Get().FindConsoleVariable(TEXT("r.GPUCsvStatsEnabled"));
+		const bool bCsvStatsEnabled = CVarGPUCsvStatsEnabled ? CVarGPUCsvStatsEnabled->GetBool() : false;
 		FCsvProfiler* CsvProfiler = bCsvStatsEnabled ? FCsvProfiler::Get() : nullptr;
+#else
+		const bool bCsvStatsEnabled = false;
 #endif
 		const bool GPUStatsChildTimesIncluded = !!CVarGPUStatsChildTimesIncluded.GetValueOnRenderThread();
 		uint64 TotalUs = 0llu;
@@ -517,18 +510,8 @@ public:
 			}
 		}
 
-#if STATS
-		const double TotalMs = TotalUs / 1000.;
-		FThreadStats::AddMessage(GET_STATFNAME(Stat_GPU_Total), EStatOperation::Set, TotalMs);
-		TRACE_STAT_SET(GET_STATFNAME(Stat_GPU_Total), TotalMs);
-#endif
-
-#if CSV_PROFILER_STATS
-		if (CsvProfiler)
-		{
-			CsvProfiler->RecordCustomStat(CSV_STAT_FNAME(Total), CSV_CATEGORY_INDEX(GPU), TotalUs / 1000.f, ECsvCustomStatOp::Set);
-		}
-#endif
+		extern RHI_API void RHISetGPUStatTotals(bool bCsvStatsEnabled, double TotalMs);
+		RHISetGPUStatTotals(bCsvStatsEnabled, TotalUs / 1000.0);
 
 #if GPUPROFILERTRACE_ENABLED
 		TArray<TArray<int32>> GpuProfilerEventChildrenIndices;
@@ -903,7 +886,8 @@ void LatchAreGPUStatsEnabled()
 #else
 
 		// If we only have CSV stats, only capture if CSV GPU stats are enabled, and we're capturing
-		if (!CVarGPUCsvStatsEnabled.GetValueOnRenderThread())
+		static IConsoleVariable* CVarGPUCsvStatsEnabled = IConsoleManager::Get().FindConsoleVariable(TEXT("r.GPUCsvStatsEnabled"));
+		if (!CVarGPUCsvStatsEnabled || !CVarGPUCsvStatsEnabled->GetBool())
 		{
 			return false;
 		}
diff --git a/Engine/Source/Runtime/RenderCore/Private/RenderCaptureInterface.cpp b/Engine/Source/Runtime/RenderCore/Private/RenderCaptureInterface.cpp
--- a/Engine/Source/Runtime/RenderCore/Private/RenderCaptureInterface.cpp
+++ b/Engine/Source/Runtime/RenderCore/Private/RenderCaptureInterface.cpp
@@ -35,7 +35,7 @@ namespace RenderCaptureInterface
 #if WITH_RHI_BREADCRUMBS
 				if (bPushEvent)
 				{
-					Breadcrumb->Emplace(RHICommandListLocal, TStatId(), TEXT("%s"), EventName);
+					Breadcrumb->Emplace(RHICommandListLocal, FRHIBreadcrumbData(__FILE__, __LINE__, TStatId(), NAME_None), TEXT("%s"), EventName);
 				}
 #endif
 			});
@@ -58,7 +58,7 @@ namespace RenderCaptureInterface
 			if (bEvent)
 			{
 				RHIBreadcrumb = MakeUnique<TOptional<FRHIBreadcrumbEventManual>>();
-				RHIBreadcrumb->Emplace(*RHICommandList, TStatId(), TEXT("%s"), FString(InEventName));
+				RHIBreadcrumb->Emplace(*RHICommandList, FRHIBreadcrumbData(__FILE__, __LINE__, TStatId(), NAME_None), TEXT("%s"), FString(InEventName));
 			}
 #endif
 		}
@@ -85,7 +85,7 @@ namespace RenderCaptureInterface
 #if RDG_EVENTS
 			if (bEvent)
 			{
-				RDGEvent.Emplace(*GraphBuilder, ERDGScopeFlags::None, TStatId(), RDG_EVENT_NAME("%s", InEventName));
+				RDGEvent.Emplace(*GraphBuilder, ERDGScopeFlags::None, FRHIBreadcrumbData(__FILE__, __LINE__, TStatId(), NAME_None), RDG_EVENT_NAME("%s", InEventName));
 			}
 #endif
 		}
diff --git a/Engine/Source/Runtime/RenderCore/Private/RenderGraphBuilder.cpp b/Engine/Source/Runtime/RenderCore/Private/RenderGraphBuilder.cpp
--- a/Engine/Source/Runtime/RenderCore/Private/RenderGraphBuilder.cpp
+++ b/Engine/Source/Runtime/RenderCore/Private/RenderGraphBuilder.cpp
@@ -2589,7 +2589,7 @@ FRDGPass* FRDGBuilder::SetupParameterPass(FRDGPass* Pass)
 	if (ScopeState.ScopeMode == ERDGScopeMode::AllEventsAndPassNames)
 	{
 		FRDGEventName Name = Pass->GetEventName();
-		PassNameScope.Emplace(*this, ERDGScopeFlags::None, TStatId(), MoveTemp(Name));
+		PassNameScope.Emplace(*this, ERDGScopeFlags::None, FRHIBreadcrumbData(__FILE__, __LINE__, TStatId(), NAME_None), MoveTemp(Name));
 	}
 #endif
 
diff --git a/Engine/Source/Runtime/RenderCore/Private/RenderGraphEvent.cpp b/Engine/Source/Runtime/RenderCore/Private/RenderGraphEvent.cpp
--- a/Engine/Source/Runtime/RenderCore/Private/RenderGraphEvent.cpp
+++ b/Engine/Source/Runtime/RenderCore/Private/RenderGraphEvent.cpp
@@ -286,7 +286,7 @@ FRDGEventName::FRDGEventName(const TCHAR* EventFormat, ...)
 }
 
 #if WITH_RHI_BREADCRUMBS
-FRHIBreadcrumbNode* FRDGEventName::AllocBreadcrumb(TStatId StatId, FRHIBreadcrumbAllocator& Allocator) const
+FRHIBreadcrumbNode* FRDGEventName::AllocBreadcrumb(FRHIBreadcrumbData&& Data, FRHIBreadcrumbAllocator& Allocator) const
 {
 #if RDG_EVENTS == RDG_EVENTS_STRING_COPY
 
@@ -297,7 +297,7 @@ FRHIBreadcrumbNode* FRDGEventName::AllocBreadcrumb(TStatId StatId, FRHIBreadcrum
 				: TEXT("<unnamed pass>");
 			
 			// Cast hack to force treating EventFormat as a string literal.
-			return Allocator.AllocBreadcrumb(StatId, *reinterpret_cast<TCHAR const(*)[1]>(String));
+			return Allocator.AllocBreadcrumb(MoveTemp(Data), *reinterpret_cast<TCHAR const(*)[1]>(String));
 		}
 		else
 		{
@@ -310,13 +310,13 @@ FRHIBreadcrumbNode* FRDGEventName::AllocBreadcrumb(TStatId StatId, FRHIBreadcrum
 			FMemory::Memcpy(StrCopy, *FormattedEventName, Size);
 
 			// Cast hack to force treating StrCopy as a string literal.
-			return Allocator.AllocBreadcrumb(StatId, *reinterpret_cast<TCHAR const(*)[1]>(StrCopy));
+			return Allocator.AllocBreadcrumb(MoveTemp(Data), *reinterpret_cast<TCHAR const(*)[1]>(StrCopy));
 		}
 
 #elif RDG_EVENTS >= RDG_EVENTS_STRING_REF
 
 		// Cast hack to force treating EventFormat as a string literal.
-		return Allocator.AllocBreadcrumb(StatId, *reinterpret_cast<TCHAR const(*)[1]>(EventFormat));
+		return Allocator.AllocBreadcrumb(MoveTemp(Data), *reinterpret_cast<TCHAR const(*)[1]>(EventFormat));
 
 #else
 
diff --git a/Engine/Source/Runtime/RenderCore/Public/ProfilingDebugging/RealtimeGPUProfiler.h b/Engine/Source/Runtime/RenderCore/Public/ProfilingDebugging/RealtimeGPUProfiler.h
--- a/Engine/Source/Runtime/RenderCore/Public/ProfilingDebugging/RealtimeGPUProfiler.h
+++ b/Engine/Source/Runtime/RenderCore/Public/ProfilingDebugging/RealtimeGPUProfiler.h
@@ -33,7 +33,7 @@
 
 	public:
 		template<size_t N, typename... TArgs>
-		FRHIBreadcrumbEvent_GameThread(bool bCondition, TCHAR const(&FormatString)[N], TArgs&&... Args)
+		FRHIBreadcrumbEvent_GameThread(bool bCondition, FRHIBreadcrumbData&& Data, TCHAR const(&FormatString)[N], TArgs&&... Args)
 			: Event(bCondition ? new TOptional<FRHIBreadcrumbEventScope> : nullptr)
 		{
 			check(IsInGameThread());
@@ -42,13 +42,14 @@
 				ENQUEUE_RENDER_COMMAND(FRHIBreadcrumbEvent_GameThread_Begin)(
 				[
 					Event = Event,
+					Data = MoveTemp(Data),
 					Values = std::make_tuple(std::forward<TArgs>(Args)...),
 					&FormatString
 				](FRHICommandListImmediate& RHICmdList) mutable
 				{
-					std::apply([Event, &RHICmdList, &FormatString](auto&&... InnerArgs)
+					std::apply([Event, &RHICmdList, &Data, &FormatString](auto&&... InnerArgs)
 					{
-						Event->Emplace(RHICmdList, TStatId(), true, FormatString, std::forward<TArgs>(InnerArgs)...);
+						Event->Emplace(RHICmdList, MoveTemp(Data), true, FormatString, std::forward<TArgs>(InnerArgs)...);
 					}, Values);
 				});
 			}
@@ -67,11 +68,11 @@
 		}
 	};
 
-	#define RHI_BREADCRUMB_EVENT_GAMETHREAD(                       Format, ...) FRHIBreadcrumbEvent_GameThread ANONYMOUS_VARIABLE(BreadcrumbEvent_GameThread)(true     , TEXT(Format), ##__VA_ARGS__)
-	#define RHI_BREADCRUMB_EVENT_CONDITIONAL_GAMETHREAD(Condition, Format, ...) FRHIBreadcrumbEvent_GameThread ANONYMOUS_VARIABLE(BreadcrumbEvent_GameThread)(Condition, TEXT(Format), ##__VA_ARGS__)
+	#define RHI_BREADCRUMB_EVENT_GAMETHREAD(                       Format, ...) FRHIBreadcrumbEvent_GameThread ANONYMOUS_VARIABLE(BreadcrumbEvent_GameThread)(true     , FRHIBreadcrumbData(__FILE__, __LINE__, TStatId(), NAME_None), TEXT(Format), ##__VA_ARGS__)
+	#define RHI_BREADCRUMB_EVENT_CONDITIONAL_GAMETHREAD(Condition, Format, ...) FRHIBreadcrumbEvent_GameThread ANONYMOUS_VARIABLE(BreadcrumbEvent_GameThread)(Condition, FRHIBreadcrumbData(__FILE__, __LINE__, TStatId(), NAME_None), TEXT(Format), ##__VA_ARGS__)
 
 	// Used only for back compat with SCOPED_DRAW_EVENTF_GAMETHREAD
-	#define RHI_BREADCRUMB_EVENT_GAMETHREAD_STR_DEPRECATED(Format, ...) FRHIBreadcrumbEvent_GameThread ANONYMOUS_VARIABLE(BreadcrumbEvent_GameThread)(true, Format, ##__VA_ARGS__)
+	#define RHI_BREADCRUMB_EVENT_GAMETHREAD_STR_DEPRECATED(Format, ...) FRHIBreadcrumbEvent_GameThread ANONYMOUS_VARIABLE(BreadcrumbEvent_GameThread)(true, FRHIBreadcrumbData(__FILE__, __LINE__, TStatId(), NAME_None), Format, ##__VA_ARGS__)
 
 #else
 
@@ -141,8 +142,6 @@
 
 	#if HAS_GPU_STATS
 
-		CSV_DECLARE_CATEGORY_MODULE_EXTERN(RENDERCORE_API, GPU);
-
 		// @todo
 		#define DECLARE_GPU_STAT(                StatName            ) DECLARE_FLOAT_COUNTER_STAT(TEXT(#StatName)       , Stat_GPU_##StatName, STATGROUP_GPU  ); CSV_DEFINE_STAT(GPU,StatName);
 		#define DECLARE_GPU_STAT_NAMED(          StatName, NameString) DECLARE_FLOAT_COUNTER_STAT(NameString            , Stat_GPU_##StatName, STATGROUP_GPU  ); CSV_DEFINE_STAT(GPU,StatName);
@@ -184,8 +183,6 @@ class FRealtimeGPUProfilerFrame;
 
 #if HAS_GPU_STATS
 
-	CSV_DECLARE_CATEGORY_MODULE_EXTERN(RENDERCORE_API,GPU);
-
 	// The DECLARE_GPU_STAT macros both declare and define a stat (for use in a single CPP)
 	#define DECLARE_GPU_STAT(StatName)                            DECLARE_FLOAT_COUNTER_STAT(TEXT(#StatName)       , Stat_GPU_##StatName, STATGROUP_GPU  ); CSV_DEFINE_STAT(GPU,StatName);         static FRHIDrawStatsCategory DrawcallCountCategory_##StatName;
 	#define DECLARE_GPU_STAT_NAMED(StatName, NameString)          DECLARE_FLOAT_COUNTER_STAT(NameString            , Stat_GPU_##StatName, STATGROUP_GPU  ); CSV_DEFINE_STAT(GPU,StatName);         static FRHIDrawStatsCategory DrawcallCountCategory_##StatName;
diff --git a/Engine/Source/Runtime/RenderCore/Public/RenderGraphEvent.h b/Engine/Source/Runtime/RenderCore/Public/RenderGraphEvent.h
--- a/Engine/Source/Runtime/RenderCore/Public/RenderGraphEvent.h
+++ b/Engine/Source/Runtime/RenderCore/Public/RenderGraphEvent.h
@@ -49,7 +49,7 @@ public:
 	const TCHAR* GetTCHAR() const;
 
 #if WITH_RHI_BREADCRUMBS
-	FRHIBreadcrumbNode* AllocBreadcrumb(TStatId StatId, FRHIBreadcrumbAllocator& Allocator) const;
+	FRHIBreadcrumbNode* AllocBreadcrumb(FRHIBreadcrumbData&& Data, FRHIBreadcrumbAllocator& Allocator) const;
 #endif // WITH_RHI_BREADCRUMBS
 
 private:
@@ -171,7 +171,7 @@ struct FRDGScope_Budget
 		FRHIBreadcrumbNode* Node = nullptr;
 	#endif
 
-		inline FRDGScope_RHI(FRDGScopeState& State, TStatId StatId, FRDGEventName&& Name);
+		inline FRDGScope_RHI(FRDGScopeState& State, FRHIBreadcrumbData&& Data, FRDGEventName&& Name);
 		inline void ImmediateEnd(FRDGScopeState& State);
 
 	#if WITH_RHI_BREADCRUMBS
@@ -348,7 +348,7 @@ private:
 		TRDGEventScopeGuard<FRDGScope_RHI> PREPROCESSOR_JOIN(__RDG_ScopeRef_,__LINE__)( \
 			(GraphBuilder)                                                              \
 			, ERDGScopeFlags::None                                                      \
-			, TStatId()                                                                 \
+			, FRHIBreadcrumbData(__FILE__, __LINE__, TStatId(), NAME_None)              \
 			, RDG_EVENT_NAME(Format, ##__VA_ARGS__)                                     \
 		)
 
@@ -357,7 +357,12 @@ private:
 			TRDGEventScopeGuard<FRDGScope_RHI> PREPROCESSOR_JOIN(__RDG_ScopeRef_,__LINE__)( \
 				(GraphBuilder)                                                              \
 				, ERDGScopeFlags::Stat                                                      \
-				, GET_STATID(Stat_GPU_##StatName)                                           \
+				, FRHIBreadcrumbData(                                                       \
+					  __FILE__                                                              \
+					, __LINE__                                                              \
+					, GET_STATID(Stat_GPU_##StatName)                                       \
+					, CSV_STAT_FNAME(StatName)                                              \
+				)                                                                           \
 				, RDG_EVENT_NAME(Format, ##__VA_ARGS__)                                     \
 			)
 	#else
@@ -374,7 +379,12 @@ private:
 				PREPROCESSOR_JOIN(__RDG_ScopeRef_,__LINE__).Emplace(                               \
 					(GraphBuilder)                                                                 \
 					, ERDGScopeFlags::None                                                         \
-					, TStatId()                                                                    \
+					, FRHIBreadcrumbData(                                                          \
+						  __FILE__                                                                 \
+						, __LINE__                                                                 \
+						, TStatId()                                                                \
+						, NAME_None                                                                \
+					)                                                                              \
 					, RDG_EVENT_NAME(Format, ##__VA_ARGS__)                                        \
 				);                                                                                 \
 			}   \
@@ -390,7 +400,12 @@ private:
 					PREPROCESSOR_JOIN(__RDG_ScopeRef_,__LINE__).Emplace(                               \
 						(GraphBuilder)                                                                 \
 						, ERDGScopeFlags::Stat                                                         \
-						, GET_STATID(Stat_GPU_##StatName)                                              \
+						, FRHIBreadcrumbData(                                                          \
+							  __FILE__                                                                 \
+							, __LINE__                                                                 \
+							, GET_STATID(Stat_GPU_##StatName)                                          \
+							, CSV_STAT_FNAME(StatName)                                                 \
+						)                                                                              \
 						, RDG_EVENT_NAME(Format, ##__VA_ARGS__)                                        \
 					);                                                                                 \
 				}   \
@@ -409,7 +424,7 @@ private:
 			PREPROCESSOR_JOIN(__RDG_ScopeRef_,__LINE__).Emplace(                                   \
 				(GraphBuilder)                                                                     \
 				, ERDGScopeFlags::Final                                                            \
-				, TStatId()                                                                        \
+				, FRHIBreadcrumbData(__FILE__, __LINE__, TStatId(), NAME_None)                     \
 				, RDG_EVENT_NAME(Format, ##__VA_ARGS__)                                            \
 			);                                                                                     \
 		} while (false)
diff --git a/Engine/Source/Runtime/RenderCore/Public/RenderGraphEvent.inl b/Engine/Source/Runtime/RenderCore/Public/RenderGraphEvent.inl
--- a/Engine/Source/Runtime/RenderCore/Public/RenderGraphEvent.inl
+++ b/Engine/Source/Runtime/RenderCore/Public/RenderGraphEvent.inl
@@ -100,10 +100,10 @@ inline void FRDGScope_GPU::EndGPU(FRHIComputeCommandList& RHICmdList)
 
 #if RDG_EVENTS
 
-	inline FRDGScope_RHI::FRDGScope_RHI(FRDGScopeState& State, TStatId StatId, FRDGEventName&& Name)
+	inline FRDGScope_RHI::FRDGScope_RHI(FRDGScopeState& State, FRHIBreadcrumbData&& Data, FRDGEventName&& Name)
 		: Name(MoveTemp(Name))
 	#if WITH_RHI_BREADCRUMBS
-		, Node(Name.AllocBreadcrumb(StatId, State.GetBreadcrumbAllocator()))
+		, Node(Name.AllocBreadcrumb(MoveTemp(Data), State.GetBreadcrumbAllocator()))
 	#endif
 	{
 	#if WITH_RHI_BREADCRUMBS
diff --git a/Engine/Source/Runtime/Renderer/Private/MeshPassProcessor.cpp b/Engine/Source/Runtime/Renderer/Private/MeshPassProcessor.cpp
--- a/Engine/Source/Runtime/Renderer/Private/MeshPassProcessor.cpp
+++ b/Engine/Source/Runtime/Renderer/Private/MeshPassProcessor.cpp
@@ -1522,8 +1522,13 @@ bool FMeshDrawCommand::SubmitDraw(
 	FMeshDrawCommandStateCache& RESTRICT StateCache)
 {
 #if WANTS_DRAW_MESH_EVENTS
-	FMeshDrawEvent MeshEvent(MeshDrawCommand, InstanceFactor, RHICmdList);
+	RHI_BREADCRUMB_EVENT_CONDITIONAL(RHICmdList, GShowMaterialDrawEvents != 0, "%s %s (%u instances)"
+		, MeshDrawCommand.DebugData.MaterialRenderProxy->GetMaterialName()
+		, MeshDrawCommand.DebugData.ResourceName
+		, MeshDrawCommand.NumInstances * InstanceFactor
+	);
 #endif
+
 	bool bAllowSkipDrawCommand = true;
 	if (SubmitDrawBegin(MeshDrawCommand, GraphicsMinimalPipelineStateSet, SceneArgs, InstanceFactor, RHICmdList, StateCache, bAllowSkipDrawCommand))
 	{
@@ -2251,20 +2256,6 @@ void FPassProcessorManager::SetPassFlags(EShadingPath ShadingPath, EMeshPass::Ty
 	}
 }
 
-#if WANTS_DRAW_MESH_EVENTS
-FMeshDrawCommand::FMeshDrawEvent::FMeshDrawEvent(const FMeshDrawCommand& MeshDrawCommand, const uint32 InstanceFactor, FRHICommandList& RHICmdList)
-	: Breadcrumb(
-		  RHICmdList
-		, TStatId()
-		, GShowMaterialDrawEvents != 0
-		, TEXT("%s %s (%u instances)")
-		, MeshDrawCommand.DebugData.MaterialRenderProxy->GetMaterialName()
-		, MeshDrawCommand.DebugData.ResourceName
-		, MeshDrawCommand.NumInstances * InstanceFactor
-	)
-{}
-#endif
-
 void AddRenderTargetInfo(
 	EPixelFormat PixelFormat,
 	ETextureCreateFlags CreateFlags,
diff --git a/Engine/Source/Runtime/Renderer/Private/ShadowDepthRendering.cpp b/Engine/Source/Runtime/Renderer/Private/ShadowDepthRendering.cpp
--- a/Engine/Source/Runtime/Renderer/Private/ShadowDepthRendering.cpp
+++ b/Engine/Source/Runtime/Renderer/Private/ShadowDepthRendering.cpp
@@ -1597,7 +1597,7 @@ void FSceneRenderer::RenderShadowDepthMapAtlases(FRDGBuilder& GraphBuilder)
 				FString LightNameWithLevel;
 				GetLightNameForDrawEvent(CurrentLightForDrawEvent, LightNameWithLevel);
 
-				RDGScope.Emplace(GraphBuilder, ERDGScopeFlags::None, TStatId(), RDG_EVENT_NAME("%s", *LightNameWithLevel));
+				RDGScope.Emplace(GraphBuilder, ERDGScopeFlags::None, FRHIBreadcrumbData(__FILE__, __LINE__, TStatId(), NAME_None), RDG_EVENT_NAME("%s", *LightNameWithLevel));
 			}
 		#endif
 		};
diff --git a/Engine/Source/Runtime/Renderer/Public/MeshPassProcessor.h b/Engine/Source/Runtime/Renderer/Public/MeshPassProcessor.h
--- a/Engine/Source/Runtime/Renderer/Public/MeshPassProcessor.h
+++ b/Engine/Source/Runtime/Renderer/Public/MeshPassProcessor.h
@@ -1395,18 +1395,6 @@ private:
 	mutable FMeshDrawCommandDebugData DebugData;
 #endif
 
-#if WANTS_DRAW_MESH_EVENTS
-public:
-	friend struct FMeshDrawEvent;
-	struct FMeshDrawEvent
-	{
-	private:
-		FRHIBreadcrumbEventScope Breadcrumb;
-	public:
-		FMeshDrawEvent(const FMeshDrawCommand& MeshDrawCommand, const uint32 InstanceFactor, FRHICommandList& RHICmdList);
-	};
-#endif
-
 #if MESH_DRAW_COMMAND_STATS
 private:
 	FMeshDrawCommandStatsData StatsData;
