Fix build error when HAS_GPU_STATS is zero.

diff --git a/Engine/Source/Runtime/RHI/Private/GPUProfiler.cpp b/Engine/Source/Runtime/RHI/Private/GPUProfiler.cpp
--- a/Engine/Source/Runtime/RHI/Private/GPUProfiler.cpp
+++ b/Engine/Source/Runtime/RHI/Private/GPUProfiler.cpp
@@ -26,6 +26,8 @@ CSV_DEFINE_STAT(GPU, Total);
 // We can remove this once the old GPU profiler code has been deleted, and RHI_NEW_GPU_PROFILER is set to 1 permanently.
 extern RHI_API void RHISetGPUStatTotals(bool bCsvStatsEnabled, double TotalMs)
 {
+#if HAS_GPU_STATS
+
 #if STATS
 	FThreadStats::AddMessage(GET_STATFNAME(Stat_GPU_Total), EStatOperation::Set, TotalMs);
 	TRACE_STAT_SET(GET_STATFNAME(Stat_GPU_Total), TotalMs);
@@ -37,6 +39,8 @@ extern RHI_API void RHISetGPUStatTotals(bool bCsvStatsEnabled, double TotalMs)
 		FCsvProfiler::Get()->RecordCustomStat(CSV_STAT_FNAME(Total), CSV_CATEGORY_INDEX(GPU), TotalMs, ECsvCustomStatOp::Set);
 	}
 #endif
+
+#endif // HAS_GPU_STATS
 }
 
 static TAutoConsoleVariable<int> CVarGPUCsvStatsEnabled(
