TRHIPipelineArray properly initialized, this fix EngineTest in Debug build with rhi validation on

diff --git a/Engine/Source/Runtime/RHI/Private/RHIValidation.cpp b/Engine/Source/Runtime/RHI/Private/RHIValidation.cpp
--- a/Engine/Source/Runtime/RHI/Private/RHIValidation.cpp
+++ b/Engine/Source/Runtime/RHI/Private/RHIValidation.cpp
@@ -405,8 +405,8 @@ void FValidationRHI::RHIFinalizeContext(FRHIFinalizeContextArgs&& Args, TRHIPipe
 {
 	FRHIFinalizeContextArgs FinalArgs;
 	
-	TRHIPipelineArray<IRHIPlatformCommandList*> FinalizedCommandLists;
-	TRHIPipelineArray<FValidationCommandList*> OuterCommandLists;
+	TRHIPipelineArray<IRHIPlatformCommandList*> FinalizedCommandLists { InPlace, nullptr };
+	TRHIPipelineArray<FValidationCommandList*> OuterCommandLists { InPlace, nullptr };
 	
 	// Re-combine the args so that the validation matches a normal call to RHIFinalizeContext
 	for(IRHIComputeContext* Context : Args.Contexts)
