Don't try to wait on, or shutdown, the PSO precache data structure if it's not initialized.

diff --git a/Engine/Source/Runtime/RHI/Private/PipelineStateCache.cpp b/Engine/Source/Runtime/RHI/Private/PipelineStateCache.cpp
--- a/Engine/Source/Runtime/RHI/Private/PipelineStateCache.cpp
+++ b/Engine/Source/Runtime/RHI/Private/PipelineStateCache.cpp
@@ -3117,6 +3117,7 @@ void PipelineStateCache::FlushResources()
 	GWorkGraphPipelineCache.FlushResources(bDiscardAndSwap);
 	GGraphicsPipelineCache.FlushResources(bDiscardAndSwap);
 
+	check(GPrecacheGraphicsPipelineCache && GPrecacheComputePipelineCache);
 	GPrecacheGraphicsPipelineCache->ProcessDelayedCleanup();
 	GPrecacheComputePipelineCache->ProcessDelayedCleanup();
 
@@ -4221,8 +4222,15 @@ void PipelineStateCache::WaitForAllTasks()
 	GComputePipelineCache.WaitTasksComplete();
 	GWorkGraphPipelineCache.WaitTasksComplete();
 	GGraphicsPipelineCache.WaitTasksComplete();
-	GPrecacheGraphicsPipelineCache->WaitTasksComplete();
-	GPrecacheComputePipelineCache->WaitTasksComplete();
+
+	if (GPrecacheGraphicsPipelineCache)
+	{
+		GPrecacheGraphicsPipelineCache->WaitTasksComplete();
+	}
+	if (GPrecacheComputePipelineCache)
+	{
+		GPrecacheComputePipelineCache->WaitTasksComplete();
+	}
 }
 
 void PipelineStateCache::Init()
