Schedule "new PSO" delegate to be broadcast from the game thread, because it's not safe to access UObject-based delegates from anywhere else.

Also move the IsBound() check on the delegate to the game thread.

diff --git a/Engine/Source/Runtime/RHI/Private/PipelineFileCache.cpp b/Engine/Source/Runtime/RHI/Private/PipelineFileCache.cpp
--- a/Engine/Source/Runtime/RHI/Private/PipelineFileCache.cpp
+++ b/Engine/Source/Runtime/RHI/Private/PipelineFileCache.cpp
@@ -6,6 +6,7 @@ PipelineFileCache.cpp: Pipeline state cache implementation.
 
 #include "PipelineFileCache.h"
 #include "Containers/List.h"
+#include "Containers/Ticker.h"
 #include "PipelineStateCache.h"
 #include "HAL/IConsoleManager.h"
 #include "Misc/EngineVersion.h"
@@ -3533,12 +3534,19 @@ void FPipelineFileCacheManager::BroadcastNewPSOsDelegate()
 		NewPSOsToReport.Empty(32);
 	}
 
-	if (ReportNewPSOs() && PSOLoggedEvent.IsBound())
+	if (!PSOs.IsEmpty() && ReportNewPSOs())
 	{
-		for (FPipelineCacheFileFormatPSO& PSO : PSOs)
+		// It's not safe to touch UObjects-based delegates from the render thread.
+		ExecuteOnGameThread(TEXT("OnPipelineStateLoggedBroadcastGT"), [PSOs = MoveTemp(PSOs)]() mutable
 		{
-			PSOLoggedEvent.Broadcast(PSO);
-		}
+			if (PSOLoggedEvent.IsBound())
+			{
+				for (FPipelineCacheFileFormatPSO& PSO : PSOs)
+				{
+					PSOLoggedEvent.Broadcast(PSO);
+				}
+			}
+		});
 	}
 }
 
diff --git a/Engine/Source/Runtime/RHI/Public/PipelineFileCache.h b/Engine/Source/Runtime/RHI/Public/PipelineFileCache.h
--- a/Engine/Source/Runtime/RHI/Public/PipelineFileCache.h
+++ b/Engine/Source/Runtime/RHI/Public/PipelineFileCache.h
@@ -367,6 +367,7 @@ public:
 
 	/*
 	 * If the delegate is set, broadcasts any new PSOs that were encountered since the last time the delegate was broadcast.
+	 * The broadcast is scheduled to be run on the game thread.
 	 */
 	RHI_API static void BroadcastNewPSOsDelegate();
 	
