Removing MGPU type overhead when compiling single GPU platforms.
WITH_MGPU should only be defined on Desktop PC.

diff --git a/Engine/Source/Runtime/D3D12RHI/Private/D3D12PoolAllocator.cpp b/Engine/Source/Runtime/D3D12RHI/Private/D3D12PoolAllocator.cpp
--- a/Engine/Source/Runtime/D3D12RHI/Private/D3D12PoolAllocator.cpp
+++ b/Engine/Source/Runtime/D3D12RHI/Private/D3D12PoolAllocator.cpp
@@ -252,11 +252,13 @@ FD3D12PoolAllocator::~FD3D12PoolAllocator()
 
 bool FD3D12PoolAllocator::SupportsAllocation(D3D12_HEAP_TYPE InHeapType, D3D12_RESOURCE_FLAGS InResourceFlags, EBufferUsageFlags InBufferUsage, ED3D12ResourceStateMode InResourceStateMode, uint32 Alignment) const
 {
+#if WITH_MGPU
 	// NNE resources must be in heaps visible on GPU0 only.  Required by DirectML.
 	if (EnumHasAnyFlags(InBufferUsage, EBufferUsageFlags::NNE) && VisibilityMask != FRHIGPUMask::GPU0())
 	{
 		return false;
 	}
+#endif
 
 	FD3D12ResourceInitConfig InInitConfig = GetResourceAllocatorInitConfig(InHeapType, InResourceFlags, InBufferUsage);
 	EResourceAllocationStrategy InAllocationStrategy = GetResourceAllocationStrategy(InResourceFlags, InResourceStateMode, Alignment);
diff --git a/Engine/Source/Runtime/D3D12RHI/Private/D3D12RHICommon.h b/Engine/Source/Runtime/D3D12RHI/Private/D3D12RHICommon.h
--- a/Engine/Source/Runtime/D3D12RHI/Private/D3D12RHICommon.h
+++ b/Engine/Source/Runtime/D3D12RHI/Private/D3D12RHICommon.h
@@ -100,19 +100,28 @@ class FD3D12GPUObject
 {
 public:
 	FD3D12GPUObject(FRHIGPUMask InGPUMask, FRHIGPUMask InVisibiltyMask)
+#if WITH_MGPU
 		: GPUMask(InGPUMask)
 		, VisibilityMask(InVisibiltyMask)
+#endif
 	{
 		// Note that node mask can't be null.
 	}
 
-	FORCEINLINE const FRHIGPUMask& GetGPUMask() const { return GPUMask; }
-	FORCEINLINE const FRHIGPUMask& GetVisibilityMask() const { return VisibilityMask; }
+#if WITH_MGPU
+	FORCEINLINE FRHIGPUMask GetGPUMask() const { return GPUMask; }
+	FORCEINLINE FRHIGPUMask GetVisibilityMask() const { return VisibilityMask; }
+#else
+	constexpr FRHIGPUMask GetGPUMask() const { return FRHIGPUMask::GPU0(); }
+	constexpr FRHIGPUMask GetVisibilityMask() const { return FRHIGPUMask::GPU0(); }
+#endif
 
 protected:
+#if WITH_MGPU
 	const FRHIGPUMask GPUMask;
 	// Which GPUs have direct access to this object
 	const FRHIGPUMask VisibilityMask;
+#endif
 };
 
 class FD3D12SingleNodeGPUObject : public FD3D12GPUObject
@@ -120,16 +129,24 @@ class FD3D12SingleNodeGPUObject : public FD3D12GPUObject
 public:
 	FD3D12SingleNodeGPUObject(FRHIGPUMask GPUMask)
 		: FD3D12GPUObject(GPUMask, GPUMask)
+#if WITH_MGPU
 		, GPUIndex(GPUMask.ToIndex())
+#endif
 	{}
 
+#if WITH_MGPU
 	FORCEINLINE uint32 GetGPUIndex() const
 	{
 		return GPUIndex;
 	}
+#else
+	constexpr uint32 GetGPUIndex() const { return 0; }
+#endif
 
 private:
+#if WITH_MGPU
 	uint32 GPUIndex;
+#endif
 };
 
 class FD3D12MultiNodeGPUObject : public FD3D12GPUObject
diff --git a/Engine/Source/Runtime/RHI/Public/MultiGPU.h b/Engine/Source/Runtime/RHI/Public/MultiGPU.h
--- a/Engine/Source/Runtime/RHI/Public/MultiGPU.h
+++ b/Engine/Source/Runtime/RHI/Public/MultiGPU.h
@@ -7,132 +7,203 @@
 #pragma once
 
 #include "Containers/Array.h"
-#include "Containers/ContainerAllocationPolicies.h"
 #include "Math/UnrealMathUtility.h"
 #include "Misc/AssertionMacros.h"
 
-#if PLATFORM_DESKTOP
-#define WITH_MGPU 1	// Explicit MGPU
+#if DO_GUARD_SLOW
+	#define GPUMASK_CONSTEXPR
 #else
-#define WITH_MGPU 0
+	#define GPUMASK_CONSTEXPR constexpr
 #endif
 
 #if WITH_MGPU
-#define MAX_NUM_GPUS 8
-extern RHI_API uint32 GNumExplicitGPUsForRendering;
-extern RHI_API uint32 GVirtualMGPU;
+	#define MAX_NUM_GPUS 8
+	extern RHI_API uint32 GNumExplicitGPUsForRendering;
+	extern RHI_API uint32 GVirtualMGPU;
+	#define SGPU_CONSTEXPR
 #else
-#define MAX_NUM_GPUS 1
-#define GNumExplicitGPUsForRendering 1
-#define GVirtualMGPU 0
+	#define MAX_NUM_GPUS 1
+	#define GNumExplicitGPUsForRendering 1
+	#define GVirtualMGPU 0
+	#define SGPU_CONSTEXPR GPUMASK_CONSTEXPR
 #endif
 
 /** A mask where each bit is a GPU index. Can not be empty so that non SLI platforms can optimize it to be always 1.  */
 struct FRHIGPUMask
 {
 private:
-
+#if WITH_MGPU
 	uint32 GPUMask;
-
-	FORCEINLINE explicit FRHIGPUMask(uint32 InGPUMask) : GPUMask(InGPUMask)
+	FORCEINLINE GPUMASK_CONSTEXPR uint32 GetMask() const
 	{
-#if WITH_MGPU
-		check(InGPUMask != 0);
+		return GPUMask;
+	}
 #else
-		check(InGPUMask == 1);
+	FORCEINLINE constexpr uint32 GetMask() const
+	{
+		return 1;
+	}
 #endif
+
+
+#if WITH_MGPU
+	FORCEINLINE explicit GPUMASK_CONSTEXPR FRHIGPUMask(uint32 InGPUMask)
+		: GPUMask(InGPUMask)
+	{
+		checkSlow(InGPUMask != 0);
+	}
+#else
+	FORCEINLINE explicit GPUMASK_CONSTEXPR FRHIGPUMask(uint32 InGPUMask)
+	{
+		checkSlow(InGPUMask == 1);
 	}
+#endif
 
 public:
+	FORCEINLINE GPUMASK_CONSTEXPR FRHIGPUMask()
+		: FRHIGPUMask(FRHIGPUMask::GPU0())
+	{
+	}
+
+#if WITH_MGPU
+	FORCEINLINE uint32 ToIndex() const
+	{
+		checkSlow(HasSingleIndex());
+		return FMath::CountTrailingZeros(GetMask());
+	}
 
-	FORCEINLINE FRHIGPUMask() : FRHIGPUMask(FRHIGPUMask::GPU0())
+	FORCEINLINE SGPU_CONSTEXPR bool HasSingleIndex() const
 	{
+		return FMath::IsPowerOfTwo(GetMask());
 	}
 
-	FORCEINLINE static FRHIGPUMask FromIndex(uint32 GPUIndex) { return FRHIGPUMask(1 << GPUIndex); }
+	FORCEINLINE SGPU_CONSTEXPR uint32 GetNumActive() const
+	{
+		return FPlatformMath::CountBits(GetMask());
+	}
 
-	FORCEINLINE uint32 ToIndex() const
+	FORCEINLINE SGPU_CONSTEXPR uint32 GetLastIndex() const
 	{
-#if WITH_MGPU
-		check(HasSingleIndex());
-		return FMath::CountTrailingZeros(GPUMask);
+		return FPlatformMath::FloorLog2(GetMask());
+	}
+
+	FORCEINLINE SGPU_CONSTEXPR uint32 GetFirstIndex() const
+	{
+		return FPlatformMath::CountTrailingZeros(GetMask());
+	}
 #else
+	FORCEINLINE constexpr uint32 ToIndex() const
+	{
 		return 0;
-#endif
 	}
 
-	FORCEINLINE bool HasSingleIndex() const
+	FORCEINLINE constexpr bool HasSingleIndex() const
 	{
-#if WITH_MGPU
-		return FMath::IsPowerOfTwo(GPUMask);
-#else
 		return true;
-#endif
 	}
 
-	FORCEINLINE uint32 GetNumActive() const
+	FORCEINLINE constexpr uint32 GetNumActive() const
 	{
-#if WITH_MGPU
-		return FPlatformMath::CountBits(GPUMask);
-#else
 		return 1;
-#endif
 	}
 
-	FORCEINLINE uint32 GetLastIndex() const
+	FORCEINLINE constexpr uint32 GetLastIndex() const
 	{
-#if WITH_MGPU
-		return FPlatformMath::FloorLog2(GPUMask);
-#else
 		return 0;
-#endif
 	}
 
-	FORCEINLINE uint32 GetFirstIndex() const
+	FORCEINLINE constexpr uint32 GetFirstIndex() const
 	{
-#if WITH_MGPU
-		return FPlatformMath::CountTrailingZeros(GPUMask);
-#else
 		return 0;
+	}
 #endif
+
+	FORCEINLINE SGPU_CONSTEXPR bool Contains(uint32 GPUIndex) const
+	{
+		return (GetMask() & (1 << GPUIndex)) != 0;
+	}
+
+	FORCEINLINE SGPU_CONSTEXPR bool ContainsAll(FRHIGPUMask Rhs) const
+	{
+		return (GetMask() & Rhs.GetMask()) == Rhs.GetMask();
+	}
+
+	FORCEINLINE SGPU_CONSTEXPR bool Intersects(FRHIGPUMask Rhs) const
+	{
+		return (GetMask() & Rhs.GetMask()) != 0;
+	}
+
+	FORCEINLINE SGPU_CONSTEXPR bool operator ==(FRHIGPUMask Rhs) const
+	{
+		return GetMask() == Rhs.GetMask();
 	}
 
-	FORCEINLINE bool Contains(uint32 GPUIndex) const { return (GPUMask & (1 << GPUIndex)) != 0; }
-	FORCEINLINE bool ContainsAll(const FRHIGPUMask& Rhs) const { return (GPUMask & Rhs.GPUMask) == Rhs.GPUMask; }
-	FORCEINLINE bool Intersects(const FRHIGPUMask& Rhs) const { return (GPUMask & Rhs.GPUMask) != 0; }
+	FORCEINLINE SGPU_CONSTEXPR bool operator !=(FRHIGPUMask Rhs) const
+	{
+		return GetMask() != Rhs.GetMask();
+	}
+
+	FORCEINLINE SGPU_CONSTEXPR void operator |=(FRHIGPUMask Rhs)
+	{
+#if WITH_MGPU
+		GPUMask |= Rhs.GetMask();
+#endif
+	}
 
-	FORCEINLINE bool operator ==(const FRHIGPUMask& Rhs) const { return GPUMask == Rhs.GPUMask; }
-	FORCEINLINE bool operator !=(const FRHIGPUMask& Rhs) const { return GPUMask != Rhs.GPUMask; }
+	FORCEINLINE SGPU_CONSTEXPR void operator &=(FRHIGPUMask Rhs)
+	{
+#if WITH_MGPU
+		GPUMask &= Rhs.GetMask();
+#endif
+	}
 
-	void operator |=(const FRHIGPUMask& Rhs) { GPUMask |= Rhs.GPUMask; }
-	void operator &=(const FRHIGPUMask& Rhs) { GPUMask &= Rhs.GPUMask; }
+	FORCEINLINE SGPU_CONSTEXPR uint32 GetNative() const
+	{
+		return GVirtualMGPU ? 1 : GetMask();
+	}
 
-	FORCEINLINE uint32 GetNative() const { return GVirtualMGPU ? 1 : GPUMask; }
+	FORCEINLINE SGPU_CONSTEXPR FRHIGPUMask operator &(FRHIGPUMask Rhs) const
+	{
+		return FRHIGPUMask(GetMask() & Rhs.GetMask());
+	}
 
-	FORCEINLINE FRHIGPUMask operator &(const FRHIGPUMask& Rhs) const
+	FORCEINLINE SGPU_CONSTEXPR FRHIGPUMask operator |(FRHIGPUMask Rhs) const
 	{
-		return FRHIGPUMask(GPUMask & Rhs.GPUMask);
+		return FRHIGPUMask(GetMask() | Rhs.GetMask());
 	}
 
-	FORCEINLINE FRHIGPUMask operator |(const FRHIGPUMask& Rhs) const
+	FORCEINLINE static GPUMASK_CONSTEXPR FRHIGPUMask FromIndex(uint32 GPUIndex)
 	{
-		return FRHIGPUMask(GPUMask | Rhs.GPUMask);
+		return FRHIGPUMask(1 << GPUIndex);
 	}
 
-	FORCEINLINE static const FRHIGPUMask GPU0() { return FRHIGPUMask(1); }
-	FORCEINLINE static const FRHIGPUMask All() { return FRHIGPUMask((1 << GNumExplicitGPUsForRendering) - 1); }
-	FORCEINLINE static const FRHIGPUMask FilterGPUsBefore(uint32 GPUIndex) { return FRHIGPUMask(~((1u << GPUIndex) - 1)) & All(); }
+	FORCEINLINE static GPUMASK_CONSTEXPR FRHIGPUMask GPU0()
+	{
+		return FRHIGPUMask(1);
+	}
+	
+	FORCEINLINE static SGPU_CONSTEXPR FRHIGPUMask All()
+	{
+		return FRHIGPUMask((1 << GNumExplicitGPUsForRendering) - 1);
+	}
+	
+	FORCEINLINE static SGPU_CONSTEXPR FRHIGPUMask FilterGPUsBefore(uint32 GPUIndex)
+	{
+		return FRHIGPUMask(~((1u << GPUIndex) - 1)) & All();
+	}
 
 	struct FIterator
 	{
-		FORCEINLINE explicit FIterator(const uint32 InGPUMask) : GPUMask(InGPUMask), FirstGPUIndexInMask(0)
-		{
+		FORCEINLINE explicit FIterator(const uint32 InGPUMask)
+			: GPUMask(InGPUMask)
 #if WITH_MGPU
-			FirstGPUIndexInMask = FPlatformMath::CountTrailingZeros(InGPUMask);
+			, FirstGPUIndexInMask(FPlatformMath::CountTrailingZeros(InGPUMask))
 #endif
+		{
 		}
 
-		FORCEINLINE explicit FIterator(const FRHIGPUMask& InGPUMask) : FIterator(InGPUMask.GPUMask)
+		FORCEINLINE explicit FIterator(FRHIGPUMask InGPUMask)
+			: FIterator(InGPUMask.GetMask())
 		{
 		}
 
@@ -154,16 +225,60 @@ public:
 			return Copy;
 		}
 
-		FORCEINLINE uint32 operator*() const { return FirstGPUIndexInMask; }
-		FORCEINLINE bool operator !=(const FIterator& Rhs) const { return GPUMask != Rhs.GPUMask; }
-		FORCEINLINE explicit operator bool() const { return GPUMask != 0; }
-		FORCEINLINE bool operator !() const { return !(bool)*this; }
+		FORCEINLINE uint32 operator*() const
+		{
+			return GetFirstIndexInMask();
+		}
+
+		FORCEINLINE bool operator !=(const FIterator& Rhs) const
+		{
+			return GetMask() != Rhs.GetMask();
+		}
+
+		FORCEINLINE explicit operator bool() const
+		{
+			return GetMask() != 0;
+		}
+
+		FORCEINLINE bool operator !() const
+		{
+			return !(bool)*this;
+		}
 
 	private:
+		// NOTE: we cannot remove this in single GPU mode since we need to actually iterate once.
 		uint32 GPUMask;
-		unsigned long FirstGPUIndexInMask;
+
+		FORCEINLINE uint32 GetMask() const
+		{
+			return GPUMask;
+		}
+
+#if WITH_MGPU
+		uint32 FirstGPUIndexInMask;
+		FORCEINLINE uint32 GetFirstIndexInMask() const
+		{
+			return FirstGPUIndexInMask;
+		}
+#else
+		FORCEINLINE constexpr uint32 GetFirstIndexInMask() const
+		{
+			return 0;
+		}
+#endif
 	};
 
-	FORCEINLINE friend FRHIGPUMask::FIterator begin(const FRHIGPUMask& NodeMask) { return FRHIGPUMask::FIterator(NodeMask.GPUMask); }
-	FORCEINLINE friend FRHIGPUMask::FIterator end(const FRHIGPUMask& NodeMask) { return FRHIGPUMask::FIterator(0); }
+	FORCEINLINE friend FRHIGPUMask::FIterator begin(FRHIGPUMask NodeMask)
+	{
+		return FRHIGPUMask::FIterator(NodeMask.GetMask());
+	}
+
+	FORCEINLINE friend FRHIGPUMask::FIterator end(FRHIGPUMask NodeMask)
+	{
+		return FRHIGPUMask::FIterator(0);
+	}
 };
+
+#if UE_ENABLE_INCLUDE_ORDER_DEPRECATED_IN_5_5
+#include "Containers/ContainerAllocationPolicies.h"
+#endif
diff --git a/Engine/Source/Runtime/RHI/RHI.Build.cs b/Engine/Source/Runtime/RHI/RHI.Build.cs
--- a/Engine/Source/Runtime/RHI/RHI.Build.cs
+++ b/Engine/Source/Runtime/RHI/RHI.Build.cs
@@ -2,6 +2,7 @@
 
 using UnrealBuildTool;
 using System;
+using EpicGames.Core;
 
 public class RHI : ModuleRules
 {
@@ -14,6 +15,8 @@ public class RHI : ModuleRules
 		// @todo - new gpu profiler. This is experimental.
 		PublicDefinitions.Add("RHI_NEW_GPU_PROFILER=0");
 
+		PublicDefinitions.AddDefinition("WITH_MGPU", Target.Platform.IsInGroup(UnrealPlatformGroup.Windows) && Target.Platform.IsInGroup(UnrealPlatformGroup.Desktop));
+
 		if (Target.bCompileAgainstEngine)
 		{
 			DynamicallyLoadedModuleNames.Add("NullDrv");
