[RHI]
* Another attempt at fixing compile/link errors. Problem is that for platforms with PLATFORM_USES_FIXED_RHI_CLASS declaration of Execute functions are in one module and implementation in another.. which causes conflicts of dllexport vs dllimport. This change makes it possible to control the instantiation of Execute functions and dll storage from the outside of the header

diff --git a/Engine/Source/Runtime/RHI/Public/RHICommandList.h b/Engine/Source/Runtime/RHI/Public/RHICommandList.h
--- a/Engine/Source/Runtime/RHI/Public/RHICommandList.h
+++ b/Engine/Source/Runtime/RHI/Public/RHICommandList.h
@@ -89,6 +89,10 @@ DECLARE_STATS_GROUP(TEXT("RHICommands"),STATGROUP_RHI_COMMANDS, STATCAT_Advanced
 #define RHISTAT(Method)
 #endif
 
+#if !defined(RHI_EXECUTE_API)
+#define RHI_EXECUTE_API RHI_API
+#endif
+
 enum class ERHIThreadMode
 {
 	None,
@@ -1411,7 +1415,7 @@ private:
 	FRHICommandListBase(FPersistentState const& InPersistentState);
 
 	// Replays recorded commands. Used internally, do not call directly.
-	RHI_API void Execute();
+	RHI_EXECUTE_API void Execute();
 
 	friend class FRHICommandListExecutor;
 	friend class FRHICommandListIterator;
@@ -1498,7 +1502,7 @@ FRHICOMMAND_MACRO(FRHICommandBeginUpdateMultiFrameResource)
 		: Texture(InTexture)
 	{
 	}
-	RHI_API void Execute(FRHICommandListBase& CmdList);
+	RHI_EXECUTE_API void Execute(FRHICommandListBase& CmdList);
 };
 
 FRHICOMMAND_MACRO(FRHICommandEndUpdateMultiFrameResource)
@@ -1508,7 +1512,7 @@ FRHICOMMAND_MACRO(FRHICommandEndUpdateMultiFrameResource)
 		: Texture(InTexture)
 	{
 	}
-	RHI_API void Execute(FRHICommandListBase& CmdList);
+	RHI_EXECUTE_API void Execute(FRHICommandListBase& CmdList);
 };
 
 FRHICOMMAND_MACRO(FRHICommandBeginUpdateMultiFrameUAV)
@@ -1518,7 +1522,7 @@ FRHICOMMAND_MACRO(FRHICommandBeginUpdateMultiFrameUAV)
 		: UAV(InUAV)
 	{
 	}
-	RHI_API void Execute(FRHICommandListBase& CmdList);
+	RHI_EXECUTE_API void Execute(FRHICommandListBase& CmdList);
 };
 
 FRHICOMMAND_MACRO(FRHICommandEndUpdateMultiFrameUAV)
@@ -1528,7 +1532,7 @@ FRHICOMMAND_MACRO(FRHICommandEndUpdateMultiFrameUAV)
 		: UAV(InUAV)
 	{
 	}
-	RHI_API void Execute(FRHICommandListBase& CmdList);
+	RHI_EXECUTE_API void Execute(FRHICommandListBase& CmdList);
 };
 
 #if WITH_MGPU
@@ -1539,7 +1543,7 @@ FRHICOMMAND_MACRO(FRHICommandSetGPUMask)
 		: GPUMask(InGPUMask)
 	{
 	}
-	RHI_API void Execute(FRHICommandListBase& CmdList);
+	RHI_EXECUTE_API void Execute(FRHICommandListBase& CmdList);
 };
 
 FRHICOMMAND_MACRO(FRHICommandTransferResources)
@@ -1551,7 +1555,7 @@ FRHICOMMAND_MACRO(FRHICommandTransferResources)
 	{
 	}
 
-	RHI_API void Execute(FRHICommandListBase& CmdList);
+	RHI_EXECUTE_API void Execute(FRHICommandListBase& CmdList);
 };
 
 FRHICOMMAND_MACRO(FRHICommandTransferResourceSignal)
@@ -1565,7 +1569,7 @@ FRHICOMMAND_MACRO(FRHICommandTransferResourceSignal)
 	{
 	}
 
-	RHI_API void Execute(FRHICommandListBase & CmdList);
+	RHI_EXECUTE_API void Execute(FRHICommandListBase & CmdList);
 };
 
 FRHICOMMAND_MACRO(FRHICommandTransferResourceWait)
@@ -1577,7 +1581,7 @@ FRHICOMMAND_MACRO(FRHICommandTransferResourceWait)
 	{
 	}
 
-	RHI_API void Execute(FRHICommandListBase & CmdList);
+	RHI_EXECUTE_API void Execute(FRHICommandListBase & CmdList);
 };
 
 FRHICOMMAND_MACRO(FRHICommandCrossGPUTransfer)
@@ -1593,7 +1597,7 @@ FRHICOMMAND_MACRO(FRHICommandCrossGPUTransfer)
 	{
 	}
 
-	RHI_API void Execute(FRHICommandListBase & CmdList);
+	RHI_EXECUTE_API void Execute(FRHICommandListBase & CmdList);
 };
 
 FRHICOMMAND_MACRO(FRHICommandCrossGPUTransferSignal)
@@ -1607,7 +1611,7 @@ FRHICOMMAND_MACRO(FRHICommandCrossGPUTransferSignal)
 	{
 	}
 
-	RHI_API void Execute(FRHICommandListBase & CmdList);
+	RHI_EXECUTE_API void Execute(FRHICommandListBase & CmdList);
 };
 
 FRHICOMMAND_MACRO(FRHICommandCrossGPUTransferWait)
@@ -1619,7 +1623,7 @@ FRHICOMMAND_MACRO(FRHICommandCrossGPUTransferWait)
 	{
 	}
 
-	RHI_API void Execute(FRHICommandListBase & CmdList);
+	RHI_EXECUTE_API void Execute(FRHICommandListBase & CmdList);
 };
 #endif // WITH_MGPU
 
@@ -1630,7 +1634,7 @@ FRHICOMMAND_MACRO(FRHICommandSetStencilRef)
 		: StencilRef(InStencilRef)
 	{
 	}
-	RHI_API void Execute(FRHICommandListBase& CmdList);
+	RHI_EXECUTE_API void Execute(FRHICommandListBase& CmdList);
 };
 
 FRHICOMMAND_MACRO_TPL(TRHIShader, FRHICommandSetShaderParameters)
@@ -1655,7 +1659,7 @@ FRHICOMMAND_MACRO_TPL(TRHIShader, FRHICommandSetShaderParameters)
 		, BindlessParameters(InBindlessParameters)
 	{
 	}
-	RHI_API void Execute(FRHICommandListBase& CmdList);
+	RHI_EXECUTE_API void Execute(FRHICommandListBase& CmdList);
 };
 
 FRHICOMMAND_MACRO_TPL(TRHIShader, FRHICommandSetShaderUnbinds)
@@ -1668,7 +1672,7 @@ FRHICOMMAND_MACRO_TPL(TRHIShader, FRHICommandSetShaderUnbinds)
 		, Unbinds(InUnbinds)
 	{
 	}
-	RHI_API void Execute(FRHICommandListBase & CmdList);
+	RHI_EXECUTE_API void Execute(FRHICommandListBase & CmdList);
 };
 
 FRHICOMMAND_MACRO(FRHICommandDrawPrimitive)
@@ -1682,7 +1686,7 @@ FRHICOMMAND_MACRO(FRHICommandDrawPrimitive)
 		, NumInstances(InNumInstances)
 	{
 	}
-	RHI_API void Execute(FRHICommandListBase& CmdList);
+	RHI_EXECUTE_API void Execute(FRHICommandListBase& CmdList);
 };
 
 FRHICOMMAND_MACRO(FRHICommandDrawIndexedPrimitive)
@@ -1704,7 +1708,7 @@ FRHICOMMAND_MACRO(FRHICommandDrawIndexedPrimitive)
 		, NumInstances(InNumInstances)
 	{
 	}
-	RHI_API void Execute(FRHICommandListBase& CmdList);
+	RHI_EXECUTE_API void Execute(FRHICommandListBase& CmdList);
 };
 
 FRHICOMMAND_MACRO(FRHICommandSetBlendFactor)
@@ -1714,7 +1718,7 @@ FRHICOMMAND_MACRO(FRHICommandSetBlendFactor)
 		: BlendFactor(InBlendFactor)
 	{
 	}
-	RHI_API void Execute(FRHICommandListBase& CmdList);
+	RHI_EXECUTE_API void Execute(FRHICommandListBase& CmdList);
 };
 
 FRHICOMMAND_MACRO(FRHICommandSetStreamSource)
@@ -1728,7 +1732,7 @@ FRHICOMMAND_MACRO(FRHICommandSetStreamSource)
 		, Offset(InOffset)
 	{
 	}
-	RHI_API void Execute(FRHICommandListBase& CmdList);
+	RHI_EXECUTE_API void Execute(FRHICommandListBase& CmdList);
 };
 
 FRHICOMMAND_MACRO(FRHICommandSetViewport)
@@ -1748,7 +1752,7 @@ FRHICOMMAND_MACRO(FRHICommandSetViewport)
 		, MaxZ(InMaxZ)
 	{
 	}
-	RHI_API void Execute(FRHICommandListBase& CmdList);
+	RHI_EXECUTE_API void Execute(FRHICommandListBase& CmdList);
 };
 
 FRHICOMMAND_MACRO(FRHICommandSetStereoViewport)
@@ -1776,7 +1780,7 @@ FRHICOMMAND_MACRO(FRHICommandSetStereoViewport)
 		, MaxZ(InMaxZ)
 	{
 	}
-	RHI_API void Execute(FRHICommandListBase& CmdList);
+	RHI_EXECUTE_API void Execute(FRHICommandListBase& CmdList);
 };
 
 FRHICOMMAND_MACRO(FRHICommandSetScissorRect)
@@ -1794,7 +1798,7 @@ FRHICOMMAND_MACRO(FRHICommandSetScissorRect)
 		, MaxY(InMaxY)
 	{
 	}
-	RHI_API void Execute(FRHICommandListBase& CmdList);
+	RHI_EXECUTE_API void Execute(FRHICommandListBase& CmdList);
 };
 
 FRHICOMMAND_MACRO(FRHICommandBeginRenderPass)
@@ -1808,7 +1812,7 @@ FRHICOMMAND_MACRO(FRHICommandBeginRenderPass)
 	{
 	}
 
-	RHI_API void Execute(FRHICommandListBase& CmdList);
+	RHI_EXECUTE_API void Execute(FRHICommandListBase& CmdList);
 };
 
 FRHICOMMAND_MACRO(FRHICommandEndRenderPass)
@@ -1817,7 +1821,7 @@ FRHICOMMAND_MACRO(FRHICommandEndRenderPass)
 	{
 	}
 
-	RHI_API void Execute(FRHICommandListBase& CmdList);
+	RHI_EXECUTE_API void Execute(FRHICommandListBase& CmdList);
 };
 
 FRHICOMMAND_MACRO(FRHICommandNextSubpass)
@@ -1826,7 +1830,7 @@ FRHICOMMAND_MACRO(FRHICommandNextSubpass)
 	{
 	}
 
-	RHI_API void Execute(FRHICommandListBase& CmdList);
+	RHI_EXECUTE_API void Execute(FRHICommandListBase& CmdList);
 };
 
 FRHICOMMAND_MACRO(FRHICommandSetComputePipelineState)
@@ -1836,7 +1840,7 @@ FRHICOMMAND_MACRO(FRHICommandSetComputePipelineState)
 		: ComputePipelineState(InComputePipelineState)
 	{
 	}
-	RHI_API void Execute(FRHICommandListBase& CmdList);
+	RHI_EXECUTE_API void Execute(FRHICommandListBase& CmdList);
 };
 
 FRHICOMMAND_MACRO(FRHICommandSetGraphicsPipelineState)
@@ -1850,7 +1854,7 @@ FRHICOMMAND_MACRO(FRHICommandSetGraphicsPipelineState)
 		, bApplyAdditionalState(bInApplyAdditionalState)
 	{
 	}
-	RHI_API void Execute(FRHICommandListBase& CmdList);
+	RHI_EXECUTE_API void Execute(FRHICommandListBase& CmdList);
 };
 
 #if PLATFORM_USE_FALLBACK_PSO
@@ -1865,7 +1869,7 @@ FRHICOMMAND_MACRO(FRHICommandSetGraphicsPipelineStateFromInitializer)
 		, bApplyAdditionalState(bInApplyAdditionalState)
 	{
 	}
-	RHI_API void Execute(FRHICommandListBase& CmdList);
+	RHI_EXECUTE_API void Execute(FRHICommandListBase& CmdList);
 };
 #endif
 
@@ -1880,7 +1884,7 @@ FRHICOMMAND_MACRO(FRHICommandDispatchComputeShader)
 		, ThreadGroupCountZ(InThreadGroupCountZ)
 	{
 	}
-	RHI_API void Execute(FRHICommandListBase& CmdList);
+	RHI_EXECUTE_API void Execute(FRHICommandListBase& CmdList);
 };
 
 FRHICOMMAND_MACRO(FRHICommandDispatchIndirectComputeShader)
@@ -1892,7 +1896,7 @@ FRHICOMMAND_MACRO(FRHICommandDispatchIndirectComputeShader)
 		, ArgumentOffset(InArgumentOffset)
 	{
 	}
-	RHI_API void Execute(FRHICommandListBase& CmdList);
+	RHI_EXECUTE_API void Execute(FRHICommandListBase& CmdList);
 };
 
 using FRHIRecordBundleComputeDispatchCallback = TFunction<void(FRHIShaderBundleComputeDispatch& Dispatch)>;
@@ -1925,7 +1929,7 @@ FRHICOMMAND_MACRO(FRHICommandDispatchComputeShaderBundle)
 		, bEmulated(bInEmulated)
 	{
 	}
-	RHI_API void Execute(FRHICommandListBase& CmdList);
+	RHI_EXECUTE_API void Execute(FRHICommandListBase& CmdList);
 };
 
 FRHICOMMAND_MACRO(FRHICommandDispatchGraphicsShaderBundle)
@@ -1958,7 +1962,7 @@ FRHICOMMAND_MACRO(FRHICommandDispatchGraphicsShaderBundle)
 		, bEmulated(bInEmulated)
 	{
 	}
-	RHI_API void Execute(FRHICommandListBase& CmdList);
+	RHI_EXECUTE_API void Execute(FRHICommandListBase& CmdList);
 };
 
 FRHICOMMAND_MACRO(FRHICommandSetShaderRootConstants)
@@ -1974,31 +1978,31 @@ FRHICOMMAND_MACRO(FRHICommandSetShaderRootConstants)
 		: Constants(InConstants)
 	{
 	}
-	RHI_API void Execute(FRHICommandListBase& CmdList);
+	RHI_EXECUTE_API void Execute(FRHICommandListBase& CmdList);
 };
 
 FRHICOMMAND_MACRO(FRHICommandBeginUAVOverlap)
 {
-	RHI_API void Execute(FRHICommandListBase& CmdList);
+	RHI_EXECUTE_API void Execute(FRHICommandListBase& CmdList);
 };
 
 FRHICOMMAND_MACRO(FRHICommandEndUAVOverlap)
 {
-	RHI_API void Execute(FRHICommandListBase& CmdList);
+	RHI_EXECUTE_API void Execute(FRHICommandListBase& CmdList);
 };
 
 FRHICOMMAND_MACRO(FRHICommandBeginSpecificUAVOverlap)
 {
 	TArrayView<FRHIUnorderedAccessView* const> UAVs;
 	FORCEINLINE_DEBUGGABLE FRHICommandBeginSpecificUAVOverlap(TArrayView<FRHIUnorderedAccessView* const> InUAVs) : UAVs(InUAVs) {}
-	RHI_API void Execute(FRHICommandListBase& CmdList);
+	RHI_EXECUTE_API void Execute(FRHICommandListBase& CmdList);
 };
 
 FRHICOMMAND_MACRO(FRHICommandEndSpecificUAVOverlap)
 {
 	TArrayView<FRHIUnorderedAccessView* const> UAVs;
 	FORCEINLINE_DEBUGGABLE FRHICommandEndSpecificUAVOverlap(TArrayView<FRHIUnorderedAccessView* const> InUAVs) : UAVs(InUAVs) {}
-	RHI_API void Execute(FRHICommandListBase & CmdList);
+	RHI_EXECUTE_API void Execute(FRHICommandListBase & CmdList);
 };
 
 FRHICOMMAND_MACRO(FRHICommandDrawPrimitiveIndirect)
@@ -2010,7 +2014,7 @@ FRHICOMMAND_MACRO(FRHICommandDrawPrimitiveIndirect)
 		, ArgumentOffset(InArgumentOffset)
 	{
 	}
-	RHI_API void Execute(FRHICommandListBase& CmdList);
+	RHI_EXECUTE_API void Execute(FRHICommandListBase& CmdList);
 };
 
 FRHICOMMAND_MACRO(FRHICommandDrawIndexedIndirect)
@@ -2027,7 +2031,7 @@ FRHICOMMAND_MACRO(FRHICommandDrawIndexedIndirect)
 		, NumInstances(InNumInstances)
 	{
 	}
-	RHI_API void Execute(FRHICommandListBase& CmdList);
+	RHI_EXECUTE_API void Execute(FRHICommandListBase& CmdList);
 };
 
 FRHICOMMAND_MACRO(FRHICommandDrawIndexedPrimitiveIndirect)
@@ -2042,7 +2046,7 @@ FRHICOMMAND_MACRO(FRHICommandDrawIndexedPrimitiveIndirect)
 		, ArgumentOffset(InArgumentOffset)
 	{
 	}
-	RHI_API void Execute(FRHICommandListBase& CmdList);
+	RHI_EXECUTE_API void Execute(FRHICommandListBase& CmdList);
 };
 
 FRHICOMMAND_MACRO(FRHICommandMultiDrawIndexedPrimitiveIndirect)
@@ -2062,7 +2066,7 @@ FRHICOMMAND_MACRO(FRHICommandMultiDrawIndexedPrimitiveIndirect)
 		, MaxDrawArguments(InMaxDrawArguments)
 	{
 	}
-	RHI_API void Execute(FRHICommandListBase & CmdList);
+	RHI_EXECUTE_API void Execute(FRHICommandListBase & CmdList);
 };
 
 FRHICOMMAND_MACRO(FRHICommandDispatchMeshShader)
@@ -2077,7 +2081,7 @@ FRHICOMMAND_MACRO(FRHICommandDispatchMeshShader)
 		, ThreadGroupCountZ(InThreadGroupCountZ)
 	{
 	}
-	RHI_API void Execute(FRHICommandListBase & CmdList);
+	RHI_EXECUTE_API void Execute(FRHICommandListBase & CmdList);
 };
 
 FRHICOMMAND_MACRO(FRHICommandDispatchIndirectMeshShader)
@@ -2090,7 +2094,7 @@ FRHICOMMAND_MACRO(FRHICommandDispatchIndirectMeshShader)
 		, ArgumentOffset(InArgumentOffset)
 	{
 	}
-	RHI_API void Execute(FRHICommandListBase & CmdList);
+	RHI_EXECUTE_API void Execute(FRHICommandListBase & CmdList);
 };
 
 FRHICOMMAND_MACRO(FRHICommandSetDepthBounds)
@@ -2103,7 +2107,7 @@ FRHICOMMAND_MACRO(FRHICommandSetDepthBounds)
 		, MaxDepth(InMaxDepth)
 	{
 	}
-	RHI_API void Execute(FRHICommandListBase& CmdList);
+	RHI_EXECUTE_API void Execute(FRHICommandListBase& CmdList);
 };
 
 FRHICOMMAND_MACRO(FRHICommandSetShadingRate)
@@ -2116,7 +2120,7 @@ FRHICOMMAND_MACRO(FRHICommandSetShadingRate)
 		Combiner(InCombiner)
 	{
 	}
-	RHI_API void Execute(FRHICommandListBase& CmdList);
+	RHI_EXECUTE_API void Execute(FRHICommandListBase& CmdList);
 };
 
 FRHICOMMAND_MACRO(FRHICommandClearUAVFloat)
@@ -2129,7 +2133,7 @@ FRHICOMMAND_MACRO(FRHICommandClearUAVFloat)
 		, Values(InValues)
 	{
 	}
-	RHI_API void Execute(FRHICommandListBase& CmdList);
+	RHI_EXECUTE_API void Execute(FRHICommandListBase& CmdList);
 };
 
 FRHICOMMAND_MACRO(FRHICommandClearUAVUint)
@@ -2142,7 +2146,7 @@ FRHICOMMAND_MACRO(FRHICommandClearUAVUint)
 		, Values(InValues)
 	{
 	}
-	RHI_API void Execute(FRHICommandListBase& CmdList);
+	RHI_EXECUTE_API void Execute(FRHICommandListBase& CmdList);
 };
 
 FRHICOMMAND_MACRO(FRHICommandCopyTexture)
@@ -2159,7 +2163,7 @@ FRHICOMMAND_MACRO(FRHICommandCopyTexture)
 		ensure(SourceTexture);
 		ensure(DestTexture);
 	}
-	RHI_API void Execute(FRHICommandListBase& CmdList);
+	RHI_EXECUTE_API void Execute(FRHICommandListBase& CmdList);
 };
 
 FRHICOMMAND_MACRO(FRHICommandResummarizeHTile)
@@ -2170,7 +2174,7 @@ FRHICOMMAND_MACRO(FRHICommandResummarizeHTile)
 	: DepthTexture(InDepthTexture)
 	{
 	}
-	RHI_API void Execute(FRHICommandListBase& CmdList);
+	RHI_EXECUTE_API void Execute(FRHICommandListBase& CmdList);
 };
 
 FRHICOMMAND_MACRO(FRHICommandBeginTransitions)
@@ -2182,7 +2186,7 @@ FRHICOMMAND_MACRO(FRHICommandBeginTransitions)
 	{
 	}
 
-	RHI_API void Execute(FRHICommandListBase& CmdList);
+	RHI_EXECUTE_API void Execute(FRHICommandListBase& CmdList);
 };
 
 FRHICOMMAND_MACRO(FRHICommandEndTransitions)
@@ -2194,7 +2198,7 @@ FRHICOMMAND_MACRO(FRHICommandEndTransitions)
 	{
 	}
 
-	RHI_API void Execute(FRHICommandListBase& CmdList);
+	RHI_EXECUTE_API void Execute(FRHICommandListBase& CmdList);
 };
 
 FRHICOMMAND_MACRO(FRHICommandResourceTransition)
@@ -2206,7 +2210,7 @@ FRHICOMMAND_MACRO(FRHICommandResourceTransition)
 	{
 	}
 
-	RHI_API void Execute(FRHICommandListBase& CmdList);
+	RHI_EXECUTE_API void Execute(FRHICommandListBase& CmdList);
 };
 
 FRHICOMMAND_MACRO(FRHICommandSetTrackedAccess)
@@ -2218,7 +2222,7 @@ FRHICOMMAND_MACRO(FRHICommandSetTrackedAccess)
 	{
 	}
 
-	RHI_API void Execute(FRHICommandListBase& CmdList);
+	RHI_EXECUTE_API void Execute(FRHICommandListBase& CmdList);
 };
 
 FRHICOMMAND_MACRO(FRHICommandSetAsyncComputeBudget)
@@ -2229,7 +2233,7 @@ FRHICOMMAND_MACRO(FRHICommandSetAsyncComputeBudget)
 		: Budget(InBudget)
 	{
 	}
-	RHI_API void Execute(FRHICommandListBase& CmdList);
+	RHI_EXECUTE_API void Execute(FRHICommandListBase& CmdList);
 };
 
 FRHICOMMAND_MACRO(FRHICommandCopyToStagingBuffer)
@@ -2246,7 +2250,7 @@ FRHICOMMAND_MACRO(FRHICommandCopyToStagingBuffer)
 		, NumBytes(InNumBytes)
 	{
 	}
-	RHI_API void Execute(FRHICommandListBase& CmdList);
+	RHI_EXECUTE_API void Execute(FRHICommandListBase& CmdList);
 };
 
 FRHICOMMAND_MACRO(FRHICommandWriteGPUFence)
@@ -2261,7 +2265,7 @@ FRHICOMMAND_MACRO(FRHICommandWriteGPUFence)
 			Fence->NumPendingWriteCommands.Increment();
 		}
 	}
-	RHI_API void Execute(FRHICommandListBase& CmdList);
+	RHI_EXECUTE_API void Execute(FRHICommandListBase& CmdList);
 };
 
 FRHICOMMAND_MACRO(FRHICommandSetStaticUniformBuffers)
@@ -2272,7 +2276,7 @@ FRHICOMMAND_MACRO(FRHICommandSetStaticUniformBuffers)
 		: UniformBuffers(InUniformBuffers)
 	{}
 
-	RHI_API void Execute(FRHICommandListBase & CmdList);
+	RHI_EXECUTE_API void Execute(FRHICommandListBase & CmdList);
 };
 
 FRHICOMMAND_MACRO(FRHICommandSetStaticUniformBuffer)
@@ -2285,7 +2289,7 @@ FRHICOMMAND_MACRO(FRHICommandSetStaticUniformBuffer)
 		, Slot(InSlot)
 	{}
 
-	RHI_API void Execute(FRHICommandListBase & CmdList);
+	RHI_EXECUTE_API void Execute(FRHICommandListBase & CmdList);
 };
 
 FRHICOMMAND_MACRO(FRHICommandSetUniformBufferDynamicOffset)
@@ -2298,7 +2302,7 @@ FRHICOMMAND_MACRO(FRHICommandSetUniformBufferDynamicOffset)
 		, Slot(InSlot)
 	{
 	}
-	RHI_API void Execute(FRHICommandListBase& CmdList);
+	RHI_EXECUTE_API void Execute(FRHICommandListBase& CmdList);
 };
 
 FRHICOMMAND_MACRO(FRHICommandBeginRenderQuery)
@@ -2309,7 +2313,7 @@ FRHICOMMAND_MACRO(FRHICommandBeginRenderQuery)
 		: RenderQuery(InRenderQuery)
 	{
 	}
-	RHI_API void Execute(FRHICommandListBase& CmdList);
+	RHI_EXECUTE_API void Execute(FRHICommandListBase& CmdList);
 };
 
 FRHICOMMAND_MACRO(FRHICommandEndRenderQuery)
@@ -2320,7 +2324,7 @@ FRHICOMMAND_MACRO(FRHICommandEndRenderQuery)
 		: RenderQuery(InRenderQuery)
 	{
 	}
-	RHI_API void Execute(FRHICommandListBase& CmdList);
+	RHI_EXECUTE_API void Execute(FRHICommandListBase& CmdList);
 };
 
 FRHICOMMAND_MACRO(FRHICommandCalibrateTimers)
@@ -2331,17 +2335,17 @@ FRHICOMMAND_MACRO(FRHICommandCalibrateTimers)
 		: CalibrationQuery(CalibrationQuery)
 	{
 	}
-	RHI_API void Execute(FRHICommandListBase & CmdList);
+	RHI_EXECUTE_API void Execute(FRHICommandListBase & CmdList);
 };
 
 FRHICOMMAND_MACRO(FRHICommandPostExternalCommandsReset)
 {
-	RHI_API void Execute(FRHICommandListBase& CmdList);
+	RHI_EXECUTE_API void Execute(FRHICommandListBase& CmdList);
 };
 
 FRHICOMMAND_MACRO(FRHICommandPollOcclusionQueries)
 {
-	RHI_API void Execute(FRHICommandListBase& CmdList);
+	RHI_EXECUTE_API void Execute(FRHICommandListBase& CmdList);
 };
 
 FRHICOMMAND_MACRO(FRHICommandBeginDrawingViewport)
@@ -2354,7 +2358,7 @@ FRHICOMMAND_MACRO(FRHICommandBeginDrawingViewport)
 		, RenderTargetRHI(InRenderTargetRHI)
 	{
 	}
-	RHI_API void Execute(FRHICommandListBase& CmdList);
+	RHI_EXECUTE_API void Execute(FRHICommandListBase& CmdList);
 };
 
 FRHICOMMAND_MACRO(FRHICommandEndDrawingViewport)
@@ -2369,7 +2373,7 @@ FRHICOMMAND_MACRO(FRHICommandEndDrawingViewport)
 		, bLockToVsync(InbLockToVsync)
 	{
 	}
-	RHI_API void Execute(FRHICommandListBase& CmdList);
+	RHI_EXECUTE_API void Execute(FRHICommandListBase& CmdList);
 };
 
 FRHICOMMAND_MACRO(FRHICommandDiscardRenderTargets)
@@ -2385,7 +2389,7 @@ FRHICOMMAND_MACRO(FRHICommandDiscardRenderTargets)
 	{
 	}
 	
-	RHI_API void Execute(FRHICommandListBase& CmdList);
+	RHI_EXECUTE_API void Execute(FRHICommandListBase& CmdList);
 };
 
 FRHICOMMAND_MACRO(FRHICommandCopyBufferRegion)
@@ -2404,7 +2408,7 @@ FRHICOMMAND_MACRO(FRHICommandCopyBufferRegion)
 		, NumBytes(InNumBytes)
 	{}
 
-	RHI_API void Execute(FRHICommandListBase& CmdList);
+	RHI_EXECUTE_API void Execute(FRHICommandListBase& CmdList);
 };
 
 FRHICOMMAND_UNNAMED(FRHICommandBindAccelerationStructureMemory)
@@ -2419,7 +2423,7 @@ FRHICOMMAND_UNNAMED(FRHICommandBindAccelerationStructureMemory)
 		, BufferOffset(InBufferOffset)
 	{}
 
-	RHI_API void Execute(FRHICommandListBase& CmdList);
+	RHI_EXECUTE_API void Execute(FRHICommandListBase& CmdList);
 };
 
 FRHICOMMAND_UNNAMED(FRHICommandBuildAccelerationStructure)
@@ -2430,7 +2434,7 @@ FRHICOMMAND_UNNAMED(FRHICommandBuildAccelerationStructure)
 		: SceneBuildParams(MoveTemp(InSceneBuildParams))
 	{}
 
-	RHI_API void Execute(FRHICommandListBase& CmdList);
+	RHI_EXECUTE_API void Execute(FRHICommandListBase& CmdList);
 };
 
 FRHICOMMAND_MACRO(FRHICommandCommitRayTracingBindings)
@@ -2441,7 +2445,7 @@ FRHICOMMAND_MACRO(FRHICommandCommitRayTracingBindings)
 		: Scene(InScene)
 	{}
 
-	RHI_API void Execute(FRHICommandListBase& CmdList);
+	RHI_EXECUTE_API void Execute(FRHICommandListBase& CmdList);
 };
 
 FRHICOMMAND_MACRO(FRHICommandCommitShaderBindingTable)
@@ -2452,7 +2456,7 @@ FRHICOMMAND_MACRO(FRHICommandCommitShaderBindingTable)
 		: SBT(InSBT)
 	{}
 
-	RHI_API void Execute(FRHICommandListBase& CmdList);
+	RHI_EXECUTE_API void Execute(FRHICommandListBase& CmdList);
 };
 
 FRHICOMMAND_MACRO(FRHICommandClearRayTracingBindings)
@@ -2463,7 +2467,7 @@ FRHICOMMAND_MACRO(FRHICommandClearRayTracingBindings)
 		: Scene(InScene)
 	{}
 
-	RHI_API void Execute(FRHICommandListBase & CmdList);
+	RHI_EXECUTE_API void Execute(FRHICommandListBase & CmdList);
 };
 
 FRHICOMMAND_MACRO(FRHICommandClearShaderBindingTable)
@@ -2474,7 +2478,7 @@ FRHICOMMAND_MACRO(FRHICommandClearShaderBindingTable)
 		: SBT(InSBT)
 	{}
 
-	RHI_API void Execute(FRHICommandListBase & CmdList);
+	RHI_EXECUTE_API void Execute(FRHICommandListBase & CmdList);
 };
 
 FRHICOMMAND_UNNAMED(FRHICommandBuildAccelerationStructures)
@@ -2489,7 +2493,7 @@ FRHICOMMAND_UNNAMED(FRHICommandBuildAccelerationStructures)
 		, ScratchBuffer(ScratchBufferRange.Buffer)
 	{}
 
-	RHI_API void Execute(FRHICommandListBase& CmdList);
+	RHI_EXECUTE_API void Execute(FRHICommandListBase& CmdList);
 };
 
 FRHICOMMAND_MACRO(FRHICommandRayTraceDispatch)
@@ -2552,7 +2556,7 @@ FRHICOMMAND_MACRO(FRHICommandRayTraceDispatch)
 		, Height(0)
 	{}
 
-	RHI_API void Execute(FRHICommandListBase& CmdList);
+	RHI_EXECUTE_API void Execute(FRHICommandListBase& CmdList);
 };
 
 FRHICOMMAND_MACRO(FRHICommandSetBindingsOnShaderBindingTable)
@@ -2586,11 +2590,11 @@ FRHICOMMAND_MACRO(FRHICommandSetBindingsOnShaderBindingTable)
 
 	}
 
-	RHI_API void Execute(FRHICommandListBase& CmdList);
+	RHI_EXECUTE_API void Execute(FRHICommandListBase& CmdList);
 };
 
-template<> RHI_API void FRHICommandSetShaderParameters           <FRHIComputeShader>::Execute(FRHICommandListBase& CmdList);
-template<> RHI_API void FRHICommandSetShaderUnbinds              <FRHIComputeShader>::Execute(FRHICommandListBase& CmdList);
+template<> RHI_EXECUTE_API void FRHICommandSetShaderParameters           <FRHIComputeShader>::Execute(FRHICommandListBase& CmdList);
+template<> RHI_EXECUTE_API void FRHICommandSetShaderUnbinds              <FRHIComputeShader>::Execute(FRHICommandListBase& CmdList);
 
 extern RHI_API FRHIComputePipelineState*	ExecuteSetComputePipelineState(FComputePipelineState* ComputePipelineState);
 extern RHI_API FRHIGraphicsPipelineState*	ExecuteSetGraphicsPipelineState(class FGraphicsPipelineState* GraphicsPipelineState);
@@ -3437,8 +3441,8 @@ public:
 	}
 };
 
-template<> RHI_API void FRHICommandSetShaderParameters           <FRHIGraphicsShader>::Execute(FRHICommandListBase& CmdList);
-template<> RHI_API void FRHICommandSetShaderUnbinds              <FRHIGraphicsShader>::Execute(FRHICommandListBase& CmdList);
+template<> RHI_EXECUTE_API void FRHICommandSetShaderParameters           <FRHIGraphicsShader>::Execute(FRHICommandListBase& CmdList);
+template<> RHI_EXECUTE_API void FRHICommandSetShaderUnbinds              <FRHIGraphicsShader>::Execute(FRHICommandListBase& CmdList);
 
 class FRHICommandList : public FRHIComputeCommandList
 {
diff --git a/Engine/Source/Runtime/RHI/Public/RHICommandListCommandExecutes.inl b/Engine/Source/Runtime/RHI/Public/RHICommandListCommandExecutes.inl
--- a/Engine/Source/Runtime/RHI/Public/RHICommandListCommandExecutes.inl
+++ b/Engine/Source/Runtime/RHI/Public/RHICommandListCommandExecutes.inl
@@ -117,25 +117,25 @@ void FRHICommandSetStencilRef::Execute(FRHICommandListBase& CmdList)
 	INTERNAL_DECORATOR(RHISetStencilRef)(StencilRef);
 }
 
-template<> RHI_API void FRHICommandSetShaderParameters<FRHIComputeShader>::Execute(FRHICommandListBase& CmdList)
+template<> RHI_EXECUTE_API void FRHICommandSetShaderParameters<FRHIComputeShader>::Execute(FRHICommandListBase& CmdList)
 {
 	RHISTAT(SetShaderParameters);
 	INTERNAL_DECORATOR_COMPUTE(RHISetShaderParameters)(Shader, ParametersData, Parameters, ResourceParameters, BindlessParameters);
 }
 
-template<> RHI_API void FRHICommandSetShaderParameters<FRHIGraphicsShader>::Execute(FRHICommandListBase& CmdList)
+template<> RHI_EXECUTE_API void FRHICommandSetShaderParameters<FRHIGraphicsShader>::Execute(FRHICommandListBase& CmdList)
 {
 	RHISTAT(SetShaderParameters);
 	INTERNAL_DECORATOR(RHISetShaderParameters)(Shader, ParametersData, Parameters, ResourceParameters, BindlessParameters);
 }
 
-template<> RHI_API void FRHICommandSetShaderUnbinds<FRHIComputeShader>::Execute(FRHICommandListBase& CmdList)
+template<> RHI_EXECUTE_API void FRHICommandSetShaderUnbinds<FRHIComputeShader>::Execute(FRHICommandListBase& CmdList)
 {
 	RHISTAT(SetShaderUnbinds);
 	INTERNAL_DECORATOR_COMPUTE(RHISetShaderUnbinds)(Shader, Unbinds);
 }
 
-template<> RHI_API void FRHICommandSetShaderUnbinds<FRHIGraphicsShader>::Execute(FRHICommandListBase& CmdList)
+template<> RHI_EXECUTE_API void FRHICommandSetShaderUnbinds<FRHIGraphicsShader>::Execute(FRHICommandListBase& CmdList)
 {
 	RHISTAT(SetShaderUnbinds);
 	INTERNAL_DECORATOR(RHISetShaderUnbinds)(Shader, Unbinds);
