Implemented workaround for uniform buffer creation using a global render thread id. If command list recording straddles EndFrame boundaries this causes a validation error.

diff --git a/Engine/Source/Runtime/RHI/Private/RHIValidation.cpp b/Engine/Source/Runtime/RHI/Private/RHIValidation.cpp
--- a/Engine/Source/Runtime/RHI/Private/RHIValidation.cpp
+++ b/Engine/Source/Runtime/RHI/Private/RHIValidation.cpp
@@ -2331,7 +2331,7 @@ namespace RHIValidation
 
 		RHI_VALIDATION_CHECK(bContainsNullContents == false, TEXT("Uniform buffer created with null contents is now being bound for rendering on an RHI context. The contents must first be updated."));
 
-		if (UniformBufferUsage != UniformBuffer_MultiFrame && AllocatedFrameID != ValidateRHI->RHIThreadFrameID)
+		if (UniformBufferUsage != UniformBuffer_MultiFrame && AllocatedFrameID < ValidateRHI->RHIThreadFrameID)
 		{
 			FString ErrorMessage = TEXT("Non MultiFrame Uniform buffer has been allocated in a previous frame. The data could have been deleted already!");
 			if (AllocatedCallstack != nullptr)
diff --git a/Engine/Source/Runtime/RHI/Public/RHIValidation.h b/Engine/Source/Runtime/RHI/Public/RHIValidation.h
--- a/Engine/Source/Runtime/RHI/Public/RHIValidation.h
+++ b/Engine/Source/Runtime/RHI/Public/RHIValidation.h
@@ -396,8 +396,13 @@ public:
 		check(Layout);
 		check(Layout->Resources.Num() > 0 || Layout->ConstantBufferSize > 0);
 		FUniformBufferRHIRef UniformBuffer = RHI->RHICreateUniformBuffer(Contents, Layout, Usage, Validation);
-		// Use the render thread frame ID for any non RHI thread allocations.
-		UniformBuffer->InitLifetimeTracking(IsInRHIThread() ? RHIThreadFrameID : RenderThreadFrameID, Contents, Usage);
+
+		// Use the render thread frame ID for any non RHI thread allocations. TODO: This is actually incorrect as command list recording on the render thread timeline
+		// can straddle EndFrame boundaries, causing a uniform buffer allocated in frame N to be recorded as being allocated in frame N+1. The solution here is to introduce
+		// a command list to RHICreateUniformBuffer so that the correct render thread frame can be propagated. Unfortunately, in the meantime, this means that the lifetime
+		// tracker can miss legitimate cases. For example, if a (single frame) uniform buffer is allocated in frame N (but straddles to frame N+1), and then is incorrectly
+		// used in frame N+1, that test will pass because they are equal. However, since the issue is timing dependent, it's still likely to catch legitimate allocation misuses.
+		UniformBuffer->InitLifetimeTracking(IsInRHIThread() ? RHIThreadFrameID : RenderThreadFrameID.load(std::memory_order_relaxed), Contents, Usage);
 		return UniformBuffer;
 	}
 
@@ -406,7 +411,11 @@ public:
 		check(UniformBufferRHI);
 		check(Contents);
 		RHI->RHIUpdateUniformBuffer(RHICmdList, UniformBufferRHI, Contents);
-		UniformBufferRHI->UpdateAllocation(RenderThreadFrameID);
+
+		RHICmdList.EnqueueLambda([this, UniformBufferRHI] (FRHICommandListBase& RHICmdList)
+		{
+			UniformBufferRHI->UpdateAllocation(RHIThreadFrameID);
+		});
 	}
 
 	/**
@@ -1248,7 +1257,7 @@ public:
 	FDynamicRHI*				RHI;
 	TMap<FRHIDepthStencilState*, FDepthStencilStateInitializerRHI> DepthStencilStates;
 
-	uint64						RenderThreadFrameID;
+	std::atomic_uint64_t		RenderThreadFrameID;
 	uint64						RHIThreadFrameID;
 
 private:
