Make r.PSOPrecache.KeepInMemoryUntilUsed only work on NVIDIA since there is no benefit on any other vendor.

diff --git a/Engine/Source/Runtime/RHI/Private/PipelineStateCache.cpp b/Engine/Source/Runtime/RHI/Private/PipelineStateCache.cpp
--- a/Engine/Source/Runtime/RHI/Private/PipelineStateCache.cpp
+++ b/Engine/Source/Runtime/RHI/Private/PipelineStateCache.cpp
@@ -318,8 +318,8 @@ int32 GPSOPrecacheKeepInMemoryUntilUsed = 0;
 static FAutoConsoleVariableRef CVarPSOPrecacheKeepInMemoryUntilUsed(
 	TEXT("r.PSOPrecache.KeepInMemoryUntilUsed"),
 	GPSOPrecacheKeepInMemoryUntilUsed,
-	TEXT("If enabled, precached PSOs will be kept in memory instead of being deleted immediately after creation, and will only be deleted once they are actually used for rendering.\n")
-	TEXT("This can speed up the re-creation of precached PSOs for certain drivers and avoid small hitches, at the cost of memory.\n")
+	TEXT("If enabled and if the underlying GPU vendor is NVIDIA, precached PSOs will be kept in memory instead of being deleted immediately after creation, and will only be deleted once they are actually used for rendering.\n")
+	TEXT("This can speed up the re-creation of precached PSOs for NVIDIA drivers and avoid small hitches, at the cost of memory.\n")
 	TEXT("It's recommended to set r.PSOPrecache.KeepInMemoryGraphicsMaxNum and r.PSOPrecache.KeepInMemoryComputeMaxNum to a non-zero value to ensure the number of in-memory PSOs is bounded."),
 	ECVF_ReadOnly);
 
@@ -339,6 +339,11 @@ static FAutoConsoleVariableRef CVarPSOPrecacheKeepInMemoryComputeMaxNum(
 	TEXT("If set to 0, no limit will be applied (not recommended outside of testing, as it can cause unbounded memory usage)."),
 	ECVF_ReadOnly);
 
+bool ShouldKeepPrecachedPSOsInMemory()
+{
+	return GPSOPrecacheKeepInMemoryUntilUsed && IsRHIDeviceNVIDIA();
+}
+
 const TCHAR* LexToString(EPSOPrecacheResult Result)
 {
 	switch (Result)
@@ -1938,7 +1943,7 @@ public:
 		: PSOType(InType),
 		MaxInMemoryPSOs(InMaxInMemoryPSOs)
 	{
-		if (GPSOPrecacheKeepInMemoryUntilUsed && InMaxInMemoryPSOs > 0)
+		if (ShouldKeepPrecachedPSOsInMemory() && InMaxInMemoryPSOs > 0)
 		{
 			InMemoryPSOIndices.Init(kInvalidPSOIndex, InMaxInMemoryPSOs);
 		}
@@ -2035,7 +2040,7 @@ protected:
 				UpdateActiveCompileCount(true /*Increment*/);
 			}
 
-			if (GPSOPrecacheKeepInMemoryUntilUsed)
+			if (ShouldKeepPrecachedPSOsInMemory())
 			{
 				if (MaxInMemoryPSOs > 0)
 				{
@@ -2198,7 +2203,7 @@ public:
 			PreviousStateMask = FindResult->AddPSOPrecacheState(CompleteStateMask);
 
 			// Add to array of precached PSOs so it can be cleaned up
-			if (!GPSOPrecacheKeepInMemoryUntilUsed)
+			if (!ShouldKeepPrecachedPSOsInMemory())
 			{
 				PrecachedPSOsToCleanup.Add(InitializerHash);
 			}
@@ -2237,7 +2242,7 @@ public:
 			uint64 InitializerHash = PrecachedPSOsToCleanup[Index];
 
 			FPrecacheTask* FindResult = PrecachedPSOInitializerData.Find(InitializerHash);
-			check(FindResult && (GPSOPrecacheKeepInMemoryUntilUsed || IsCompilationDone((FindResult->ReadPSOPrecacheState()))));
+			check(FindResult && (ShouldKeepPrecachedPSOsInMemory() || IsCompilationDone((FindResult->ReadPSOPrecacheState()))));
 			if (!FindResult || !FindResult->PipelineState)
 			{
 				// Was already cleaned up (can happen if it was marked as used).
@@ -2251,7 +2256,7 @@ public:
 				delete FindResult->PipelineState;
 				FindResult->PipelineState = nullptr;
 
-				if (GPSOPrecacheKeepInMemoryUntilUsed)
+				if (ShouldKeepPrecachedPSOsInMemory())
 				{
 					DEC_DWORD_STAT(STAT_InMemoryPrecachedPSOCount);
 				}
@@ -2266,7 +2271,7 @@ public:
 	// Does not do anything if PSOs are not kept in memory, or if a bound on the number of PSOs kept in memory is set.
 	void MarkPSOAsUsed(const TPrecachedPSOInitializer& Initializer)
 	{
-		if (!GPSOPrecacheKeepInMemoryUntilUsed || MaxInMemoryPSOs > 0)
+		if (!ShouldKeepPrecachedPSOsInMemory() || MaxInMemoryPSOs > 0)
 		{
 			return;
 		}
