[HWRT] Deprecated FRayTracingSceneInitializer2 and use FRayTracingSceneInitializer instead.

diff --git a/Engine/Plugins/Experimental/GPULightmass/Source/GPULightmass/Private/LightmapRenderer.cpp b/Engine/Plugins/Experimental/GPULightmass/Source/GPULightmass/Private/LightmapRenderer.cpp
--- a/Engine/Plugins/Experimental/GPULightmass/Source/GPULightmass/Private/LightmapRenderer.cpp
+++ b/Engine/Plugins/Experimental/GPULightmass/Source/GPULightmass/Private/LightmapRenderer.cpp
@@ -1151,7 +1151,7 @@ bool FSceneRenderState::SetupRayTracingScene(FRDGBuilder& GraphBuilder, FSceneUn
 			const ERayTracingAccelerationStructureFlags SceneBuildFlags = ERayTracingAccelerationStructureFlags::FastTrace;
 
 			{
-				FRayTracingSceneInitializer2 Initializer;
+				FRayTracingSceneInitializer Initializer;
 				Initializer.DebugName = FName(TEXT("LightmapRendererRayTracingScene"));
 				Initializer.MaxNumInstances = RayTracingSceneInitializationData.NumNativeGPUSceneInstances + RayTracingSceneInitializationData.NumNativeCPUInstances;
 				Initializer.NumTotalSegments = RayTracingSceneInitializationData.TotalNumSegments;
@@ -1160,7 +1160,7 @@ bool FSceneRenderState::SetupRayTracingScene(FRDGBuilder& GraphBuilder, FSceneUn
 				RayTracingScene = RHICreateRayTracingScene(MoveTemp(Initializer));
 			}
 
-			const FRayTracingSceneInitializer2& SceneInitializer = RayTracingScene->GetInitializer();
+			const FRayTracingSceneInitializer& SceneInitializer = RayTracingScene->GetInitializer();
 
 			FRayTracingAccelerationStructureSize SizeInfo = RHICalcRayTracingSceneSize(SceneInitializer);
 			FRHIResourceCreateInfo BufferCreateInfo(TEXT("LightmassRayTracingSceneBuffer"));
diff --git a/Engine/Source/Runtime/Apple/MetalRHI/Private/MetalDynamicRHI.h b/Engine/Source/Runtime/Apple/MetalRHI/Private/MetalDynamicRHI.h
--- a/Engine/Source/Runtime/Apple/MetalRHI/Private/MetalDynamicRHI.h
+++ b/Engine/Source/Runtime/Apple/MetalRHI/Private/MetalDynamicRHI.h
@@ -155,11 +155,11 @@ public:
 	virtual uint16 RHIGetPlatformTextureMaxSampleCount() override;
 
 #if METAL_RHI_RAYTRACING
-	virtual FRayTracingAccelerationStructureSize RHICalcRayTracingSceneSize(const FRayTracingSceneInitializer2& Initializer) final override;
+	virtual FRayTracingAccelerationStructureSize RHICalcRayTracingSceneSize(const FRayTracingSceneInitializer& Initializer) final override;
 	virtual FRayTracingAccelerationStructureSize RHICalcRayTracingGeometrySize(const FRayTracingGeometryInitializer& Initializer) final override;
 
 	virtual FRayTracingGeometryRHIRef RHICreateRayTracingGeometry(FRHICommandListBase& RHICmdList, const FRayTracingGeometryInitializer& Initializer) final override;
-	virtual FRayTracingSceneRHIRef RHICreateRayTracingScene(FRayTracingSceneInitializer2 Initializer) final override;
+	virtual FRayTracingSceneRHIRef RHICreateRayTracingScene(FRayTracingSceneInitializer Initializer) final override;
 	virtual FRayTracingShaderRHIRef RHICreateRayTracingShader(TArrayView<const uint8> Code, const FSHAHash& Hash, EShaderFrequency ShaderFrequency) final override;
 	virtual FRayTracingPipelineStateRHIRef RHICreateRayTracingPipelineState(const FRayTracingPipelineStateInitializer& Initializer) final override;
 
diff --git a/Engine/Source/Runtime/Apple/MetalRHI/Private/MetalRayTracing.cpp b/Engine/Source/Runtime/Apple/MetalRHI/Private/MetalRayTracing.cpp
--- a/Engine/Source/Runtime/Apple/MetalRHI/Private/MetalRayTracing.cpp
+++ b/Engine/Source/Runtime/Apple/MetalRHI/Private/MetalRayTracing.cpp
@@ -314,7 +314,7 @@ FRayTracingAccelerationStructureSize FMetalDynamicRHI::RHICalcRayTracingGeometry
 	return CalcRayTracingGeometrySize(AccelerationStructureDescriptor);
 }
 
-FRayTracingAccelerationStructureSize FMetalDynamicRHI::RHICalcRayTracingSceneSize(const FRayTracingSceneInitializer2& Initializer)
+FRayTracingAccelerationStructureSize FMetalDynamicRHI::RHICalcRayTracingSceneSize(const FRayTracingSceneInitializer& Initializer)
 {
 	// TODO: Do we need to take in account the flags provided by the function call?
 	// TODO: Can we get away with the instance count only? (works on AS; what about AMD?)
@@ -420,7 +420,7 @@ void FMetalRayTracingGeometry::RebuildDescriptors()
 	FillPrimitiveAccelerationStructureDesc(AccelerationStructureDescriptor, Initializer, GeomArray);
 }
 
-FMetalRayTracingScene::FMetalRayTracingScene(FRayTracingSceneInitializer2 InInitializer)
+FMetalRayTracingScene::FMetalRayTracingScene(FRayTracingSceneInitializer InInitializer)
 	: Initializer(MoveTemp(InInitializer))
 {
 	MTL::InstanceAccelerationStructureDescriptor* InstanceDescriptor;
@@ -805,7 +805,7 @@ void FMetalRHICommandContext::RHISetBindingsOnShaderBindingTable(
 	checkNoEntry();
 }
 
-FRayTracingSceneRHIRef FMetalDynamicRHI::RHICreateRayTracingScene(FRayTracingSceneInitializer2 Initializer)
+FRayTracingSceneRHIRef FMetalDynamicRHI::RHICreateRayTracingScene(FRayTracingSceneInitializer Initializer)
 {
     MTL_SCOPED_AUTORELEASE_POOL;
     return new FMetalRayTracingScene(MoveTemp(Initializer));
diff --git a/Engine/Source/Runtime/Apple/MetalRHI/Private/MetalRayTracing.h b/Engine/Source/Runtime/Apple/MetalRHI/Private/MetalRayTracing.h
--- a/Engine/Source/Runtime/Apple/MetalRHI/Private/MetalRayTracing.h
+++ b/Engine/Source/Runtime/Apple/MetalRHI/Private/MetalRayTracing.h
@@ -74,7 +74,7 @@ private:
 class FMetalRayTracingScene : public FRHIRayTracingScene
 {
 public:
-	FMetalRayTracingScene(FRayTracingSceneInitializer2 InInitializer);
+	FMetalRayTracingScene(FRayTracingSceneInitializer InInitializer);
 	virtual ~FMetalRayTracingScene();
 
 	void BindBuffer(FRHIBuffer* InBuffer, uint32 InBufferOffset);
@@ -85,7 +85,7 @@ public:
 
 	void BuildPerInstanceGeometryParameterBuffer();
 
-	inline const FRayTracingSceneInitializer2& GetInitializer() const override final { return Initializer; }
+	inline const FRayTracingSceneInitializer& GetInitializer() const override final { return Initializer; }
 
 	TRefCountPtr<FMetalShaderResourceView> InstanceBufferSRV;
 
@@ -96,7 +96,7 @@ private:
 
 private:
 	/** The initializer provided to build the scene. Contains all the free standing stuff that used to be owned by the RT implementation. */
-	const FRayTracingSceneInitializer2 Initializer;
+	const FRayTracingSceneInitializer Initializer;
 
 	/** Acceleration Structure for the whole scene. */
 	TRefCountPtr<FMetalRHIBuffer> AccelerationStructureBuffer;
diff --git a/Engine/Source/Runtime/D3D12RHI/Private/D3D12RHIPrivate.h b/Engine/Source/Runtime/D3D12RHI/Private/D3D12RHIPrivate.h
--- a/Engine/Source/Runtime/D3D12RHI/Private/D3D12RHIPrivate.h
+++ b/Engine/Source/Runtime/D3D12RHI/Private/D3D12RHIPrivate.h
@@ -466,11 +466,11 @@ public:
 
 #if D3D12_RHI_RAYTRACING
 
-	virtual FRayTracingAccelerationStructureSize RHICalcRayTracingSceneSize(const FRayTracingSceneInitializer2& Initializer) final override;
+	virtual FRayTracingAccelerationStructureSize RHICalcRayTracingSceneSize(const FRayTracingSceneInitializer& Initializer) final override;
 	virtual FRayTracingAccelerationStructureSize RHICalcRayTracingGeometrySize(const FRayTracingGeometryInitializer& Initializer) final override;
 
 	virtual FRayTracingGeometryRHIRef RHICreateRayTracingGeometry(FRHICommandListBase& RHICmdList, const FRayTracingGeometryInitializer& Initializer) final override;
-	virtual FRayTracingSceneRHIRef RHICreateRayTracingScene(FRayTracingSceneInitializer2 Initializer) final override;
+	virtual FRayTracingSceneRHIRef RHICreateRayTracingScene(FRayTracingSceneInitializer Initializer) final override;
 	virtual FRayTracingShaderRHIRef RHICreateRayTracingShader(TArrayView<const uint8> Code, const FSHAHash& Hash, EShaderFrequency ShaderFrequency) final override;
 	virtual FRayTracingPipelineStateRHIRef RHICreateRayTracingPipelineState(const FRayTracingPipelineStateInitializer& Initializer) final override;
 
diff --git a/Engine/Source/Runtime/D3D12RHI/Private/D3D12RayTracing.cpp b/Engine/Source/Runtime/D3D12RHI/Private/D3D12RayTracing.cpp
--- a/Engine/Source/Runtime/D3D12RHI/Private/D3D12RayTracing.cpp
+++ b/Engine/Source/Runtime/D3D12RHI/Private/D3D12RayTracing.cpp
@@ -2448,7 +2448,7 @@ static D3D12_RAYTRACING_INSTANCE_FLAGS TranslateRayTracingInstanceFlags(ERayTrac
 	return Result;
 }
 
-FRayTracingAccelerationStructureSize FD3D12DynamicRHI::RHICalcRayTracingSceneSize(const FRayTracingSceneInitializer2& Initializer)
+FRayTracingAccelerationStructureSize FD3D12DynamicRHI::RHICalcRayTracingSceneSize(const FRayTracingSceneInitializer& Initializer)
 {
 	D3D12_BUILD_RAYTRACING_ACCELERATION_STRUCTURE_INPUTS BuildInputs = {};
 	BuildInputs.Type = D3D12_RAYTRACING_ACCELERATION_STRUCTURE_TYPE_TOP_LEVEL;
@@ -2698,7 +2698,7 @@ FRayTracingGeometryRHIRef FD3D12DynamicRHI::RHICreateRayTracingGeometry(FRHIComm
 	return new FD3D12RayTracingGeometry(RHICmdList, &Adapter, Initializer);
 }
 
-FRayTracingSceneRHIRef FD3D12DynamicRHI::RHICreateRayTracingScene(FRayTracingSceneInitializer2 Initializer)
+FRayTracingSceneRHIRef FD3D12DynamicRHI::RHICreateRayTracingScene(FRayTracingSceneInitializer Initializer)
 {
 	TRACE_CPUPROFILER_EVENT_SCOPE(CreateRayTracingScene);
 
@@ -3507,7 +3507,7 @@ void FD3D12RayTracingGeometry::CompactAccelerationStructure(FD3D12CommandContext
 	RegisterD3D12RayTracingGeometry(this);
 }
 
-FD3D12RayTracingScene::FD3D12RayTracingScene(FD3D12Adapter* Adapter, FRayTracingSceneInitializer2 InInitializer)
+FD3D12RayTracingScene::FD3D12RayTracingScene(FD3D12Adapter* Adapter, FRayTracingSceneInitializer InInitializer)
 	: FD3D12AdapterChild(Adapter), Initializer(MoveTemp(InInitializer))
 {
 	INC_DWORD_STAT(STAT_D3D12RayTracingAllocatedTLAS);
diff --git a/Engine/Source/Runtime/D3D12RHI/Private/D3D12RayTracing.h b/Engine/Source/Runtime/D3D12RHI/Private/D3D12RayTracing.h
--- a/Engine/Source/Runtime/D3D12RHI/Private/D3D12RayTracing.h
+++ b/Engine/Source/Runtime/D3D12RHI/Private/D3D12RayTracing.h
@@ -92,10 +92,10 @@ class FD3D12RayTracingScene : public FRHIRayTracingScene, public FD3D12AdapterCh
 {
 public:
 
-	FD3D12RayTracingScene(FD3D12Adapter* Adapter, FRayTracingSceneInitializer2 Initializer);
+	FD3D12RayTracingScene(FD3D12Adapter* Adapter, FRayTracingSceneInitializer Initializer);
 	~FD3D12RayTracingScene();
 
-	const FRayTracingSceneInitializer2& GetInitializer() const override final { return Initializer; }
+	const FRayTracingSceneInitializer& GetInitializer() const override final { return Initializer; }
 
 	void BindBuffer(FRHIBuffer* Buffer, uint32 BufferOffset);
 	void ReleaseBuffer();
@@ -107,7 +107,7 @@ public:
 	TRefCountPtr<FD3D12Buffer> AccelerationStructureBuffers[MAX_NUM_GPUS];
 	uint32 BufferOffset = 0;
 
-	const FRayTracingSceneInitializer2 Initializer;
+	const FRayTracingSceneInitializer Initializer;
 
 	// Unique list of geometries referenced by all instances in this scene.
 	// Any referenced geometry is kept alive while the scene is alive.
diff --git a/Engine/Source/Runtime/D3D12RHI/Private/D3D12RayTracingDebug.cpp b/Engine/Source/Runtime/D3D12RHI/Private/D3D12RayTracingDebug.cpp
--- a/Engine/Source/Runtime/D3D12RHI/Private/D3D12RayTracingDebug.cpp
+++ b/Engine/Source/Runtime/D3D12RHI/Private/D3D12RayTracingDebug.cpp
@@ -19,7 +19,7 @@ static void DebugSerializeScene(const FD3D12RayTracingScene& Scene, FD3D12Buffer
 {
 	// #dxr_todo: this could potentially be generalized and moved to high-level code, to be reused for all RHIs if we implement instance desc structure decoding
 
-	const FRayTracingSceneInitializer2& SceneInitializer = Scene.GetInitializer();
+	const FRayTracingSceneInitializer& SceneInitializer = Scene.GetInitializer();
 
 	FString Name = SceneInitializer.DebugName.ToString();
 	FString Filename = FString::Printf(TEXT("Scene_%s"), *Name);
diff --git a/Engine/Source/Runtime/RHI/Public/DynamicRHI.h b/Engine/Source/Runtime/RHI/Public/DynamicRHI.h
--- a/Engine/Source/Runtime/RHI/Public/DynamicRHI.h
+++ b/Engine/Source/Runtime/RHI/Public/DynamicRHI.h
@@ -949,7 +949,7 @@ public:
 
 	virtual uint16 RHIGetPlatformTextureMaxSampleCount() { return 8; };
 
-	virtual FRayTracingAccelerationStructureSize RHICalcRayTracingSceneSize(const FRayTracingSceneInitializer2& Initializer)
+	virtual FRayTracingAccelerationStructureSize RHICalcRayTracingSceneSize(const FRayTracingSceneInitializer& Initializer)
 	{
 		checkNoEntry();
 		return {};
@@ -967,7 +967,7 @@ public:
 		return nullptr;
 	}
 	
-	virtual FRayTracingSceneRHIRef RHICreateRayTracingScene(FRayTracingSceneInitializer2 Initializer)
+	virtual FRayTracingSceneRHIRef RHICreateRayTracingScene(FRayTracingSceneInitializer Initializer)
 	{
 		checkNoEntry();
 		return nullptr;
@@ -1467,22 +1467,22 @@ FORCEINLINE bool RHIMatchPrecachePSOInitializers(const FGraphicsPipelineStateIni
 	return GDynamicRHI->RHIMatchPrecachePSOInitializers(LHS, RHS);
 }
 
-FORCEINLINE FRayTracingAccelerationStructureSize RHICalcRayTracingSceneSize(const FRayTracingSceneInitializer2& Initializer)
+FORCEINLINE FRayTracingAccelerationStructureSize RHICalcRayTracingSceneSize(const FRayTracingSceneInitializer& Initializer)
 {
 	return GDynamicRHI->RHICalcRayTracingSceneSize(Initializer);
 }
 
-UE_DEPRECATED(5.5, "Provide FRayTracingSceneInitializer2 instead.")
+UE_DEPRECATED(5.5, "Provide FRayTracingSceneInitializer instead.")
 FORCEINLINE FRayTracingAccelerationStructureSize RHICalcRayTracingSceneSize(uint32 MaxInstances, ERayTracingAccelerationStructureFlags Flags)
 {
-	FRayTracingSceneInitializer2 Initializer;
+	FRayTracingSceneInitializer Initializer;
 	Initializer.MaxNumInstances = MaxInstances;
 	Initializer.BuildFlags = Flags;
 
 	return GDynamicRHI->RHICalcRayTracingSceneSize(Initializer);
 }
 
-FORCEINLINE FRayTracingSceneRHIRef RHICreateRayTracingScene(FRayTracingSceneInitializer2 Initializer)
+FORCEINLINE FRayTracingSceneRHIRef RHICreateRayTracingScene(FRayTracingSceneInitializer Initializer)
 {
 	PRAGMA_DISABLE_DEPRECATION_WARNINGS
 	if (Initializer.NumNativeInstancesPerLayer.Num() >= 1)
diff --git a/Engine/Source/Runtime/RHI/Public/RHIResources.h b/Engine/Source/Runtime/RHI/Public/RHIResources.h
--- a/Engine/Source/Runtime/RHI/Public/RHIResources.h
+++ b/Engine/Source/Runtime/RHI/Public/RHIResources.h
@@ -3306,7 +3306,7 @@ struct FRayTracingShaderBindingTableInitializer
 	uint32 NumCallableShaderSlots = 0;
 };
 
-struct FRayTracingSceneInitializer2
+struct FRayTracingSceneInitializer
 {
 	// Unique list of geometries referenced by all instances in this scene.
 	// Any referenced geometry is kept alive while the scene is alive.
@@ -3351,12 +3351,14 @@ struct FRayTracingSceneInitializer2
 	FName DebugName;
 
 	PRAGMA_DISABLE_DEPRECATION_WARNINGS
-	FRayTracingSceneInitializer2() = default;
-	FRayTracingSceneInitializer2(FRayTracingSceneInitializer2&&) = default;
-	FRayTracingSceneInitializer2& operator=(FRayTracingSceneInitializer2&&) = default;
+	FRayTracingSceneInitializer() = default;
+	FRayTracingSceneInitializer(FRayTracingSceneInitializer&&) = default;
+	FRayTracingSceneInitializer& operator=(FRayTracingSceneInitializer&&) = default;
 	PRAGMA_ENABLE_DEPRECATION_WARNINGS
 };
 
+using FRayTracingSceneInitializer2 UE_DEPRECATED(5.5, "Use FRayTracingSceneInitializer instead") = FRayTracingSceneInitializer;
+
 struct FRayTracingAccelerationStructureSize
 {
 	uint64 ResultSize = 0;
@@ -3414,7 +3416,7 @@ class FRHIRayTracingScene
 	: public FRHIRayTracingAccelerationStructure
 {
 public:
-	virtual const FRayTracingSceneInitializer2& GetInitializer() const = 0;
+	virtual const FRayTracingSceneInitializer& GetInitializer() const = 0;
 
 	UE_DEPRECATED(5.5, "Use GetOrCreateInlineBufferSRV on the FRHIShaderBindingTable instead of the Scene")
 	virtual FRHIShaderResourceView* GetOrCreateMetadataBufferSRV(FRHICommandListImmediate& RHICmdList)
diff --git a/Engine/Source/Runtime/RHI/Public/RHIValidation.h b/Engine/Source/Runtime/RHI/Public/RHIValidation.h
--- a/Engine/Source/Runtime/RHI/Public/RHIValidation.h
+++ b/Engine/Source/Runtime/RHI/Public/RHIValidation.h
@@ -1212,7 +1212,7 @@ public:
 		return Result;
 	}
 
-	virtual FRayTracingSceneRHIRef RHICreateRayTracingScene(FRayTracingSceneInitializer2 Initializer) override final
+	virtual FRayTracingSceneRHIRef RHICreateRayTracingScene(FRayTracingSceneInitializer Initializer) override final
 	{
 		FName DebugName = Initializer.DebugName;
 		FRayTracingSceneRHIRef Result = RHI->RHICreateRayTracingScene(MoveTemp(Initializer));
@@ -1230,7 +1230,7 @@ public:
 		return RHI->RHICreateRayTracingPipelineState(Initializer);
 	}
 
-	virtual FRayTracingAccelerationStructureSize RHICalcRayTracingSceneSize(const FRayTracingSceneInitializer2& Initializer) override final
+	virtual FRayTracingAccelerationStructureSize RHICalcRayTracingSceneSize(const FRayTracingSceneInitializer& Initializer) override final
 	{
 		return RHI->RHICalcRayTracingSceneSize(Initializer);
 	}
diff --git a/Engine/Source/Runtime/Renderer/Private/RayTracing/RayTracingInstanceBufferUtil.cpp b/Engine/Source/Runtime/Renderer/Private/RayTracing/RayTracingInstanceBufferUtil.cpp
--- a/Engine/Source/Runtime/Renderer/Private/RayTracing/RayTracingInstanceBufferUtil.cpp
+++ b/Engine/Source/Runtime/Renderer/Private/RayTracing/RayTracingInstanceBufferUtil.cpp
@@ -114,7 +114,7 @@ FRayTracingSceneWithGeometryInstances CreateRayTracingSceneWithGeometryInstances
 	Output.BaseInstancePrefixSum.SetNumUninitialized(NumSceneInstances);
 	Output.PerInstanceGeometries.SetNumUninitialized(NumSceneInstances);
 
-	FRayTracingSceneInitializer2 Initializer;
+	FRayTracingSceneInitializer Initializer;
 	Initializer.DebugName = FName(TEXT("FRayTracingScene"));
 	Initializer.BuildFlags = BuildFlags;
 
@@ -338,7 +338,7 @@ void FillRayTracingInstanceUploadBuffer(
 	TArrayView<FRayTracingInstanceDescriptorInput> OutInstanceUploadData,
 	TArrayView<FVector4f> OutTransformData)
 {
-	const FRayTracingSceneInitializer2& SceneInitializer = RayTracingSceneRHI->GetInitializer();
+	const FRayTracingSceneInitializer& SceneInitializer = RayTracingSceneRHI->GetInitializer();
 
 	FillRayTracingInstanceUploadBuffer(
 		RayTracingSceneRHI,
diff --git a/Engine/Source/Runtime/Renderer/Private/RayTracing/RayTracingScene.cpp b/Engine/Source/Runtime/Renderer/Private/RayTracing/RayTracingScene.cpp
--- a/Engine/Source/Runtime/Renderer/Private/RayTracing/RayTracingScene.cpp
+++ b/Engine/Source/Runtime/Renderer/Private/RayTracing/RayTracingScene.cpp
@@ -95,7 +95,7 @@ void FRayTracingScene::Create(FRDGBuilder& GraphBuilder, const FViewInfo& View,
 			TArray<TRefCountPtr<FRHIRayTracingGeometry>> ReferencedGeometries;
 			TArray<FRHIRayTracingGeometry*> PerInstanceGeometries;
 
-			FRayTracingSceneInitializer2 Initializer;
+			FRayTracingSceneInitializer Initializer;
 			Initializer.DebugName = FName(TEXT("FRayTracingScene"));
 			Initializer.MaxNumInstances = Layer.InitializationData.NumNativeGPUSceneInstances + Layer.InitializationData.NumNativeCPUInstances;
 			Initializer.NumTotalSegments = Layer.InitializationData.TotalNumSegments;
diff --git a/Engine/Source/Runtime/Renderer/Private/Tests/RayTracingTestbed.cpp b/Engine/Source/Runtime/Renderer/Private/Tests/RayTracingTestbed.cpp
--- a/Engine/Source/Runtime/Renderer/Private/Tests/RayTracingTestbed.cpp
+++ b/Engine/Source/Runtime/Renderer/Private/Tests/RayTracingTestbed.cpp
@@ -182,7 +182,7 @@ bool RunRayTracingTestbed_RenderThread(const FString& Parameters)
 
 	FRayTracingSceneRHIRef RayTracingSceneRHI;
 	{
-		FRayTracingSceneInitializer2 Initializer;
+		FRayTracingSceneInitializer Initializer;
 		Initializer.DebugName = FName(TEXT("FRayTracingScene"));
 		Initializer.MaxNumInstances = RayTracingSceneInitializationData.NumNativeGPUSceneInstances + RayTracingSceneInitializationData.NumNativeCPUInstances;
 		Initializer.NumTotalSegments = RayTracingSceneInitializationData.TotalNumSegments;
@@ -191,7 +191,7 @@ bool RunRayTracingTestbed_RenderThread(const FString& Parameters)
 		RayTracingSceneRHI = RHICreateRayTracingScene(MoveTemp(Initializer));
 	}
 
-	const FRayTracingSceneInitializer2& SceneInitializer = RayTracingSceneRHI->GetInitializer();
+	const FRayTracingSceneInitializer& SceneInitializer = RayTracingSceneRHI->GetInitializer();
 
 	FRayTracingAccelerationStructureSize SceneSizeInfo = RHICalcRayTracingSceneSize(SceneInitializer);
 	FRHIResourceCreateInfo SceneBufferCreateInfo(TEXT("RayTracingTestBedSceneBuffer"));
diff --git a/Engine/Source/Runtime/VulkanRHI/Private/VulkanRayTracing.cpp b/Engine/Source/Runtime/VulkanRHI/Private/VulkanRayTracing.cpp
--- a/Engine/Source/Runtime/VulkanRHI/Private/VulkanRayTracing.cpp
+++ b/Engine/Source/Runtime/VulkanRHI/Private/VulkanRayTracing.cpp
@@ -701,7 +701,7 @@ static VkGeometryInstanceFlagsKHR TranslateRayTracingInstanceFlags(ERayTracingIn
 	return Result;
 }
 
-FVulkanRayTracingScene::FVulkanRayTracingScene(FRayTracingSceneInitializer2 InInitializer, FVulkanDevice* InDevice)
+FVulkanRayTracingScene::FVulkanRayTracingScene(FRayTracingSceneInitializer InInitializer, FVulkanDevice* InDevice)
 	: FDeviceChild(InDevice)
 	, Initializer(MoveTemp(InInitializer))
 
@@ -1138,7 +1138,7 @@ FRHIShaderBindingTable* FVulkanRayTracingScene::FindOrCreateShaderBindingTable(c
 	return CreatedShaderTable;
 }
 
-FRayTracingAccelerationStructureSize FVulkanDynamicRHI::RHICalcRayTracingSceneSize(const FRayTracingSceneInitializer2& Initializer)
+FRayTracingAccelerationStructureSize FVulkanDynamicRHI::RHICalcRayTracingSceneSize(const FRayTracingSceneInitializer& Initializer)
 {
 	FVkRtTLASBuildData BuildData;
 	const VkDeviceAddress InstanceBufferAddress = 0; // No device address available when only querying TLAS size
@@ -1174,7 +1174,7 @@ FRayTracingAccelerationStructureSize FVulkanDynamicRHI::RHICalcRayTracingGeometr
 	return Result;
 }
 
-FRayTracingSceneRHIRef FVulkanDynamicRHI::RHICreateRayTracingScene(FRayTracingSceneInitializer2 Initializer)
+FRayTracingSceneRHIRef FVulkanDynamicRHI::RHICreateRayTracingScene(FRayTracingSceneInitializer Initializer)
 {
 	return new FVulkanRayTracingScene(MoveTemp(Initializer), GetDevice());
 }
diff --git a/Engine/Source/Runtime/VulkanRHI/Private/VulkanRayTracing.h b/Engine/Source/Runtime/VulkanRHI/Private/VulkanRayTracing.h
--- a/Engine/Source/Runtime/VulkanRHI/Private/VulkanRayTracing.h
+++ b/Engine/Source/Runtime/VulkanRHI/Private/VulkanRayTracing.h
@@ -215,10 +215,10 @@ class FVulkanRayTracingScene : public FRHIRayTracingScene, public VulkanRHI::FDe
 	friend FVulkanCommandListContext;
 
 public:
-	FVulkanRayTracingScene(FRayTracingSceneInitializer2 Initializer, FVulkanDevice* InDevice);
+	FVulkanRayTracingScene(FRayTracingSceneInitializer Initializer, FVulkanDevice* InDevice);
 	~FVulkanRayTracingScene();
 
-	const FRayTracingSceneInitializer2& GetInitializer() const override final { return Initializer; }
+	const FRayTracingSceneInitializer& GetInitializer() const override final { return Initializer; }
 
 	void BindBuffer(FRHIBuffer* InBuffer, uint32 InBufferOffset);
 
@@ -239,7 +239,7 @@ public:
 
 	using FRHIRayTracingAccelerationStructure::SizeInfo;
 
-	const FRayTracingSceneInitializer2 Initializer;
+	const FRayTracingSceneInitializer Initializer;
 
 	// Unique list of geometries referenced by all instances in this scene.
 	// Any referenced geometry is kept alive while the scene is alive.
diff --git a/Engine/Source/Runtime/VulkanRHI/Public/VulkanDynamicRHI.h b/Engine/Source/Runtime/VulkanRHI/Public/VulkanDynamicRHI.h
--- a/Engine/Source/Runtime/VulkanRHI/Public/VulkanDynamicRHI.h
+++ b/Engine/Source/Runtime/VulkanRHI/Public/VulkanDynamicRHI.h
@@ -232,10 +232,10 @@ public:
 	virtual FRHIResourceCollectionRef RHICreateResourceCollection(FRHICommandListBase& RHICmdList, TConstArrayView<FRHIResourceCollectionMember> InMembers) final;
 #endif
 
-	virtual FRayTracingAccelerationStructureSize RHICalcRayTracingSceneSize(const FRayTracingSceneInitializer2& Initializer) final override;
+	virtual FRayTracingAccelerationStructureSize RHICalcRayTracingSceneSize(const FRayTracingSceneInitializer& Initializer) final override;
 	virtual FRayTracingAccelerationStructureSize RHICalcRayTracingGeometrySize(const FRayTracingGeometryInitializer& Initializer) final override;
 	virtual FRayTracingGeometryRHIRef RHICreateRayTracingGeometry(FRHICommandListBase& RHICmdList, const FRayTracingGeometryInitializer& Initializer) final override;
-	virtual FRayTracingSceneRHIRef RHICreateRayTracingScene(FRayTracingSceneInitializer2 Initializer) final override;
+	virtual FRayTracingSceneRHIRef RHICreateRayTracingScene(FRayTracingSceneInitializer Initializer) final override;
 	virtual FRayTracingShaderRHIRef RHICreateRayTracingShader(TArrayView<const uint8> Code, const FSHAHash& Hash, EShaderFrequency ShaderFrequency) final override;
 	virtual FRayTracingPipelineStateRHIRef RHICreateRayTracingPipelineState(const FRayTracingPipelineStateInitializer& Initializer) final override;
 
