Add RHI_TEMP_SHOW_GPU_WORK which when enabled makes the Insights GPU track pop breadcrumb markers when the GPU is marked idle.
 - This allows us to see the GPU busy regions before we have API support for it in Insights
 - All busy regions first get marked with "GPUBusy" before the breadcrumbs are pushed. Emitting EndWork events is deferred to allow coalescing of overlapping busy regions on platforms that support top-of-pipe timestamps.
 - This can be quite noisy, hence leaving it on a macro to disable easily. It also causes the individual breadcrumbs to be broken up into multiple smaller ones.

Add optional args to GPU profiler event constructors to allow platform RHIs to construct the events in their entirety
 - Some RHIs still opt to fill in the GPU timestamp fields later

diff --git a/Engine/Source/Runtime/RHI/Private/GPUProfiler.cpp b/Engine/Source/Runtime/RHI/Private/GPUProfiler.cpp
--- a/Engine/Source/Runtime/RHI/Private/GPUProfiler.cpp
+++ b/Engine/Source/Runtime/RHI/Private/GPUProfiler.cpp
@@ -706,6 +706,10 @@ bool FGPUTiming::GAreGlobalsInitialized = false;
 // until we have a new API that is capable of displaying more info.
 #define RHI_TEMP_USE_GPU_TRACE (1 && GPUPROFILERTRACE_ENABLED)
 
+// When enabled, adds "GPUWork" markers to the GPU Insights trace to show where the GPU is busy or idle.
+// Causes breadcrumbs to be pushed / popped multiple times, breaking them up on the timeline.
+#define RHI_TEMP_SHOW_GPU_WORK (1 && RHI_TEMP_USE_GPU_TRACE)
+
 namespace UE::RHI::GPUProfiler
 {
 	static TArray<FEventSink*>& GetSinks()
@@ -871,6 +875,7 @@ namespace UE::RHI::GPUProfiler
 
 		#if WITH_RHI_BREADCRUMBS
 			TMap<FRHIBreadcrumbData_Stats, int32> ActiveStats;
+			FRHIBreadcrumbNode* Breadcrumb = nullptr;
 		#endif
 		};
 
@@ -881,11 +886,62 @@ namespace UE::RHI::GPUProfiler
 
 	#if RHI_TEMP_USE_GPU_TRACE
 		uint32 FrameNumber = 0;
+		uint64 MaxTraceTime = 0;
+
+		uint64 GPUToTrace(uint64 GPUTimestamp)
+		{
+			uint64 TraceTime = (uint64)(FPlatformTime::ToMilliseconds64(GPUTimestamp) * 1000.0);
+
+			//
+			// Some platforms support top-of-pipe timestamps, meaning BeginWork/BeginBreadcrumb events
+			// that occur logically after EndWork/EndBreadcrumb events in the command stream can have
+			// a timestamp that is earlier than the subsequent begin event due the GPU workload overlap.
+			//
+			// The old Insights API cannot support this, and simply doesn't display the events if their
+			// timestamps aren't strictly sequential. Work around this by emitting the Max() of the current
+			// timestamp, and the largest timestamp we've seen before.
+			//
+
+			MaxTraceTime = FMath::Max(TraceTime, MaxTraceTime);
+			return MaxTraceTime;
+		}
+
+	#if RHI_TEMP_SHOW_GPU_WORK
+		bool bNeedsEnd = false;
+		uint64 MaxEndTimeBOP = 0;
+
+		// Emitting FEndWork events to Insights are deferred until we know there isn't an overlapping FBeginWork event
+		// that would otherwise prevent the GPU going idle. This is done to coalesce markers to make them less noisy.
+		bool EmitEndWork(FQueueState& QueueState, TOptional<uint64> NextBeginWorkTimeTOP)
+		{
+			bool bEmitEnd = bNeedsEnd && (!NextBeginWorkTimeTOP.IsSet() || NextBeginWorkTimeTOP.GetValue() > MaxEndTimeBOP);
+			bool bNeedsBegin = !bNeedsEnd || bEmitEnd;
+
+			if (bEmitEnd)
+			{
+				uint64 TraceTime = GPUToTrace(MaxEndTimeBOP);
+
+			#if WITH_RHI_BREADCRUMBS
+				for (FRHIBreadcrumbNode* Current = QueueState.Breadcrumb; Current; Current = Current->GetParent())
+				{
+					FGpuProfilerTrace::EndEvent(TraceTime);
+				}
+			#endif
+
+				FGpuProfilerTrace::EndEvent(TraceTime); // GPUWork event
+
+				bNeedsEnd = false;
+			}
+
+			return bNeedsBegin;
+		}
+	#endif // RHI_TEMP_SHOW_GPU_WORK
+
 		bool ShouldEmitGPU(FQueue const& Queue)
 		{
 			return Queue.Type == FQueue::EType::Graphics && Queue.GPU == 0 && Queue.Index == 0;
 		}
-	#endif
+	#endif // RHI_TEMP_USE_GPU_TRACE
 
 		void InitializeQueues(TConstArrayView<FQueue> Queues) override
 		{
@@ -923,6 +979,39 @@ namespace UE::RHI::GPUProfiler
 						FEvent::FBeginWork const& BeginWork = Event->Value.Get<FEvent::FBeginWork>();
 						QueueState.Timestamps.Queue.AddTimestamp(BeginWork.GPUTimestampTOP, true);
 
+					#if RHI_TEMP_SHOW_GPU_WORK
+						if (ShouldEmitGPU(Queue))
+						{
+							if (EmitEndWork(QueueState, BeginWork.GPUTimestampTOP))
+							{
+								uint64 TraceTime = GPUToTrace(BeginWork.GPUTimestampTOP);
+								static FName GPUWorkName = ([]()
+								{
+									FName Name("GPUWork");
+									FGpuProfilerTrace::SpecifyEventByName(Name);
+									return Name;
+								})();
+
+								FGpuProfilerTrace::BeginEventByName(GPUWorkName, FrameNumber, TraceTime);
+
+							#if WITH_RHI_BREADCRUMBS
+								FRHIBreadcrumb::FBuffer Buffer;
+								auto Recurse = [&](auto& Recurse, FRHIBreadcrumbNode* Current) -> void
+								{
+									if (!Current)
+										return;
+
+									Recurse(Recurse, Current->GetParent());
+
+									FName Name(Current->Name.GetTCHAR(Buffer));
+									FGpuProfilerTrace::BeginEventByName(Name, FrameNumber, TraceTime);
+								};
+								Recurse(Recurse, QueueState.Breadcrumb);
+							#endif
+							}
+						}
+					#endif
+
 					#if WITH_RHI_BREADCRUMBS
 						// Apply the timestamp to all active stats
 						for (auto const& [Stat, RefCount] : QueueState.ActiveStats)
@@ -938,14 +1027,22 @@ namespace UE::RHI::GPUProfiler
 						check(QueueState.bBusy);
 						QueueState.bBusy = false;
 
-						uint64 Value = Event->Value.Get<FEvent::FEndWork>().GPUTimestampBOP;
-						QueueState.Timestamps.Queue.AddTimestamp(Value, false);
+						FEvent::FEndWork const& EndWork = Event->Value.Get<FEvent::FEndWork>();
+						QueueState.Timestamps.Queue.AddTimestamp(EndWork.GPUTimestampBOP, false);
 
 					#if WITH_RHI_BREADCRUMBS
 						// Apply the timestamp to all active stats
 						for (auto const& [Stat, RefCount] : QueueState.ActiveStats)
 						{
-							QueueState.Timestamps.Stats.FindChecked(Stat).AddTimestamp(Value, false);
+							QueueState.Timestamps.Stats.FindChecked(Stat).AddTimestamp(EndWork.GPUTimestampBOP, false);
+						}
+					#endif
+
+					#if RHI_TEMP_SHOW_GPU_WORK
+						if (ShouldEmitGPU(Queue))
+						{
+							MaxEndTimeBOP = FMath::Max(MaxEndTimeBOP, EndWork.GPUTimestampBOP);
+							bNeedsEnd = true;
 						}
 					#endif
 					}
@@ -965,7 +1062,7 @@ namespace UE::RHI::GPUProfiler
 							FName Name(Str);
 
 							FGpuProfilerTrace::SpecifyEventByName(Name);
-							FGpuProfilerTrace::BeginEventByName(Name, FrameNumber, (uint64)(FPlatformTime::ToMilliseconds64(BeginBreadcrumb.GPUTimestampTOP) * 1000.0));
+							FGpuProfilerTrace::BeginEventByName(Name, FrameNumber, GPUToTrace(BeginBreadcrumb.GPUTimestampTOP));
 						}
 					#endif
 
@@ -979,6 +1076,8 @@ namespace UE::RHI::GPUProfiler
 								QueueState.Timestamps.Stats.FindOrAdd(Stat).AddTimestamp(BeginBreadcrumb.GPUTimestampTOP, true);
 							}
 						}
+
+						QueueState.Breadcrumb = BeginBreadcrumb.Breadcrumb;
 					}
 					break;
 
@@ -990,7 +1089,7 @@ namespace UE::RHI::GPUProfiler
 					#if RHI_TEMP_USE_GPU_TRACE
 						if (ShouldEmitGPU(Queue))
 						{
-							FGpuProfilerTrace::EndEvent((uint64)(FPlatformTime::ToMilliseconds64(EndBreadcrumb.GPUTimestampBOP) * 1000.0));
+							FGpuProfilerTrace::EndEvent(GPUToTrace(EndBreadcrumb.GPUTimestampBOP));
 						}
 					#endif
 
@@ -1004,6 +1103,8 @@ namespace UE::RHI::GPUProfiler
 								QueueState.ActiveStats.FindAndRemoveChecked(Stat);
 							}							
 						}
+
+						QueueState.Breadcrumb = EndBreadcrumb.Breadcrumb->GetParent();
 					}
 					break;
 			#endif // WITH_RHI_BREADCRUMBS
@@ -1050,6 +1151,10 @@ namespace UE::RHI::GPUProfiler
 						// End the current Insights GPU frame + start the next one
 						if (ShouldEmitGPU(Queue))
 						{
+						#if RHI_TEMP_SHOW_GPU_WORK
+							EmitEndWork(QueueState, {});
+						#endif
+
 							check(FrameNumber == FrameBoundary.FrameNumber);
 							FGpuProfilerTrace::EndFrame(0);
 							FrameNumber++;
diff --git a/Engine/Source/Runtime/RHI/Public/GPUProfiler.h b/Engine/Source/Runtime/RHI/Public/GPUProfiler.h
--- a/Engine/Source/Runtime/RHI/Public/GPUProfiler.h
+++ b/Engine/Source/Runtime/RHI/Public/GPUProfiler.h
@@ -44,19 +44,19 @@ namespace UE::RHI::GPUProfiler
 
 		FQueue() = default;
 
-		FQueue(EType Type, uint8 GPU, uint8 Index)
+		constexpr FQueue(EType Type, uint8 GPU, uint8 Index)
 			: Type   (Type)
 			, GPU    (GPU)
 			, Index  (Index)
 			, Padding(0)
 		{}
 
-		bool operator == (FQueue const& RHS) const
+		constexpr bool operator == (FQueue const& RHS) const
 		{
 			return Value == RHS.Value;
 		}
 
-		bool operator != (FQueue const& RHS) const
+		constexpr bool operator != (FQueue const& RHS) const
 		{
 			return !(*this == RHS);
 		}
@@ -115,10 +115,11 @@ namespace UE::RHI::GPUProfiler
 		struct FBeginBreadcrumb
 		{
 			FRHIBreadcrumbNode* const Breadcrumb;
-			uint64 GPUTimestampTOP = 0;
+			uint64 GPUTimestampTOP;
 
-			FBeginBreadcrumb(FRHIBreadcrumbNode* Breadcrumb)
+			FBeginBreadcrumb(FRHIBreadcrumbNode* Breadcrumb, uint64 GPUTimestampTOP = 0)
 				: Breadcrumb(Breadcrumb)
+				, GPUTimestampTOP(GPUTimestampTOP)
 			{}
 		};
 
@@ -127,8 +128,9 @@ namespace UE::RHI::GPUProfiler
 			FRHIBreadcrumbNode* const Breadcrumb;
 			uint64 GPUTimestampBOP = 0;
 
-			FEndBreadcrumb(FRHIBreadcrumbNode* Breadcrumb)
+			FEndBreadcrumb(FRHIBreadcrumbNode* Breadcrumb, uint64 GPUTimestampBOP = 0)
 				: Breadcrumb(Breadcrumb)
+				, GPUTimestampBOP(GPUTimestampBOP)
 			{}
 		};
 	#endif
@@ -140,17 +142,22 @@ namespace UE::RHI::GPUProfiler
 			uint64 CPUTimestamp;
 
 			// TOP timestamp of when the work actually started on the GPU.
-			uint64 GPUTimestampTOP = 0;
+			uint64 GPUTimestampTOP;
 
-			FBeginWork(uint64 CPUTimestamp)
+			FBeginWork(uint64 CPUTimestamp, uint64 GPUTimestampTOP = 0)
 				: CPUTimestamp(CPUTimestamp)
+				, GPUTimestampTOP(GPUTimestampTOP)
 			{}
 		};
 
 		// Inserted when the GPU completes work on a queue and goes idle.
 		struct FEndWork
 		{
-			uint64 GPUTimestampBOP = 0;
+			uint64 GPUTimestampBOP;
+
+			FEndWork(uint64 GPUTimestampBOP = 0)
+				: GPUTimestampBOP(GPUTimestampBOP)
+			{}
 		};
 
 		struct FStats
