Add a DDPI entry for Shader Execution Reordering

We need to be able to identify platforms which support this extension when compiling shaders without referring to what the current machine happens to support (such as at cook time). Add an entry to DDPI and use it in the path tracer to decide if the SER extension should be compiled.

#rnx

diff --git a/Engine/Config/Windows/DataDrivenPlatformInfo.ini b/Engine/Config/Windows/DataDrivenPlatformInfo.ini
--- a/Engine/Config/Windows/DataDrivenPlatformInfo.ini
+++ b/Engine/Config/Windows/DataDrivenPlatformInfo.ini
@@ -91,6 +91,7 @@ bSupportsRayTracingCallableShaders=true
 bSupportsRayTracingProceduralPrimitive=true
 bSupportsRayTracingIndirectInstanceData=true
 bSupportsPathTracing=true
+bSupportsShaderExecutionReordering=true
 bSupportsHighEndRayTracingEffects=true
 bSupportsGPUScene=true
 bSupportsNanite=true
diff --git a/Engine/Source/Runtime/RHI/Private/DataDrivenShaderPlatformInfo.cpp b/Engine/Source/Runtime/RHI/Private/DataDrivenShaderPlatformInfo.cpp
--- a/Engine/Source/Runtime/RHI/Private/DataDrivenShaderPlatformInfo.cpp
+++ b/Engine/Source/Runtime/RHI/Private/DataDrivenShaderPlatformInfo.cpp
@@ -259,6 +259,7 @@ void FGenericDataDrivenShaderPlatformInfo::ParseDataDrivenShaderInfo(const FConf
 	GET_SECTION_BOOL_HELPER(bSupportsRayTracingTraversalStatistics);
 	GET_SECTION_BOOL_HELPER(bSupportsRayTracingIndirectInstanceData);
 	GET_SECTION_BOOL_HELPER(bSupportsPathTracing);
+	GET_SECTION_BOOL_HELPER(bSupportsShaderExecutionReordering);
 	GET_SECTION_BOOL_HELPER(bSupportsHighEndRayTracingEffects);
 	GET_SECTION_BOOL_HELPER(bSupportsUnrestrictedHalfFloatBuffers)
 	GET_SECTION_BOOL_HELPER(bSupportsByteBufferComputeShaders);
diff --git a/Engine/Source/Runtime/RHI/Public/DataDrivenShaderPlatformInfo.h b/Engine/Source/Runtime/RHI/Public/DataDrivenShaderPlatformInfo.h
--- a/Engine/Source/Runtime/RHI/Public/DataDrivenShaderPlatformInfo.h
+++ b/Engine/Source/Runtime/RHI/Public/DataDrivenShaderPlatformInfo.h
@@ -51,6 +51,7 @@ class FGenericDataDrivenShaderPlatformInfo
 	uint32 bSupportsRayTracingIndirectInstanceData : 1; // Whether instance transforms can be copied from the GPU to the TLAS instances buffer
 	uint32 bSupportsHighEndRayTracingEffects : 1; // Whether fully-featured RT effects can be used on the platform (with translucent shadow, etc.)
 	uint32 bSupportsPathTracing : 1; // Whether real-time path tracer is supported on this platform (avoids compiling unnecessary shaders)
+	uint32 bSupportsShaderExecutionReordering : 1; // Does the platform support Shader Execution Reordering extensions?
 	uint32 bSupportsGPUScene : 1;
 	uint32 bSupportsUnrestrictedHalfFloatBuffers : 1;
 	uint32 bSupportsByteBufferComputeShaders : 1;
@@ -390,6 +391,12 @@ public:
 		return Infos[Platform].bSupportsRayTracing && Infos[Platform].bSupportsPathTracing;
 	}
 
+	static FORCEINLINE_DEBUGGABLE const bool GetSupportsShaderExecutionReordering(const FStaticShaderPlatform Platform)
+	{
+		check(IsValid(Platform));
+		return Infos[Platform].bSupportsShaderExecutionReordering;
+	}
+
 	static FORCEINLINE_DEBUGGABLE const bool GetSupportsHighEndRayTracingEffects(const FStaticShaderPlatform Platform)
 	{
 		check(IsValid(Platform));
diff --git a/Engine/Source/Runtime/Renderer/Private/PathTracing.cpp b/Engine/Source/Runtime/Renderer/Private/PathTracing.cpp
--- a/Engine/Source/Runtime/Renderer/Private/PathTracing.cpp
+++ b/Engine/Source/Runtime/Renderer/Private/PathTracing.cpp
@@ -1140,9 +1140,10 @@ class FPathTracingRG : public FGlobalShader
 				}
 			}
 		}
-		if (PermutationVector.Get<FUseSER>() && Parameters.Platform != SP_PCD3D_SM6)
+
+		// Only compile SER extensions on platforms that could support it
+		if (PermutationVector.Get<FUseSER>() && !FDataDrivenShaderPlatformInfo::GetSupportsShaderExecutionReordering(Parameters.Platform))
 		{
-			// only D3D12 supports SER for now
 			return false;
 		}
 
