Use "GPU2" Insights track:
 - for async compute when running in single-GPU mode.
 - for the 2nd graphics queue when running in multi-GPU mode.

Make the "GPU work" mode a runtime choice based on "-tracegpuwork" param in the command line.
 - Running without the arg will only emit breadcrumb markers.
 - Running with the arg will push/pop breadcrumbs based on regions of time the GPU is actually busy doing work.

Minor changes to GpuProfilerTrace to allow for the recording of both GPU and GPU2 tracks simultenously.
 - It was previously a singleton, so only one frame could be built at once.

diff --git a/Engine/Source/Runtime/RHI/Private/GPUProfiler.cpp b/Engine/Source/Runtime/RHI/Private/GPUProfiler.cpp
--- a/Engine/Source/Runtime/RHI/Private/GPUProfiler.cpp
+++ b/Engine/Source/Runtime/RHI/Private/GPUProfiler.cpp
@@ -7,6 +7,7 @@
 #include "GPUProfiler.h"
 #include "Async/TaskGraphInterfaces.h"
 #include "Misc/WildcardString.h"
+#include "Misc/CommandLine.h"
 #include "RHI.h"
 #include "GpuProfilerTrace.h"
 
@@ -706,9 +707,9 @@ bool FGPUTiming::GAreGlobalsInitialized = false;
 // until we have a new API that is capable of displaying more info.
 #define RHI_TEMP_USE_GPU_TRACE (1 && GPUPROFILERTRACE_ENABLED)
 
-// When enabled, adds "GPUWork" markers to the GPU Insights trace to show where the GPU is busy or idle.
-// Causes breadcrumbs to be pushed / popped multiple times, breaking them up on the timeline.
-#define RHI_TEMP_SHOW_GPU_WORK (1 && RHI_TEMP_USE_GPU_TRACE)
+// When enabled, and running with a single GPU, repurposes the "GPU2" track
+// in Insights to show the single GPU's async compute queue.
+#define RHI_TEMP_USE_TRACK2_FOR_COMPUTE (1 && RHI_TEMP_USE_GPU_TRACE)
 
 namespace UE::RHI::GPUProfiler
 {
@@ -870,6 +871,8 @@ namespace UE::RHI::GPUProfiler
 
 		struct FQueueState
 		{
+			FQueue::EType Type;
+
 			bool bBusy = false;
 			FQueueTimestamps Timestamps;
 
@@ -877,6 +880,10 @@ namespace UE::RHI::GPUProfiler
 			TMap<FRHIBreadcrumbData_Stats, int32> ActiveStats;
 			FRHIBreadcrumbNode* Breadcrumb = nullptr;
 		#endif
+
+			FQueueState(FQueue const& Queue)
+				: Type(Queue.Type)
+			{}
 		};
 
 		using FFrameState = TMap<FQueue, FQueueTimestamps>;
@@ -885,62 +892,229 @@ namespace UE::RHI::GPUProfiler
 		TMap<uint32, FFrameState> Frames;
 
 	#if RHI_TEMP_USE_GPU_TRACE
-		uint32 FrameNumber = 0;
-		uint64 MaxTraceTime = 0;
-
-		uint64 GPUToTrace(uint64 GPUTimestamp)
+		struct FInsightsTrack
 		{
-			uint64 TraceTime = (uint64)(FPlatformTime::ToMilliseconds64(GPUTimestamp) * 1000.0);
-
-			//
-			// Some platforms support top-of-pipe timestamps, meaning BeginWork/BeginBreadcrumb events
-			// that occur logically after EndWork/EndBreadcrumb events in the command stream can have
-			// a timestamp that is earlier than the subsequent begin event due the GPU workload overlap.
-			//
-			// The old Insights API cannot support this, and simply doesn't display the events if their
-			// timestamps aren't strictly sequential. Work around this by emitting the Max() of the current
-			// timestamp, and the largest timestamp we've seen before.
-			//
-
-			MaxTraceTime = FMath::Max(TraceTime, MaxTraceTime);
-			return MaxTraceTime;
-		}
+			uint32 const Index;
+			uint64 MaxTraceTime = 0;
+			uint32 FrameNumber = 0;
 
-	#if RHI_TEMP_SHOW_GPU_WORK
-		bool bNeedsEnd = false;
-		uint64 MaxEndTimeBOP = 0;
+		#if DO_CHECK
+			int32 EventCounter = 0;
+		#endif
 
-		// Emitting FEndWork events to Insights are deferred until we know there isn't an overlapping FBeginWork event
-		// that would otherwise prevent the GPU going idle. This is done to coalesce markers to make them less noisy.
-		bool EmitEndWork(FQueueState& QueueState, TOptional<uint64> NextBeginWorkTimeTOP)
-		{
-			bool bEmitEnd = bNeedsEnd && (!NextBeginWorkTimeTOP.IsSet() || NextBeginWorkTimeTOP.GetValue() > MaxEndTimeBOP);
-			bool bNeedsBegin = !bNeedsEnd || bEmitEnd;
+			FInsightsTrack(uint32 Index)
+				: Index(Index)
+			{}
 
-			if (bEmitEnd)
+			uint64 GPUToTrace(uint64 GPUTimestamp)
 			{
-				uint64 TraceTime = GPUToTrace(MaxEndTimeBOP);
+				uint64 TraceTime = (uint64)(FPlatformTime::ToMilliseconds64(GPUTimestamp) * 1000.0);
+
+				//
+				// Some platforms support top-of-pipe timestamps, meaning BeginWork/BeginBreadcrumb events
+				// that occur logically after EndWork/EndBreadcrumb events in the command stream can have
+				// a timestamp that is earlier than the subsequent begin event due the GPU workload overlap.
+				//
+				// The old Insights API cannot support this, and simply doesn't display the events if their
+				// timestamps aren't strictly sequential. Work around this by emitting the Max() of the current
+				// timestamp, and the largest timestamp we've seen before.
+				//
+
+				MaxTraceTime = FMath::Max(TraceTime, MaxTraceTime);
+				return MaxTraceTime;
+			}
 
-			#if WITH_RHI_BREADCRUMBS
-				for (FRHIBreadcrumbNode* Current = QueueState.Breadcrumb; Current; Current = Current->GetParent())
+			bool bShowWork = false;
+			bool bEmittedGPUWorkName = false;
+
+			bool bNeedsEnd = false;
+			uint64 MaxEndTimeBOP = 0;
+			TOptional<uint64> LastBeginTimestampTOP;
+
+			// Emitting FEndWork events to Insights are deferred until we know there isn't an overlapping FBeginWork event
+			// that would otherwise prevent the GPU going idle. This is done to coalesce markers to make them less noisy.
+			bool EmitEndWork(FQueueState const& QueueState)
+			{
+				bool bEmitEnd = bNeedsEnd && (!LastBeginTimestampTOP.IsSet() || LastBeginTimestampTOP.GetValue() > MaxEndTimeBOP);
+				bool bNeedsBegin = !bNeedsEnd || bEmitEnd;
+
+				if (bEmitEnd)
 				{
-					FGpuProfilerTrace::EndEvent(TraceTime);
+					uint64 TraceTime = GPUToTrace(MaxEndTimeBOP);
+
+				#if WITH_RHI_BREADCRUMBS
+					for (FRHIBreadcrumbNode* Current = QueueState.Breadcrumb; Current; Current = Current->GetParent())
+					{
+						check(EventCounter-- > 0);
+						FGpuProfilerTrace::EndEvent(TraceTime, Index);
+					}
+				#endif
+
+					if (bShowWork)
+					{
+						check(EventCounter-- > 0);
+						FGpuProfilerTrace::EndEvent(TraceTime, Index); // GPUWork event
+					}
+
+					bNeedsEnd = false;
 				}
-			#endif
 
-				FGpuProfilerTrace::EndEvent(TraceTime); // GPUWork event
+				return bNeedsBegin;
+			}
+
+			void EmitBeginWork(FQueueState const& QueueState)
+			{
+				if (LastBeginTimestampTOP.IsSet())
+				{
+					uint64 TraceTime = GPUToTrace(LastBeginTimestampTOP.GetValue());
+							
+					if (bShowWork)
+					{
+						static FName GraphicsWorkName("Graphics Work");
+						static FName ComputeWorkName("Compute Work");
+						FName WorkName = QueueState.Type == FQueue::EType::Graphics
+							? GraphicsWorkName
+							: ComputeWorkName;
+
+						if (!bEmittedGPUWorkName)
+						{
+							FGpuProfilerTrace::SpecifyEventByName(WorkName, Index);
+						}
+
+						FGpuProfilerTrace::BeginEventByName(WorkName, FrameNumber, TraceTime, Index);
+						check(++EventCounter);
+					}
+
+				#if WITH_RHI_BREADCRUMBS
+					{
+						FRHIBreadcrumb::FBuffer Buffer;
+						auto Recurse = [&](auto& Recurse, FRHIBreadcrumbNode* Current) -> void
+						{
+							if (!Current)
+								return;
 
-				bNeedsEnd = false;
+							Recurse(Recurse, Current->GetParent());
+
+							FName Name(Current->Name.GetTCHAR(Buffer));
+							FGpuProfilerTrace::BeginEventByName(Name, FrameNumber, TraceTime, Index);
+							check(++EventCounter);
+						};
+						Recurse(Recurse, QueueState.Breadcrumb);
+					}
+				#endif
+
+					LastBeginTimestampTOP.Reset();
+				}
 			}
 
-			return bNeedsBegin;
-		}
-	#endif // RHI_TEMP_SHOW_GPU_WORK
+			void BeginWork(FQueueState const& QueueState, FEvent::FBeginWork const& BeginWork)
+			{
+				LastBeginTimestampTOP = BeginWork.GPUTimestampTOP;
+
+				if (bShowWork)
+				{
+					if (EmitEndWork(QueueState))
+					{
+						EmitBeginWork(QueueState);
+					}
+				}
+			}
 
-		bool ShouldEmitGPU(FQueue const& Queue)
+			void EndWork(FQueueState const& QueueState, FEvent::FEndWork const& EndWork)
+			{
+				MaxEndTimeBOP = FMath::Max(MaxEndTimeBOP, EndWork.GPUTimestampBOP);
+				bNeedsEnd = true;
+			}
+
+			void BeginBreadcrumb(FQueueState const& QueueState, FEvent::FBeginBreadcrumb const& BeginBreadcrumb)
+			{
+				FRHIBreadcrumb::FBuffer Buffer;
+				TCHAR const* Str = BeginBreadcrumb.Breadcrumb->Name.GetTCHAR(Buffer);
+				FName Name(Str);
+
+				FGpuProfilerTrace::SpecifyEventByName(Name, Index);
+				FGpuProfilerTrace::BeginEventByName(Name, FrameNumber, GPUToTrace(BeginBreadcrumb.GPUTimestampTOP), Index);
+				check(++EventCounter);
+			}
+
+			void EndBreadcrumb(FQueueState const& QueueState, FEvent::FEndBreadcrumb const& EndBreadcrumb)
+			{
+				check(EventCounter-- > 0);
+				FGpuProfilerTrace::EndEvent(GPUToTrace(EndBreadcrumb.GPUTimestampBOP), Index);
+			}
+
+			void FrameBoundary(FQueueState const& QueueState, FEvent::FFrameBoundary const& FrameBoundary)
+			{
+				// End the current Insights GPU frame + start the next one
+
+				// All breadcrumbs must be ended before the frame boundary can be emitted.
+				LastBeginTimestampTOP.Reset();
+				bool bNeedsBegin = EmitEndWork(QueueState);
+
+				check(FrameNumber == FrameBoundary.FrameNumber);
+				check(EventCounter == 0);
+				FGpuProfilerTrace::EndFrame(Index);
+
+				FrameNumber++;
+
+				// Use 1,1 calibration to disable any adjustments Insights makes.
+				// The timestamps we use in the GPU event stream are already in the CPU clock domain.
+				FGPUTimingCalibrationTimestamp Calibration{ 1, 1 };
+				FGpuProfilerTrace::BeginFrame(Calibration, Index);
+
+				if (!bShowWork)
+				{
+					if (bNeedsBegin)
+					{
+						LastBeginTimestampTOP = MaxEndTimeBOP;
+						EmitBeginWork(QueueState);
+					}
+				}
+			}
+
+			void Initialize()
+			{
+				// When enabled, adds "GPUWork" markers to the GPU Insights trace to show where the GPU is busy or idle.
+				// Causes breadcrumbs to be pushed / popped multiple times, breaking them up on the timeline.
+				bShowWork = FParse::Param(FCommandLine::Get(), TEXT("tracegpuwork"));
+
+				// Use 1,1 calibration to disable any adjustments Insights makes.
+				// The timestamps we use in the GPU event stream are already in the CPU clock domain.
+				FGPUTimingCalibrationTimestamp Calibration{ 1, 1 };
+				check(EventCounter == 0);
+				FGpuProfilerTrace::BeginFrame(Calibration, Index);
+			}
+
+		} InsightsTracks[2] { {0}, {1} };
+
+		FInsightsTrack* GetInsightsTrack(FQueue const& Queue)
 		{
-			return Queue.Type == FQueue::EType::Graphics && Queue.GPU == 0 && Queue.Index == 0;
+			if (GNumExplicitGPUsForRendering > 1)
+			{
+				// MGPU Mode - GPU0 Graphics + GPU1 Graphics
+				if (Queue.Type == FQueue::EType::Graphics && Queue.Index == 0 && Queue.GPU < 2)
+				{
+					return &InsightsTracks[Queue.GPU];
+				}
+			}
+			else
+			{
+				// GPU0 Graphics + GPU0 Compute Mode
+				if (Queue.GPU == 0 && Queue.Index == 0)
+				{
+					switch (Queue.Type)
+					{
+					case FQueue::EType::Graphics: return &InsightsTracks[0];
+				#if RHI_TEMP_USE_TRACK2_FOR_COMPUTE
+					case FQueue::EType::Compute : return &InsightsTracks[1];
+				#endif
+					}
+				}
+			}
+
+			return nullptr;
 		}
+
 	#endif // RHI_TEMP_USE_GPU_TRACE
 
 		void InitializeQueues(TConstArrayView<FQueue> Queues) override
@@ -948,16 +1122,14 @@ namespace UE::RHI::GPUProfiler
 			for (FQueue const& Queue : Queues)
 			{
 				check(QueueStates.Find(Queue) == nullptr);
-				QueueStates.Add(Queue);
+				QueueStates.Add(Queue, Queue);
 
 			#if RHI_TEMP_USE_GPU_TRACE
 				// Start the first Insights GPU frame
-				if (ShouldEmitGPU(Queue))
+				FInsightsTrack* Track = GetInsightsTrack(Queue);
+				if (Track)
 				{
-					// Use 1,1 calibration to disable any adjustments Insights makes.
-					// The timestamps we use in the GPU event stream are already in the CPU clock domain.
-					FGPUTimingCalibrationTimestamp Calibration{ 1, 1 };
-					FGpuProfilerTrace::BeginFrame(Calibration);
+					Track->Initialize();
 				}
 			#endif
 			}
@@ -965,6 +1137,9 @@ namespace UE::RHI::GPUProfiler
 
 		void ProcessEvents(FQueue Queue, FEventStream const& EventStream) override
 		{
+		#if RHI_TEMP_USE_GPU_TRACE
+			FInsightsTrack* Track = GetInsightsTrack(Queue);
+		#endif
 			FQueueState& QueueState = QueueStates.FindChecked(Queue);
 
 			for (FEvent const* Event : EventStream)
@@ -979,36 +1154,10 @@ namespace UE::RHI::GPUProfiler
 						FEvent::FBeginWork const& BeginWork = Event->Value.Get<FEvent::FBeginWork>();
 						QueueState.Timestamps.Queue.AddTimestamp(BeginWork.GPUTimestampTOP, true);
 
-					#if RHI_TEMP_SHOW_GPU_WORK
-						if (ShouldEmitGPU(Queue))
+					#if RHI_TEMP_USE_GPU_TRACE
+						if (Track)
 						{
-							if (EmitEndWork(QueueState, BeginWork.GPUTimestampTOP))
-							{
-								uint64 TraceTime = GPUToTrace(BeginWork.GPUTimestampTOP);
-								static FName GPUWorkName = ([]()
-								{
-									FName Name("GPUWork");
-									FGpuProfilerTrace::SpecifyEventByName(Name);
-									return Name;
-								})();
-
-								FGpuProfilerTrace::BeginEventByName(GPUWorkName, FrameNumber, TraceTime);
-
-							#if WITH_RHI_BREADCRUMBS
-								FRHIBreadcrumb::FBuffer Buffer;
-								auto Recurse = [&](auto& Recurse, FRHIBreadcrumbNode* Current) -> void
-								{
-									if (!Current)
-										return;
-
-									Recurse(Recurse, Current->GetParent());
-
-									FName Name(Current->Name.GetTCHAR(Buffer));
-									FGpuProfilerTrace::BeginEventByName(Name, FrameNumber, TraceTime);
-								};
-								Recurse(Recurse, QueueState.Breadcrumb);
-							#endif
-							}
+							Track->BeginWork(QueueState, BeginWork);
 						}
 					#endif
 
@@ -1038,11 +1187,10 @@ namespace UE::RHI::GPUProfiler
 						}
 					#endif
 
-					#if RHI_TEMP_SHOW_GPU_WORK
-						if (ShouldEmitGPU(Queue))
+					#if RHI_TEMP_USE_GPU_TRACE
+						if (Track)
 						{
-							MaxEndTimeBOP = FMath::Max(MaxEndTimeBOP, EndWork.GPUTimestampBOP);
-							bNeedsEnd = true;
+							Track->EndWork(QueueState, EndWork);
 						}
 					#endif
 					}
@@ -1051,18 +1199,15 @@ namespace UE::RHI::GPUProfiler
 			#if WITH_RHI_BREADCRUMBS
 				case FEvent::EType::BeginBreadcrumb:
 					{
+						check(QueueState.bBusy);
+
 						FEvent::FBeginBreadcrumb const& BeginBreadcrumb = Event->Value.Get<FEvent::FBeginBreadcrumb>();
 						FRHIBreadcrumbData_Stats const& Stat = BeginBreadcrumb.Breadcrumb->Name.Data;
 
 					#if RHI_TEMP_USE_GPU_TRACE
-						if (ShouldEmitGPU(Queue))
+						if (Track)
 						{
-							FRHIBreadcrumb::FBuffer Buffer;
-							TCHAR const* Str = BeginBreadcrumb.Breadcrumb->Name.GetTCHAR(Buffer);
-							FName Name(Str);
-
-							FGpuProfilerTrace::SpecifyEventByName(Name);
-							FGpuProfilerTrace::BeginEventByName(Name, FrameNumber, GPUToTrace(BeginBreadcrumb.GPUTimestampTOP));
+							Track->BeginBreadcrumb(QueueState, BeginBreadcrumb);
 						}
 					#endif
 
@@ -1083,13 +1228,15 @@ namespace UE::RHI::GPUProfiler
 
 				case FEvent::EType::EndBreadcrumb:
 					{
+						check(QueueState.bBusy);
+
 						FEvent::FEndBreadcrumb const& EndBreadcrumb = Event->Value.Get<FEvent::FEndBreadcrumb>();
 						FRHIBreadcrumbData_Stats const& Stat = EndBreadcrumb.Breadcrumb->Name.Data;
 
 					#if RHI_TEMP_USE_GPU_TRACE
-						if (ShouldEmitGPU(Queue))
+						if (Track)
 						{
-							FGpuProfilerTrace::EndEvent(GPUToTrace(EndBreadcrumb.GPUTimestampBOP));
+							Track->EndBreadcrumb(QueueState, EndBreadcrumb);
 						}
 					#endif
 
@@ -1148,21 +1295,9 @@ namespace UE::RHI::GPUProfiler
 						}
 
 					#if RHI_TEMP_USE_GPU_TRACE
-						// End the current Insights GPU frame + start the next one
-						if (ShouldEmitGPU(Queue))
+						if (Track)
 						{
-						#if RHI_TEMP_SHOW_GPU_WORK
-							EmitEndWork(QueueState, {});
-						#endif
-
-							check(FrameNumber == FrameBoundary.FrameNumber);
-							FGpuProfilerTrace::EndFrame(0);
-							FrameNumber++;
-
-							// Use 1,1 calibration to disable any adjustments Insights makes.
-							// The timestamps we use in the GPU event stream are already in the CPU clock domain.
-							FGPUTimingCalibrationTimestamp Calibration{ 1, 1 };
-							FGpuProfilerTrace::BeginFrame(Calibration);
+							Track->FrameBoundary(QueueState, FrameBoundary);
 						}
 					#endif
 					}
diff --git a/Engine/Source/Runtime/RHI/Private/GpuProfilerTrace.cpp b/Engine/Source/Runtime/RHI/Private/GpuProfilerTrace.cpp
--- a/Engine/Source/Runtime/RHI/Private/GpuProfilerTrace.cpp
+++ b/Engine/Source/Runtime/RHI/Private/GpuProfilerTrace.cpp
@@ -22,7 +22,7 @@ static TAutoConsoleVariable<int32> CVarGpuProfilerMaxEventBufferSizeKB(
 	ECVF_Default);
 
 
-struct
+struct FGpuTraceFrame
 {
 	int64							CalibrationBias;
 	FGPUTimingCalibrationTimestamp	Calibration;
@@ -33,7 +33,13 @@ struct
 	bool							bActive;
 	uint8*							EventBuffer = nullptr;
 	uint32							MaxEventBufferSize = 0;
-} GCurrentFrame;
+};
+
+#if RHI_NEW_GPU_PROFILER
+FGpuTraceFrame GCurrentFrames[2];
+#else
+FGpuTraceFrame GCurrentFrame;
+#endif
 
 static TSet<uint32> GEventNames;
 
@@ -63,10 +69,20 @@ UE_TRACE_EVENT_END()
 
 } // namespace GpuProfilerTrace
 
-void FGpuProfilerTrace::BeginFrame(FGPUTimingCalibrationTimestamp& Calibration)
+#if RHI_NEW_GPU_PROFILER
+#define GPU_TRACE_ARG , uint32 GPUIndex
+#else
+#define GPU_TRACE_ARG
+#endif
+
+void FGpuProfilerTrace::BeginFrame(FGPUTimingCalibrationTimestamp& Calibration GPU_TRACE_ARG)
 {
 	using namespace GpuProfilerTrace;
 
+#if RHI_NEW_GPU_PROFILER
+	FGpuTraceFrame& GCurrentFrame = GCurrentFrames[GPUIndex];
+#endif
+
 	if (!bool(GpuChannel))
 	{
 		return;
@@ -87,10 +103,14 @@ void FGpuProfilerTrace::BeginFrame(FGPUTimingCalibrationTimestamp& Calibration)
 	}
 }
 
-void FGpuProfilerTrace::SpecifyEventByName(const FName& Name)
+void FGpuProfilerTrace::SpecifyEventByName(const FName& Name GPU_TRACE_ARG)
 {
 	using namespace GpuProfilerTrace;
 
+#if RHI_NEW_GPU_PROFILER
+	FGpuTraceFrame& GCurrentFrame = GCurrentFrames[GPUIndex];
+#endif
+
 	if (!GCurrentFrame.bActive)
 	{
 		return;
@@ -114,10 +134,14 @@ void FGpuProfilerTrace::SpecifyEventByName(const FName& Name)
 	}
 }
 
-void FGpuProfilerTrace::BeginEventByName(const FName& Name, uint32 FrameNumber, uint64 TimestampMicroseconds)
+void FGpuProfilerTrace::BeginEventByName(const FName& Name, uint32 FrameNumber, uint64 TimestampMicroseconds GPU_TRACE_ARG)
 {
 	using namespace GpuProfilerTrace;
 
+#if RHI_NEW_GPU_PROFILER
+	FGpuTraceFrame& GCurrentFrame = GCurrentFrames[GPUIndex];
+#endif
+
 	if (!GCurrentFrame.bActive)
 	{
 		return;
@@ -151,10 +175,14 @@ void FGpuProfilerTrace::BeginEventByName(const FName& Name, uint32 FrameNumber,
 	GCurrentFrame.EventBufferSize = BufferPtr - GCurrentFrame.EventBuffer + sizeof(uint32);
 }
 
-void FGpuProfilerTrace::EndEvent(uint64 TimestampMicroseconds)
+void FGpuProfilerTrace::EndEvent(uint64 TimestampMicroseconds GPU_TRACE_ARG)
 {
 	using namespace GpuProfilerTrace;
 
+#if RHI_NEW_GPU_PROFILER
+	FGpuTraceFrame& GCurrentFrame = GCurrentFrames[GPUIndex];
+#endif
+
 	if (!GCurrentFrame.bActive)
 	{
 		return;
@@ -181,6 +209,10 @@ void FGpuProfilerTrace::EndFrame(uint32 GPUIndex)
 {
 	using namespace GpuProfilerTrace;
 
+#if RHI_NEW_GPU_PROFILER
+	FGpuTraceFrame& GCurrentFrame = GCurrentFrames[GPUIndex];
+#endif
+
 	if (GCurrentFrame.bActive && GCurrentFrame.EventBufferSize)
 	{
 		// This subtraction is intended to be performed on uint64 to leverage the wrap around behavior defined by the standard
@@ -208,13 +240,20 @@ void FGpuProfilerTrace::EndFrame(uint32 GPUIndex)
 	GCurrentFrame.bActive = false;
 }
 
+#undef GPU_TRACE_ARG
+
 void FGpuProfilerTrace::Deinitialize()
 {
 	using namespace GpuProfilerTrace;
 
-	FMemory::Free(GCurrentFrame.EventBuffer);
-	GCurrentFrame.EventBuffer = nullptr;
-	GCurrentFrame.MaxEventBufferSize = 0;
+#if RHI_NEW_GPU_PROFILER
+	for (FGpuTraceFrame& GCurrentFrame : GCurrentFrames)
+#endif
+	{
+		FMemory::Free(GCurrentFrame.EventBuffer);
+		GCurrentFrame.EventBuffer = nullptr;
+		GCurrentFrame.MaxEventBufferSize = 0;
+	}
 }
 
 #endif
diff --git a/Engine/Source/Runtime/RHI/Public/GpuProfilerTrace.h b/Engine/Source/Runtime/RHI/Public/GpuProfilerTrace.h
--- a/Engine/Source/Runtime/RHI/Public/GpuProfilerTrace.h
+++ b/Engine/Source/Runtime/RHI/Public/GpuProfilerTrace.h
@@ -28,16 +28,25 @@ struct FGPUTimingCalibrationTimestamp
 
 class FName;
 
+#if RHI_NEW_GPU_PROFILER
+// Adds a GPUIndex argument to each function in the API without breaking back compat
+#define GPU_TRACE_ARG , uint32 GPUIndex
+#else
+#define GPU_TRACE_ARG
+#endif
+
 struct FGpuProfilerTrace
 {
-	RHI_API static void BeginFrame(struct FGPUTimingCalibrationTimestamp& Calibration);
-	RHI_API static void SpecifyEventByName(const FName& Name);
-	RHI_API static void BeginEventByName(const FName& Name, uint32 FrameNumber, uint64 TimestampMicroseconds);
-	RHI_API static void EndEvent(uint64 TimestampMicroseconds);
+	RHI_API static void BeginFrame(struct FGPUTimingCalibrationTimestamp& Calibration GPU_TRACE_ARG);
+	RHI_API static void SpecifyEventByName(const FName& Name GPU_TRACE_ARG);
+	RHI_API static void BeginEventByName(const FName& Name, uint32 FrameNumber, uint64 TimestampMicroseconds GPU_TRACE_ARG);
+	RHI_API static void EndEvent(uint64 TimestampMicroseconds GPU_TRACE_ARG);
 	RHI_API static void EndFrame(uint32 GPUIndex);
 	RHI_API static void Deinitialize();
 };
 
+#undef GPU_TRACE_ARG
+
 #define TRACE_GPUPROFILER_DEFINE_EVENT_TYPE(Name) \
 	FGpuProfilerTrace::FEventType PREPROCESSOR_JOIN(__GGpuProfilerEventType, Name)(TEXT(#Name));
 
