Add support for preview shader platforms

diff --git a/Engine/Config/BaseDeviceProfiles.ini b/Engine/Config/BaseDeviceProfiles.ini
--- a/Engine/Config/BaseDeviceProfiles.ini
+++ b/Engine/Config/BaseDeviceProfiles.ini
@@ -9,6 +9,7 @@
 +DeviceProfileNameAndTypes=WindowsServer,Windows
 +DeviceProfileNameAndTypes=WindowsClient,Windows
 +DeviceProfileNameAndTypes=Windows_Preview_ES31,Windows
++DeviceProfileNameAndTypes=Windows_Preview_ES31_SDF,Windows
 +DeviceProfileNameAndTypes=IOS,IOS
 +DeviceProfileNameAndTypes=iPadAir2,IOS
 +DeviceProfileNameAndTypes=iPadAir3,IOS
@@ -1252,3 +1253,14 @@ BaseProfileName=Windows
 +CVars=r.Substrate.Glints=0
 +CVars=r.Substrate.SpecularProfile=0
 +CVars=r.Substrate.ClosuresPerPixel=1
+
+[Windows_Preview_ES31_SDF DeviceProfile]
+DeviceType=Windows
+BaseProfileName=Windows
++CVars=r.Substrate.RoughDiffuse=0
++CVars=r.Substrate.ShadingQuality=2
++CVars=r.Substrate.TileCoord8bits=1
++CVars=r.Substrate.SheenQuality=2
++CVars=r.Substrate.Glints=0
++CVars=r.Substrate.SpecularProfile=0
++CVars=r.Substrate.ClosuresPerPixel=1
diff --git a/Engine/Config/Windows/DataDrivenPlatformInfo.ini b/Engine/Config/Windows/DataDrivenPlatformInfo.ini
--- a/Engine/Config/Windows/DataDrivenPlatformInfo.ini
+++ b/Engine/Config/Windows/DataDrivenPlatformInfo.ini
@@ -170,6 +170,20 @@ MenuTooltip=LOCTEXT("PreviewMenuTooltip_PCD3D_Mobile", "Windows using Mobile pro
 DeviceProfileName=Windows_Preview_ES31
 PreviewDeviceProfileName=Windows_Preview_ES31
 
+[PreviewShaderPlatform SP_PCD3D_ES3_1_SDF]
+ParentShaderPlatform=PCD3D_ES3_1
+bSupportsDistanceFields=true
+FriendlyName=LOCTEXT("FriendlyShaderPlatformName_D3D_Mobile_SDF", "PC D3D Mobile SDF")
+
+[PreviewPlatform PCD3D_ES3_1_SDF]
+PlatformName=PCD3D_ES3_1
+ShaderFormat=PCD3D_ES31
+ShaderPlatform=SP_PCD3D_ES3_1_SDF
+PreviewShaderPlatform=SP_PCD3D_ES3_1_SDF
+MenuTooltip=LOCTEXT("PreviewMenuTooltip_PCD3D_Mobile_SDF", "Windows using Mobile SDF profile.")
+DeviceProfileName=Windows_Preview_ES31_SDF
+PreviewDeviceProfileName=Windows_Preview_ES31_SDF
+
 [ShaderPlatform OPENGL_PCES3_1]
 Language=OpenGL
 bIsPC=true
diff --git a/Engine/Source/Editor/UnrealEd/Private/EditorEngine.cpp b/Engine/Source/Editor/UnrealEd/Private/EditorEngine.cpp
--- a/Engine/Source/Editor/UnrealEd/Private/EditorEngine.cpp
+++ b/Engine/Source/Editor/UnrealEd/Private/EditorEngine.cpp
@@ -6284,7 +6284,7 @@ void UEditorEngine::SetPreviewPlatform(const FPreviewPlatformInfo& NewPreviewPla
 	if (bChangedPreviewShaderPlatform)
 	{
 		UMaterialShaderQualitySettings* MaterialShaderQualitySettings = UMaterialShaderQualitySettings::Get();
-		MaterialShaderQualitySettings->SetPreviewPlatform(PreviewPlatform.PreviewShaderFormatName);
+		MaterialShaderQualitySettings->SetPreviewPlatform(PreviewPlatform.PreviewShaderPlatformName);
 
 		UStaticMesh::OnLodStrippingQualityLevelChanged(nullptr);
 
diff --git a/Engine/Source/Runtime/Core/Private/Misc/DataDrivenPlatformInfoRegistry.cpp b/Engine/Source/Runtime/Core/Private/Misc/DataDrivenPlatformInfoRegistry.cpp
--- a/Engine/Source/Runtime/Core/Private/Misc/DataDrivenPlatformInfoRegistry.cpp
+++ b/Engine/Source/Runtime/Core/Private/Misc/DataDrivenPlatformInfoRegistry.cpp
@@ -247,6 +247,13 @@ static FString GetSectionString(const FConfigSection& Section, FName Key)
 	return Value ? Value->GetValue() : FString();
 }
 
+// Gets a string from a section, or default string if it didn't exist
+static inline FName GetSectionString(const FConfigSection& Section, FName Key, FName DefaultValue)
+{
+	const FConfigValue* Value = Section.Find(Key);
+	return Value ? *Value->GetValue() : DefaultValue;
+}
+
 #if DDPI_HAS_EXTENDED_PLATFORMINFO_DATA
 static void ParsePreviewPlatforms(const FConfigFile& IniFile)
 {
@@ -283,7 +290,7 @@ static void ParsePreviewPlatforms(const FConfigFile& IniFile)
 				
 				FPreviewPlatformMenuItem Item;
 				Item.PlatformName = PlatformName;
-				Item.PreviewShaderPlatformName = PreviewPlatformName;
+				Item.PreviewShaderPlatformName = GetSectionString(Section.Value, FName("PreviewShaderPlatform"), PreviewPlatformName);
 				Item.ShaderFormat = *GetSectionString(Section.Value, FName("ShaderFormat"));
 				checkf(Item.ShaderFormat != NAME_None, TEXT("DataDrivenPlatformInfo section [PreviewPlatform %s] must specify a ShaderFormat"), *SectionName);
 				Item.ActiveIconPath = GetSectionString(Section.Value, FName("ActiveIconPath"));
@@ -293,7 +300,7 @@ static void ParsePreviewPlatforms(const FConfigFile& IniFile)
 				Item.ShaderPlatformToPreview = *GetSectionString(Section.Value, FName("ShaderPlatform"));
 				Item.PreviewFeatureLevelName = *GetSectionString(Section.Value, FName("PreviewFeatureLevel"));
 
-				checkf(Item.ShaderPlatformToPreview != NAME_None, TEXT("DataDrivenPlatformInfo section [PreviewPlatform %s] must specify a ShaderPlatform"), *SectionName);
+				checkf(Item.ShaderPlatformToPreview != NAME_None, TEXT("DataDrivenPlatformInfo section [PreviewPlatform %s] must specify a ShaderPlatformToPreview"), *SectionName);
 				FTextStringHelper::ReadFromBuffer(*GetSectionString(Section.Value, FName("MenuTooltip")), Item.MenuTooltip);
 				FTextStringHelper::ReadFromBuffer(*GetSectionString(Section.Value, FName("IconText")), Item.IconText);
 
diff --git a/Engine/Source/Runtime/RHI/Private/DataDrivenShaderPlatformInfo.cpp b/Engine/Source/Runtime/RHI/Private/DataDrivenShaderPlatformInfo.cpp
--- a/Engine/Source/Runtime/RHI/Private/DataDrivenShaderPlatformInfo.cpp
+++ b/Engine/Source/Runtime/RHI/Private/DataDrivenShaderPlatformInfo.cpp
@@ -52,6 +52,13 @@ static inline FString GetSectionString(const FConfigSection& Section, FName Key)
 	return Value ? Value->GetValue() : FString();
 }
 
+// Gets a string from a section, or default string if it didn't exist
+static inline FName GetSectionString(const FConfigSection& Section, FName Key, FName DefaultValue)
+{
+	const FConfigValue* Value = Section.Find(Key);
+	return Value ? *Value->GetValue() : DefaultValue;
+}
+
 // Gets a bool from a section.  It returns the original value if the setting does not exist
 static inline bool GetSectionBool(const FConfigSection& Section, FName Key, bool OriginalValue)
 {
@@ -184,11 +191,15 @@ void FGenericDataDrivenShaderPlatformInfo::ParseDataDrivenShaderInfo(const FConf
 {
 	FGenericDataDrivenShaderPlatformInfo& Info = Infos[Index];
 
-	Info.Language = *GetSectionString(Section, "Language");
-	Info.ShaderFormat = *GetSectionString(Section, "ShaderFormat");
+	Info.Language = GetSectionString(Section, "Language", Info.Language);
+	Info.ShaderFormat = GetSectionString(Section, "ShaderFormat", Info.ShaderFormat);
 	checkf(!Info.ShaderFormat.IsNone(), TEXT("Missing ShaderFormat for ShaderPlatform %s  ShaderFormat %s"), *Info.Name.ToString(), *Info.ShaderFormat.ToString());
 
-	GetFeatureLevelFromName(GetSectionString(Section, "MaxFeatureLevel"), Info.MaxFeatureLevel);
+	const FConfigValue* MaxFeatureLevelValue = Section.Find("MaxFeatureLevel");
+	if (MaxFeatureLevelValue)
+	{
+		GetFeatureLevelFromName(MaxFeatureLevelValue->GetValue(), Info.MaxFeatureLevel);
+	}
 
 	Info.ShaderPropertiesHash = 0;
 	FString ShaderPropertiesString = Info.Name.GetPlainNameString();
@@ -437,6 +448,50 @@ void FGenericDataDrivenShaderPlatformInfo::Initialize()
 				}
 #endif
 			}
+#if WITH_EDITOR
+			else if (Section.Key.StartsWith(TEXT("PreviewShaderPlatform ")))
+			{
+				const FString& SectionName = Section.Key;
+				const FConfigSection& SectionSettings = Section.Value;
+
+				const FString ParentShaderPlatformName = GetSectionString(SectionSettings, "ParentShaderPlatform");
+				const EShaderPlatform ParentShaderPlatform = ParseShaderPlatform(*ParentShaderPlatformName);
+				if (ParentShaderPlatform == SP_NumPlatforms)
+				{
+#if DDPI_HAS_EXTENDED_PLATFORMINFO_DATA
+					const bool bIsEnabled = FDataDrivenPlatformInfoRegistry::GetPlatformInfo(PlatformName).bEnabledForUse;
+#else
+					const bool bIsEnabled = true;
+#endif
+					UE_CLOG(bIsEnabled, LogRHI, Warning, TEXT("Found an unknown parent shader platform %s in a preview shader platform DataDriven ini file"), *ParentShaderPlatformName);
+					continue;
+				}
+
+				check(IsValid(ParentShaderPlatform));
+
+				// get enum value for the string name
+				const EShaderPlatform ShaderPlatform = EShaderPlatform(CustomShaderPlatform++);
+
+				FGenericDataDrivenShaderPlatformInfo& Info = Infos[ShaderPlatform];
+				const FGenericDataDrivenShaderPlatformInfo& ParentInfo = Infos[ParentShaderPlatform];
+				Info = ParentInfo;
+				Info.Name = *SectionName.Mid(22);
+				Info.bIsPreviewPlatform = true;
+				Info.bContainsValidPlatformInfo = true;
+				ParseDataDrivenShaderInfo(SectionSettings, ShaderPlatform);
+
+				ERHIFeatureLevel::Type PreviewFeatureLevel = ERHIFeatureLevel::Num;
+				if (GetFeatureLevelFromName(GetSectionString(SectionSettings, "PreviewFeatureLevel"), PreviewFeatureLevel))
+				{
+					Info.MaxFeatureLevel = PreviewFeatureLevel;
+				}
+
+				PlatformNameToShaderPlatformMap.FindOrAdd(Info.Name) = ShaderPlatform;
+
+				FDataDrivenShaderPlatformInfoEditorOnly& PreviewEditorInfo = DataDrivenShaderPlatformInfoEditorOnlyInfos[ShaderPlatform];
+				PreviewEditorInfo.PreviewShaderPlatformParent = ShaderPlatform;
+			}
+#endif
 		}
 	}
 	bInitialized = true;
diff --git a/Engine/Source/Runtime/RHI/Public/DataDrivenShaderPlatformInfo.h b/Engine/Source/Runtime/RHI/Public/DataDrivenShaderPlatformInfo.h
--- a/Engine/Source/Runtime/RHI/Public/DataDrivenShaderPlatformInfo.h
+++ b/Engine/Source/Runtime/RHI/Public/DataDrivenShaderPlatformInfo.h
@@ -139,6 +139,8 @@ class FGenericDataDrivenShaderPlatformInfo
 		SetDefaultValues();
 	}
 
+	FGenericDataDrivenShaderPlatformInfo(const FGenericDataDrivenShaderPlatformInfo&) = default;
+
 	RHI_API void SetDefaultValues();
 
 public:
