Add PrecachePSOs() to component interface.
Component interface objects now recollect PSOs when FComponentRecreateRenderStateContext is used.
 - Fixes issues where newly drawn primitives would attempt to boost stale request IDs.
Add an option to generate all static mesh component's PSO tasks as highest priority.

diff --git a/Engine/Source/Runtime/Engine/Classes/Components/ActorPrimitiveComponentInterface.h b/Engine/Source/Runtime/Engine/Classes/Components/ActorPrimitiveComponentInterface.h
--- a/Engine/Source/Runtime/Engine/Classes/Components/ActorPrimitiveComponentInterface.h
+++ b/Engine/Source/Runtime/Engine/Classes/Components/ActorPrimitiveComponentInterface.h
@@ -20,6 +20,7 @@ public:
 	void MarkRenderStateDirty() override;
 	void DestroyRenderState() override;
 	void CreateRenderState(FRegisterComponentContext* Context) override;
+	void PrecachePSOs() override;
 
 	FString GetName() const override;
 	FString GetFullName() const override;
diff --git a/Engine/Source/Runtime/Engine/Classes/Components/ComponentInterfaces.h b/Engine/Source/Runtime/Engine/Classes/Components/ComponentInterfaces.h
--- a/Engine/Source/Runtime/Engine/Classes/Components/ComponentInterfaces.h
+++ b/Engine/Source/Runtime/Engine/Classes/Components/ComponentInterfaces.h
@@ -78,6 +78,7 @@ public:
 	virtual void GetPrimitiveStats(FPrimitiveStats& PrimitiveStats) const = 0;
 	virtual UObject*	GetUObject() = 0;
 	virtual const UObject*	GetUObject() const = 0;
+	virtual void PrecachePSOs() = 0;
 
 	// helper to obtain typed UObjects 
 	template<class T> 
diff --git a/Engine/Source/Runtime/Engine/Classes/Components/PrimitiveComponent.h b/Engine/Source/Runtime/Engine/Classes/Components/PrimitiveComponent.h
--- a/Engine/Source/Runtime/Engine/Classes/Components/PrimitiveComponent.h
+++ b/Engine/Source/Runtime/Engine/Classes/Components/PrimitiveComponent.h
@@ -670,15 +670,16 @@ protected:
 		ToolTip = "When false, the underlying physics body will contain all sim data (mass, inertia tensor, etc) even if mobility is not set to Moveable"))
 	uint8 bStaticWhenNotMoveable:1;
 
+	uint8 bIgnoreBoundsForEditorFocus : 1;
+
+public:
 #if UE_WITH_PSO_PRECACHING
 	/** PSOs requested priority */
 	EPSOPrecachePriority PSOPrecacheRequestPriority : 2;
 	static_assert((int)EPSOPrecachePriority::Highest < 1 << 2);
 #endif
 
-	uint8 bIgnoreBoundsForEditorFocus : 1;
 #if WITH_EDITOR
-public:
 	uint8 bAlwaysAllowTranslucentSelect : 1;
 
 	uint8 SelectionOutlineColorIndex;
@@ -686,7 +687,6 @@ public:
 	FColor OverlayColor;
 #endif
 
-public:
 	/** If true then DoCustomNavigableGeometryExport will be called to collect navigable geometry of this component. */
 	UPROPERTY()
 	TEnumAsByte<EHasCustomNavigableGeometry::Type> bHasCustomNavigableGeometry;
diff --git a/Engine/Source/Runtime/Engine/Private/Components/ActorComponent.cpp b/Engine/Source/Runtime/Engine/Private/Components/ActorComponent.cpp
--- a/Engine/Source/Runtime/Engine/Private/Components/ActorComponent.cpp
+++ b/Engine/Source/Runtime/Engine/Private/Components/ActorComponent.cpp
@@ -330,6 +330,10 @@ FComponentRecreateRenderStateContext::~FComponentRecreateRenderStateContext()
 
 	if (ComponentInterface && !ComponentInterface->IsRenderStateCreated() && ComponentInterface->IsRegistered())
 	{
+		if (GPrecachePSOsOnComponentRecreateRenderContext)
+		{
+ 			ComponentInterface->PrecachePSOs();
+		}
 		ComponentInterface->CreateRenderState(nullptr);
 
 		UpdateAllPrimitiveSceneInfosForSingleComponentInterface(ComponentInterface, ScenesToUpdateAllPrimitiveSceneInfos);
diff --git a/Engine/Source/Runtime/Engine/Private/Components/PrimitiveComponent.cpp b/Engine/Source/Runtime/Engine/Private/Components/PrimitiveComponent.cpp
--- a/Engine/Source/Runtime/Engine/Private/Components/PrimitiveComponent.cpp
+++ b/Engine/Source/Runtime/Engine/Private/Components/PrimitiveComponent.cpp
@@ -5183,15 +5183,14 @@ FPrimitiveSceneProxy* FActorPrimitiveComponentInterface::CreateSceneProxy()
 	FPrimitiveSceneProxy* SceneProxy = Component->CreateSceneProxy();
 	Component->AssignSceneProxy(SceneProxy);
 
-#if UE_WITH_PSO_PRECACHING
-	if (ShouldBoostPSOPrecachePriorityOnDraw() && Component->IsPSOPrecaching() && Component->PSOPrecacheRequestPriority == EPSOPrecachePriority::High && SceneProxy)
-	{
-		SceneProxy->SetPSORequestsToBoostOnDraw(Component->MaterialPSOPrecacheRequestIDs);
-	}
-#endif
 	return SceneProxy;
 }
 
+void FActorPrimitiveComponentInterface::PrecachePSOs()
+{
+	UPrimitiveComponent::GetPrimitiveComponent(this)->PrecachePSOs();
+}
+
 #if WITH_EDITOR
 HHitProxy* FActorPrimitiveComponentInterface::CreateMeshHitProxy(int32 SectionIndex, int32 MaterialIndex) 
 {
diff --git a/Engine/Source/Runtime/Engine/Private/PSOPrecache.cpp b/Engine/Source/Runtime/Engine/Private/PSOPrecache.cpp
--- a/Engine/Source/Runtime/Engine/Private/PSOPrecache.cpp
+++ b/Engine/Source/Runtime/Engine/Private/PSOPrecache.cpp
@@ -9,6 +9,8 @@
 #include "HAL/IConsoleManager.h"
 #include "ShaderCodeLibrary.h"
 #include "Materials/MaterialInterface.h"
+#include "PrimitiveSceneProxy.h"
+#include "PrimitiveSceneInfo.h"
 
 static TAutoConsoleVariable<int32> CVarPrecacheGlobalComputeShaders(
 	TEXT("r.PSOPrecache.GlobalShaders"),
@@ -129,6 +131,16 @@ bool ProxyCreationWhenPSOReady()
 	return FApp::CanEverRender() && (PipelineStateCache::IsPSOPrecachingEnabled() || IsPSOShaderPreloadingEnabled()) && GPSOProxyCreationWhenPSOReady && !GIsEditor;
 }
 
+void BoostPrecachedPSORequestsOnDraw(const FPrimitiveSceneInfo* SceneInfo)
+{
+#if UE_WITH_PSO_PRECACHING 
+	if (SceneInfo && SceneInfo->Proxy)
+	{
+		SceneInfo->Proxy->BoostPrecachedPSORequestsOnDraw();
+	}
+#endif
+}
+
 FPSOPrecacheVertexFactoryData::FPSOPrecacheVertexFactoryData(
 	const FVertexFactoryType* InVertexFactoryType, const FVertexDeclarationElementList& ElementList) 
 	: VertexFactoryType(InVertexFactoryType)
diff --git a/Engine/Source/Runtime/Engine/Private/PrimitiveSceneProxy.cpp b/Engine/Source/Runtime/Engine/Private/PrimitiveSceneProxy.cpp
--- a/Engine/Source/Runtime/Engine/Private/PrimitiveSceneProxy.cpp
+++ b/Engine/Source/Runtime/Engine/Private/PrimitiveSceneProxy.cpp
@@ -684,6 +684,18 @@ FPrimitiveSceneProxy::FPrimitiveSceneProxy(const FPrimitiveSceneProxyDesc& InPro
 	{
 		bHasWorldPositionOffsetVelocity = true;
 	}
+
+#if UE_WITH_PSO_PRECACHING
+	if (UPrimitiveComponent* PrimitiveComponent = Cast<UPrimitiveComponent>(InProxyDesc.Component))
+	{
+		check(!PrimitiveComponent->ShouldRenderProxyFallbackToDefaultMaterial() || !PrimitiveComponent->MaterialPSOPrecacheRequestIDs.IsEmpty());
+		if (ShouldBoostPSOPrecachePriorityOnDraw() && PrimitiveComponent->IsPSOPrecaching() && PrimitiveComponent->PSOPrecacheRequestPriority == EPSOPrecachePriority::High)
+		{
+			SetPSORequestsToBoostOnDraw(PrimitiveComponent->MaterialPSOPrecacheRequestIDs);
+		}
+	}
+#endif
+
 }
 
 bool FPrimitiveSceneProxy::OnLevelAddedToWorld_RenderThread()
diff --git a/Engine/Source/Runtime/Engine/Private/StaticMeshSceneProxy.cpp b/Engine/Source/Runtime/Engine/Private/StaticMeshSceneProxy.cpp
--- a/Engine/Source/Runtime/Engine/Private/StaticMeshSceneProxy.cpp
+++ b/Engine/Source/Runtime/Engine/Private/StaticMeshSceneProxy.cpp
@@ -85,6 +85,16 @@ void TogglePreCulledIndexBuffers( UWorld* InWorld )
 	GUsePreCulledIndexBuffer = !GUsePreCulledIndexBuffer;
 }
 
+static bool GStaticMeshComponentBoostPSOPrecachePri = false;
+static FAutoConsoleVariableRef CVarStaticMeshComponentBoostPSOPrecachePri(
+	TEXT("r.PSOPrecache.StaticMeshComponentPSOPrecachePriority"),
+	GStaticMeshComponentBoostPSOPrecachePri,
+	TEXT("Static Mesh component PSO precache priority level.\n")
+	TEXT(" 0. Static Mesh component's PSO precache requests are set to high priority (default)\n")
+	TEXT(" 1. Static Mesh component's PSO precache requests are set to highest priority"),
+	ECVF_Default
+);
+
 FAutoConsoleCommandWithWorld GToggleUsePreCulledIndexBuffersCmd(
 	TEXT("r.TogglePreCulledIndexBuffers"),
 	TEXT("Toggles use of preculled index buffers from the command 'PreCullIndexBuffers'"),
@@ -2684,7 +2694,8 @@ FPrimitiveSceneProxy* UStaticMeshComponent::CreateSceneProxy()
 		return nullptr;
 	}
 
-	if (CheckPSOPrecachingAndBoostPriority() && GetPSOPrecacheProxyCreationStrategy() == EPSOPrecacheProxyCreationStrategy::DelayUntilPSOPrecached)
+	EPSOPrecachePriority PSOPrecachePriority = GStaticMeshComponentBoostPSOPrecachePri ? EPSOPrecachePriority::Highest : EPSOPrecachePriority::High;
+	if (CheckPSOPrecachingAndBoostPriority(PSOPrecachePriority) && GetPSOPrecacheProxyCreationStrategy() == EPSOPrecacheProxyCreationStrategy::DelayUntilPSOPrecached)
 	{
 		UE_LOG(LogStaticMesh, Verbose, TEXT("Skipping CreateSceneProxy for StaticMeshComponent %s (Static mesh component PSOs are still compiling)"), *GetFullName());
 		return nullptr;
diff --git a/Engine/Source/Runtime/Engine/Public/PSOPrecache.h b/Engine/Source/Runtime/Engine/Public/PSOPrecache.h
--- a/Engine/Source/Runtime/Engine/Public/PSOPrecache.h
+++ b/Engine/Source/Runtime/Engine/Public/PSOPrecache.h
@@ -333,3 +333,8 @@ extern ENGINE_API EPSOPrecacheProxyCreationStrategy GetPSOPrecacheProxyCreationS
  * Delay component proxy creation when it's requested PSOs are still precaching
  */
 extern ENGINE_API bool ProxyCreationWhenPSOReady();
+
+/**
+ * Boost the primitive's precache PSO jobs to highest priority.
+ */
+extern ENGINE_API void BoostPrecachedPSORequestsOnDraw(const class FPrimitiveSceneInfo* SceneInfo);
\ No newline at end of file
diff --git a/Engine/Source/Runtime/RHI/Private/PipelineStateCache.cpp b/Engine/Source/Runtime/RHI/Private/PipelineStateCache.cpp
--- a/Engine/Source/Runtime/RHI/Private/PipelineStateCache.cpp
+++ b/Engine/Source/Runtime/RHI/Private/PipelineStateCache.cpp
@@ -194,8 +194,19 @@ int32 GPSOWaitForHighPriorityRequestsOnly = 0;
 static FAutoConsoleVariableRef CVarPSOWaitForHighPriorityRequestsOnly(
 	TEXT("r.PSOPrecaching.WaitForHighPriorityRequestsOnly"),
 	GPSOWaitForHighPriorityRequestsOnly,
-	TEXT("0 to wait for all pending PSO precache requests during loading (default)g\n")
-	TEXT("1 to only wait for the high priority PSO precache requests during loading\n"),
+	TEXT("0 to wait for all pending PSO precache requests during loading (default)\n")
+	TEXT("1 to only wait for the high priority and above PSO precache requests during loading\n")
+	TEXT("2 to only wait for the highest priority PSO precache requests during loading"),
+	ECVF_Default
+);
+
+bool GPSOPrecachePermitPriorityEscalation = true;
+static FAutoConsoleVariableRef CVarPSOPrecachePermitPriorityEscalation(
+	TEXT("r.PSOPrecaching.PermitPriorityEscalation"),
+	GPSOPrecachePermitPriorityEscalation,
+	TEXT("Whether to permit requests to increase high pri PSO precaching tasks to highest.\n")
+	TEXT("1: High priority tasks can be escalated to highest if requested. (default)\n")
+	TEXT("0: High priority tasks will remain unchanged."),
 	ECVF_Default
 );
 
@@ -2177,13 +2188,15 @@ public:
 
 	uint32 NumActivePrecacheRequests()
 	{
-		if (GPSOWaitForHighPriorityRequestsOnly)
-		{
-			return FPlatformAtomics::AtomicRead(&HighPriorityCompileCount) + FPlatformAtomics::AtomicRead(&HighestPriorityCompileCount);
-		}
-		else
+		switch (GPSOWaitForHighPriorityRequestsOnly)
 		{
-			return FPlatformAtomics::AtomicRead(&ActiveCompileCount);
+			default: checkNoEntry(); [[fallthrough]];
+			case 0:
+				return FPlatformAtomics::AtomicRead(&ActiveCompileCount);
+			case 1:
+				return FPlatformAtomics::AtomicRead(&HighPriorityCompileCount) + FPlatformAtomics::AtomicRead(&HighestPriorityCompileCount);
+			case 2:
+				return FPlatformAtomics::AtomicRead(&HighestPriorityCompileCount);
 		}
 	}
 
@@ -4103,7 +4116,8 @@ void PipelineStateCache::BoostPrecachePriority(EPSOPrecachePriority PSOPrecacheP
 
 void PipelineStateCache::PrecachePSOsBoostToHighestPriority(bool bForceHighest)
 {
-	UE_LOG(LogRHI, Log, TEXT("PipelineStateCache: PSO precaching %s highest priority boost"), bForceHighest ? TEXT("enabling") : TEXT("disabling"));
+	bForceHighest = bForceHighest && GPSOPrecachePermitPriorityEscalation;
+	UE_CLOG(GForceHighToHighestPri != bForceHighest, LogRHI, Log, TEXT("PipelineStateCache: PSO precaching %s highest priority boost"), bForceHighest ? TEXT("enabling") : TEXT("disabling"));
 	GForceHighToHighestPri = bForceHighest;
 #if PSO_TRACK_CACHE_STATS
 	DumpPipelineCacheStats();
