Use the Runtime Value of bSupportsInlineRayTracing and bSupportsRayTracingShaders for preview platforms
This avoids getting in a bad state where the Preview platform has bSupportsRayTracing=true, but bSupportsRayTracingShaders=bSupportsInlineRayTracing=false while Runtime RHI globals like GRHISupportsRayTracing is true, leading to crashes because of missing shaders. Like Android VK SM5 Preview in D3D Editor

diff --git a/Engine/Source/Runtime/RHI/Private/DataDrivenShaderPlatformInfo.cpp b/Engine/Source/Runtime/RHI/Private/DataDrivenShaderPlatformInfo.cpp
--- a/Engine/Source/Runtime/RHI/Private/DataDrivenShaderPlatformInfo.cpp
+++ b/Engine/Source/Runtime/RHI/Private/DataDrivenShaderPlatformInfo.cpp
@@ -562,8 +562,6 @@ void FGenericDataDrivenShaderPlatformInfo::UpdatePreviewPlatforms()
 				PREVIEW_DISABLE_IF_RUNTIME_UNSUPPORTED(bSupportsPrimitiveShaders);
 				PREVIEW_DISABLE_IF_RUNTIME_UNSUPPORTED(bSupportsUInt64ImageAtomics);
 				PREVIEW_DISABLE_IF_RUNTIME_UNSUPPORTED(bSupportsGen5TemporalAA);
-				PREVIEW_DISABLE_IF_RUNTIME_UNSUPPORTED(bSupportsInlineRayTracing);
-				PREVIEW_DISABLE_IF_RUNTIME_UNSUPPORTED(bSupportsRayTracingShaders);
 				PREVIEW_DISABLE_IF_RUNTIME_UNSUPPORTED(bSupportsMeshShadersTier0);
 				PREVIEW_DISABLE_IF_RUNTIME_UNSUPPORTED(bSupportsMeshShadersTier1);
 				PREVIEW_DISABLE_IF_RUNTIME_UNSUPPORTED(bSupportsMobileMultiView);
@@ -578,6 +576,9 @@ void FGenericDataDrivenShaderPlatformInfo::UpdatePreviewPlatforms()
 				PREVIEW_USE_RUNTIME_VALUE(bSupportsRealTypes);
 				PREVIEW_USE_RUNTIME_VALUE(bSupportsUniformBufferObjects);
 
+				PREVIEW_USE_RUNTIME_VALUE(bSupportsInlineRayTracing);
+				PREVIEW_USE_RUNTIME_VALUE(bSupportsRayTracingShaders);
+
 				// Settings that will never be supported in preview
 				PREVIEW_FORCE_DISABLE(bSupportsShaderRootConstants);
 				PREVIEW_FORCE_DISABLE(bSupportsShaderBundleDispatch);
