Minor FRHIShaderBundle cleanup.

diff --git a/Engine/Source/Runtime/RHI/Public/RHIResources.h b/Engine/Source/Runtime/RHI/Public/RHIResources.h
--- a/Engine/Source/Runtime/RHI/Public/RHIResources.h
+++ b/Engine/Source/Runtime/RHI/Public/RHIResources.h
@@ -3478,11 +3478,8 @@ enum class ERHIShaderBundleMode : uint8
 	MAX
 };
 
-class FShaderBundleCreateInfo
+struct FShaderBundleCreateInfo
 {
-public:
-	FShaderBundleCreateInfo() = default;
-
 	uint32 NumRecords	= 0u;
 	uint32 ArgOffset	= 0u;
 	uint32 ArgStride	= 0u;
@@ -3493,14 +3490,13 @@ public:
 class FRHIShaderBundle : public FRHIResource
 {
 public:
-	const uint32 NumRecords	= 0;
-	const uint32 ArgOffset	= 0;
-	const uint32 ArgStride	= 0;
-	const ERHIShaderBundleMode Mode = ERHIShaderBundleMode::CS;
-
-	FName ModeName;
+	const uint32 NumRecords;
+	const uint32 ArgOffset;
+	const uint32 ArgStride;
+	const ERHIShaderBundleMode Mode;
 
 public:
+	FRHIShaderBundle() = delete;
 	FRHIShaderBundle(const FShaderBundleCreateInfo& CreateInfo)
 	: FRHIResource(RRT_ShaderBundle)
 	, NumRecords(CreateInfo.NumRecords)
@@ -3508,28 +3504,45 @@ public:
 	, ArgStride(CreateInfo.ArgStride)
 	, Mode(CreateInfo.Mode)
 	{
+#if DO_CHECK
 		if (Mode == ERHIShaderBundleMode::CS)
 		{
 			// Load3
 			check(ArgStride >= 12u);
-			ModeName = TEXT("CS");
 		}
 		else if (Mode == ERHIShaderBundleMode::MSPS)
 		{
 			// Load
 			check(ArgStride >= 4u);
-			ModeName = TEXT("MSPS");
 		}
 		else if (Mode == ERHIShaderBundleMode::VSPS)
 		{
 			// Load4
 			check(ArgStride >= 16u);
-			ModeName = TEXT("VSPS");
 		}
 		else
 		{
 			checkNoEntry();
 		}
+#endif
+	}
+
+	const TCHAR* GetModeName() const
+	{
+		switch (Mode)
+		{
+		case ERHIShaderBundleMode::CS:
+			return TEXT("CS");
+		case ERHIShaderBundleMode::MSPS:
+			return TEXT("MSPS");
+		case ERHIShaderBundleMode::VSPS:
+			return TEXT("VSPS");
+			break;
+		case ERHIShaderBundleMode::MAX:
+		default:
+			checkNoEntry();
+			return TEXT("<none>");
+		}
 	}
 };
 
