Metal RHI

- Added warning in the log to say you have a compatible SM6 Mac but wrong OS
- Added error message if you run with force SM6 when you dont have SM6

diff --git a/Engine/Source/Runtime/Apple/MetalRHI/Private/MetalRHI.cpp b/Engine/Source/Runtime/Apple/MetalRHI/Private/MetalRHI.cpp
--- a/Engine/Source/Runtime/Apple/MetalRHI/Private/MetalRHI.cpp
+++ b/Engine/Source/Runtime/Apple/MetalRHI/Private/MetalRHI.cpp
@@ -438,6 +438,18 @@ FMetalDynamicRHI::FMetalDynamicRHI(ERHIFeatureLevel::Type RequestedFeatureLevel)
 	bool const bRequestedSM5 = RequestedFeatureLevel == ERHIFeatureLevel::SM5 ||
                                (!bRequestedFeatureLevel && (FParse::Param(FCommandLine::Get(),TEXT("metalsm5")) || FParse::Param(FCommandLine::Get(),TEXT("metalmrt"))));
                                 
+	if(bRequestedSM6 && !bSupportsSM6)
+	{
+		if(GRHIAdapterName.Contains("Apple") && !GRHIAdapterName.Contains("M1"))
+		{
+			UE_LOG(LogMetal, Warning, TEXT("To use SM6 on this system, please ensure you are running Mac OS 15. Falling back to SM5"));
+		}
+		else
+		{
+			UE_LOG(LogMetal, Warning, TEXT("SM6 is enabled but is not supported on this system, falling back to SM5"));
+		}
+	}
+	
     if(bSupportsSM6 && bRequestedSM6)
     {
         GMaxRHIFeatureLevel = ERHIFeatureLevel::SM6;
diff --git a/Engine/Source/Runtime/RHI/Private/Apple/AppleDynamicRHI.cpp b/Engine/Source/Runtime/RHI/Private/Apple/AppleDynamicRHI.cpp
--- a/Engine/Source/Runtime/RHI/Private/Apple/AppleDynamicRHI.cpp
+++ b/Engine/Source/Runtime/RHI/Private/Apple/AppleDynamicRHI.cpp
@@ -10,6 +10,7 @@
 #include "DataDrivenShaderPlatformInfo.h"
 #include COMPILED_PLATFORM_HEADER_WITH_PREFIX(Apple/Platform, PlatformDynamicRHI.h)
 
+#define LOCTEXT_NAMESPACE "AppleRHI"
 
 //------------------------------------------------------------------------------
 // MARK: - FAppleDynamicRHIOptions Union
@@ -48,6 +49,21 @@ static inline bool ValidateAppleDynamicRHIOptions(FAppleDynamicRHIOptions* Optio
 		Options->ForceSM6 = 0;
 		Options->ForceMTL = 1;
 	}
+	if(Options->ForceSM6)
+	{
+		bool bSupportsSM6 = false;
+		if (@available(macOS 15.0, *))
+		{
+			bSupportsSM6 = true;
+		}
+		
+		if(!bSupportsSM6)
+		{
+			FMessageDialog::Open(EAppMsgType::Ok, LOCTEXT("RHIOptionsError", "-sm6 is selected but Mac requires OS 15 to support SM6"));
+			UE_LOG(LogRHI, Fatal, TEXT("-sm6 is selected but Mac requires OS 15 to support SM6"));
+			return false;
+		}
+	}
 	return true;
 }
 
