Fix compiler error for missing LOCTEXT_NAMESPACE undef  in AppleDynamicRHI.cpp

diff --git a/Engine/Source/Runtime/RHI/Private/Apple/AppleDynamicRHI.cpp b/Engine/Source/Runtime/RHI/Private/Apple/AppleDynamicRHI.cpp
--- a/Engine/Source/Runtime/RHI/Private/Apple/AppleDynamicRHI.cpp
+++ b/Engine/Source/Runtime/RHI/Private/Apple/AppleDynamicRHI.cpp
@@ -162,3 +162,5 @@ FDynamicRHI* PlatformCreateDynamicRHI()
 
 	return DynamicRHI;
 }
+
+#undef LOCTEXT_NAMESPACE
