Fixed Localization issue with duplicate keys in AppleDynamicRHI.cpp

diff --git a/Engine/Source/Runtime/RHI/Private/Apple/AppleDynamicRHI.cpp b/Engine/Source/Runtime/RHI/Private/Apple/AppleDynamicRHI.cpp
--- a/Engine/Source/Runtime/RHI/Private/Apple/AppleDynamicRHI.cpp
+++ b/Engine/Source/Runtime/RHI/Private/Apple/AppleDynamicRHI.cpp
@@ -59,7 +59,7 @@ static inline bool ValidateAppleDynamicRHIOptions(FAppleDynamicRHIOptions* Optio
 		
 		if(!bSupportsSM6)
 		{
-			FMessageDialog::Open(EAppMsgType::Ok, LOCTEXT("RHIOptionsError", "-sm6 is selected but Mac requires OS 15 to support SM6"));
+			FMessageDialog::Open(EAppMsgType::Ok, LOCTEXT("MetalRHIOptionsError", "-sm6 is selected but Mac requires OS 15 to support SM6"));
 			UE_LOG(LogRHI, Fatal, TEXT("-sm6 is selected but Mac requires OS 15 to support SM6"));
 			return false;
 		}
@@ -107,7 +107,7 @@ static void ComputeRequestedFeatureLevel(const FAppleDynamicRHIOptions& Options,
 		}
 		else
 		{
-			FMessageDialog::Open(EAppMsgType::Ok, LOCTEXT("RHIOptionsError", "No Targeted RHI is set for this project, defaulting to SM5"));
+			FMessageDialog::Open(EAppMsgType::Ok, LOCTEXT("MetalMissingTargetError", "No Targeted RHI is set for this project, defaulting to SM5"));
 			RequestedFeatureLevel = ERHIFeatureLevel::SM5;
 		}
 	}
