RHI
- Implement high level resource state tracking and patching Unknown states used in Transition before reaching RHI platform level
- New validation command line argument -RHIValidationLogStack that print human readable callstack from symbols when printing transition validation errors

diff --git a/Engine/Source/Runtime/RHI/Private/DynamicRHI.cpp b/Engine/Source/Runtime/RHI/Private/DynamicRHI.cpp
--- a/Engine/Source/Runtime/RHI/Private/DynamicRHI.cpp
+++ b/Engine/Source/Runtime/RHI/Private/DynamicRHI.cpp
@@ -23,6 +23,7 @@
 #include "RHIImmutableSamplerState.h"
 #include "RHIStrings.h"
 #include "RHITextureReference.h"
+#include "RHIAllocators.h"
 
 static_assert(sizeof(FRayTracingGeometryInstance) <= 104,
 	"Ray tracing instance descriptor is expected to be no more than 104 bytes, as there may be a very large number of them.");
@@ -282,6 +283,8 @@ void RHIInit(bool bHasEditorToken)
 		FGenericDataDrivenShaderPlatformInfo::Initialize();
 
 		GRHICommandList.LatchBypass(); // read commandline for bypass flag
+		// Platform RHIs can overidde this
+		GRHIGlobals.EnableHighLevelTransitionPatching = FParse::Param(FCommandLine::Get(), TEXT("EnableHighLevelTransitionsPatching"));
 
 		if (!FApp::CanEverRender())
 		{
@@ -797,6 +800,26 @@ FRenderQueryPoolRHIRef RHICreateRenderQueryPool(ERenderQueryType QueryType, uint
 	return new FDefaultRHIRenderQueryPool(QueryType);
 }
 
+FRHITransition* RHICreateTransition(const FRHITransitionCreateInfo& CreateInfo)
+{
+	// Placement create the transition on the heap.
+	FRHITransition* Transition = new (FRHICmdListBaseLinearAllocator::Malloc(FRHITransition::GetTotalAllocationSize(), (uint32)FRHITransition::GetAlignment())) FRHITransition(CreateInfo.SrcPipelines, CreateInfo.DstPipelines,
+		EnumHasAnyFlags(CreateInfo.Flags, ERHITransitionCreateFlags::AllowDuringRenderPass) ? ETransitionFlag::AllowDuringRenderPass : ETransitionFlag::None);
+
+	if (GRHIGlobals.EnableHighLevelTransitionPatching)
+	{
+		// Call to GDynamicRHI->RHICreateTransition will happen after the transitions have been tracked and patched, we are deep copying the CreateInfo here
+		Transition->CreationInfoData = new (FRHICmdListBaseLinearAllocator::Malloc(sizeof(FRHITransitionCreateInfoData), 16)) FRHITransitionCreateInfoData(CreateInfo);
+	}
+	else
+	{
+		EnumAddFlags(Transition->Flags, ETransitionFlag::PlatformLevelCreated);
+		GDynamicRHI->RHICreateTransition(Transition, CreateInfo);
+	}
+
+	return Transition;
+}
+
 void FDynamicRHI::RHICheckViewportHDRStatus(FRHIViewport* Viewport)
 {
 }
diff --git a/Engine/Source/Runtime/RHI/Private/RHICommandList.cpp b/Engine/Source/Runtime/RHI/Private/RHICommandList.cpp
--- a/Engine/Source/Runtime/RHI/Private/RHICommandList.cpp
+++ b/Engine/Source/Runtime/RHI/Private/RHICommandList.cpp
@@ -864,6 +864,8 @@ void FRHICommandListExecutor::FSubmitState::Dispatch(FRHICommandListBase* CmdLis
 		//
 		Prereqs.Add(GRHICommandList.LastSubmit);
 
+		CmdList->UpdateTransitionsAndResourceState();
+
 		CurrentTranslateJob->EnqueueTranslateTask(MoveTemp(Prereqs),
 			[TranslateJob = CurrentTranslateJob, CmdList]()
 			{
@@ -1680,31 +1682,7 @@ void FRHICommandListImmediate::EndFrame()
 
 void FRHIComputeCommandList::Transition(TArrayView<const FRHITransitionInfo> Infos, ERHITransitionCreateFlags CreateFlags)
 {
-	const ERHIPipeline Pipeline = GetPipeline();
-
-	if (Bypass())
-	{
-		// Stack allocate the transition
-		FMemStack& MemStack = FMemStack::Get();
-		FMemMark Mark(MemStack);
-		FRHITransition* Transition = new (MemStack.Alloc(FRHITransition::GetTotalAllocationSize(), FRHITransition::GetAlignment())) FRHITransition(Pipeline, Pipeline);
-		GDynamicRHI->RHICreateTransition(Transition, FRHITransitionCreateInfo(Pipeline, Pipeline, CreateFlags | ERHITransitionCreateFlags::NoSplit, Infos));
-
-		GetComputeContext().RHIBeginTransitions(MakeArrayView((const FRHITransition**)&Transition, 1));
-		GetComputeContext().RHIEndTransitions(MakeArrayView((const FRHITransition**)&Transition, 1));
-
-		// Manual release
-		GDynamicRHI->RHIReleaseTransition(Transition);
-		Transition->~FRHITransition();
-	}
-	else
-	{
-		// Allocate the transition in the command list
-		FRHITransition* Transition = new (Alloc(FRHITransition::GetTotalAllocationSize(), FRHITransition::GetAlignment())) FRHITransition(Pipeline, Pipeline);
-		GDynamicRHI->RHICreateTransition(Transition, FRHITransitionCreateInfo(Pipeline, Pipeline, CreateFlags | ERHITransitionCreateFlags::NoSplit, Infos));
-
-		ALLOC_COMMAND(FRHICommandResourceTransition)(Transition);
-	}
+	TransitionInternal(Infos, CreateFlags);
 
 	for (const FRHITransitionInfo& Info : Infos)
 	{
@@ -2183,6 +2161,46 @@ void FRHICommandListBase::InvalidBufferFatalError(const FRHIResourceCreateInfo&
 	UE_LOG(LogRHI, Fatal, TEXT("Attempt to create zero-sized buffer '%s', owner '%s', usage 0x%x, stride %u."), CreateInfo.DebugName ? CreateInfo.DebugName : TEXT("(nullptr)"), *CreateInfo.OwnerName.ToString(), (uint32)Usage, Stride);
 }
 
+void FRHICommandListBase::TransitionInternal(TConstArrayView<FRHITransitionInfo> Infos, ERHITransitionCreateFlags CreateFlags)
+{
+	const ERHIPipeline Pipeline = GetPipeline();
+	
+	if (Bypass())
+	{
+		// Without the FRHITransition we need to copy the transition info data because we could patch those data
+		TArray<FRHITransitionInfo, TSizedInlineAllocator<8, 32, FRHICmdListBaseArrayAllocator> > LocalTransitionInfo;
+		for (int32 i = 0; i < Infos.Num(); i++)
+		{
+			LocalTransitionInfo.Add(Infos[i]);
+
+			if (NeedTransitionsPatching())
+			{
+				PatchTransitionsInPlace((FRHITransitionInfo*)&LocalTransitionInfo[i], Pipeline, Pipeline);
+			}
+		}
+
+		// Stack allocate the transition
+		FMemStack& MemStack = FMemStack::Get();
+		FMemMark Mark(MemStack);
+		FRHITransition* Transition = new (MemStack.Alloc(FRHITransition::GetTotalAllocationSize(), FRHITransition::GetAlignment())) FRHITransition(Pipeline, Pipeline, 
+			EnumHasAnyFlags(CreateFlags, ERHITransitionCreateFlags::AllowDuringRenderPass) ? ETransitionFlag::AllowDuringRenderPass : ETransitionFlag::None);
+		GDynamicRHI->RHICreateTransition(Transition, FRHITransitionCreateInfo(Pipeline, Pipeline, CreateFlags | ERHITransitionCreateFlags::NoSplit | ERHITransitionCreateFlags::NoFence, MakeArrayView(LocalTransitionInfo)));
+
+		GetComputeContext().RHIBeginTransitions(MakeArrayView((const FRHITransition**)&Transition, 1));
+		GetComputeContext().RHIEndTransitions(MakeArrayView((const FRHITransition**)&Transition, 1));
+
+		// Manual release
+		GDynamicRHI->RHIReleaseTransition(Transition);
+		Transition->~FRHITransition();
+	}
+	else
+	{
+		FRHITransition* Transition = RHICreateTransition(FRHITransitionCreateInfo(Pipeline, Pipeline, CreateFlags | ERHITransitionCreateFlags::NoSplit | ERHITransitionCreateFlags::NoFence, Infos));
+		TrackResourceState(MakeArrayView((const FRHITransition**)&Transition, 1));
+		ALLOC_COMMAND(FRHICommandResourceTransition)(Transition);
+	}
+}
+
 FRayTracingShaderBindings UE::RHI::ConvertRayTracingShaderBindings(const FRHIBatchedShaderParameters& BatchedParameters)
 {
 	FRayTracingShaderBindings Result;
@@ -2281,3 +2299,179 @@ void FRHICommandList::RayTraceDispatchIndirect(
 	GlobalResourceBindings.Reset();
 }
 
+// track all transitions:
+// - if a BeforeAccess is Unknown add the transition to a list for patching before dispatch
+// - save each resource state to use at dispatch time
+void FRHICommandListBase::TrackResourceState(TArrayView<const FRHITransition*> Transitions)
+{
+	// If cmdlist tracking is off or RHI can handle Unknown resource, skip the tracking
+	if (NeedTransitionsPatching() == false)
+		return;
+
+	for (const FRHITransition* Transition : Transitions)
+	{
+		FRHITransitionCreateInfoData* TransitionCreateInfoData = Transition->CreationInfoData;
+
+		if (TransitionCreateInfoData == nullptr)
+		{
+			checkSlow(Transition->IsPlatformLevelCreated());
+			continue;
+		}
+
+		for (int32 i = 0; i < TransitionCreateInfoData->TransitionInfos.Num(); i++)
+		{
+			FRHITransitionInfo* TransitionInfo = (FRHITransitionInfo*)&TransitionCreateInfoData->TransitionInfos[i];
+
+			if ((TransitionInfo->Type == FRHITransitionInfo::EType::Buffer) && (GRHIGlobals.SupportImplicitBufferStateTracking))
+			{ 
+				// This RHI supports Buffers state tracking at platform level, so we don't need to track them
+				continue;
+			}
+
+			FRHIViewableResource* TransitionViewableResource = GetViewableResource(*TransitionInfo);
+			if (TransitionViewableResource == nullptr)
+			{
+				continue;
+			}
+
+			TRHIResourceState<FRHICmdListBaseArrayAllocator>* TrackedResourceState = TrackedResources.Find(TransitionViewableResource);
+			FRHISubresourceRange SubResourceIndex = *TransitionInfo;
+
+			// Both Unknown is not a supported case
+			check(!((TransitionInfo->AccessBefore == ERHIAccess::Unknown) && (TransitionInfo->AccessAfter == ERHIAccess::Unknown)));
+			// Discard to Unknown is not a supported case
+			check(!((TransitionInfo->AccessBefore == ERHIAccess::Discard) && (TransitionInfo->AccessAfter == ERHIAccess::Unknown)));
+
+			// Some RHI needs extra information on transition from Discard
+			const bool bBeforeIsDiscard = NeedsExtraTransitions() && (TransitionInfo->AccessBefore == ERHIAccess::Discard);
+
+			if ((TransitionInfo->AccessBefore == ERHIAccess::Unknown) || (TransitionInfo->AccessAfter == ERHIAccess::Unknown) || bBeforeIsDiscard)
+			{
+				ERHIAccess PreviousState = ERHIAccess::Unknown;
+				// Get the status from the previous state tracked
+				if (TrackedResourceState)
+				{
+					PreviousState = TrackedResourceState->GetSubresourceState(SubResourceIndex);
+				}
+
+				// Note: even if PreviuosState is a valid state we cannot write into TransitionInfo since we are recording the cmdlist
+				TransitionInfoToPatch.Push(FTransitionPatchInfo(TransitionInfo, PreviousState));
+			}
+	
+			// Some RHI will not transition to initial state on After = Discard when running on the Async queue
+			if (NeedsExtraTransitions())
+			{
+				if ((TransitionInfo->AccessAfter == ERHIAccess::Discard)
+					&& ((TransitionCreateInfoData->SrcPipelines == ERHIPipeline::AsyncCompute) || (TransitionCreateInfoData->SrcPipelines == ERHIPipeline::All))
+					&& (TransitionCreateInfoData->DstPipelines == ERHIPipeline::AsyncCompute))
+				{
+					TransitionInfo->Flags |= EResourceTransitionFlags::IgnoreAfterState;
+				}
+			}
+
+			const bool bNeedUpdateAfterState = (!((EnumHasAnyFlags(TransitionInfo->Flags, EResourceTransitionFlags::IgnoreAfterState)) || (TransitionInfo->AccessAfter == ERHIAccess::Unknown)));
+
+			if (bNeedUpdateAfterState)
+			{
+				// Update the TrackedResource State
+				if (TrackedResourceState)
+				{
+					TrackedResourceState->SetSubresourceState(SubResourceIndex, TransitionInfo->AccessAfter);
+				}
+				else
+				{
+					// Add a new element
+					TrackedResources.Add(TransitionViewableResource, TRHIResourceState<FRHICmdListBaseArrayAllocator>(SubResourceIndex, TransitionInfo->AccessAfter, TransitionViewableResource->GetTrackedSubResourceCount(), true));
+				}
+			}
+		}
+	}
+}
+
+void FRHICommandListBase::UpdateTransitionsAndResourceState()
+{
+	TRACE_CPUPROFILER_EVENT_SCOPE(UpdateTransitionsAndResourceState);
+
+	//If cmdlist tracking is off or RHI can handle Unknown resource, skip the tracking
+	if (NeedTransitionsPatching())
+	{
+		// Iterate over all TransitionInfo to patch for Unknown
+		for (int32 i = 0; i < TransitionInfoToPatch.Num(); i++)
+		{
+			FTransitionPatchInfo& PatchInfo = TransitionInfoToPatch[i];
+			FRHITransitionInfo* TransitionInfo = PatchInfo.TransitionInfo;
+			FRHISubresourceRange SubResourceIndex = *TransitionInfo;
+
+			const bool bBeforeIsDiscard = NeedsExtraTransitions() && (TransitionInfo->AccessBefore == ERHIAccess::Discard);
+
+			if (bBeforeIsDiscard)
+			{ 
+				if (PatchInfo.State != ERHIAccess::Unknown)
+				{
+					TransitionInfo->DiscardAccessBefore = PatchInfo.State;
+				}
+				else
+				{
+					ERHIAccess ResourceState = ERHIAccess::Unknown;
+
+					FRHIViewableResource* TransitionViewableResource = GetViewableResource(*TransitionInfo);
+
+					if (TransitionViewableResource)
+						ResourceState = TransitionViewableResource->TrackedAccess.GetSubresourceState(SubResourceIndex);
+
+					TransitionInfo->DiscardAccessBefore = ResourceState;
+				}
+			}
+			else if (TransitionInfo->AccessBefore == ERHIAccess::Unknown)
+			{
+				if (PatchInfo.State != ERHIAccess::Unknown)
+				{
+					TransitionInfo->AccessBefore = PatchInfo.State;
+				}
+				else
+				{
+					ERHIAccess ResourceState = ERHIAccess::Unknown;
+
+					FRHIViewableResource* TransitionViewableResource = GetViewableResource(*TransitionInfo);
+
+					if (TransitionViewableResource)
+						ResourceState = TransitionViewableResource->TrackedAccess.GetSubresourceState(SubResourceIndex);
+
+					TransitionInfo->AccessBefore = ResourceState;
+				}
+			}
+			else if (TransitionInfo->AccessAfter == ERHIAccess::Unknown)
+			{
+				if (PatchInfo.State != ERHIAccess::Unknown)
+				{
+					TransitionInfo->AccessAfter = PatchInfo.State;
+				}
+				else
+				{
+					ERHIAccess ResourceState = ERHIAccess::Unknown;
+
+					FRHIViewableResource* TransitionViewableResource = GetViewableResource(*TransitionInfo);
+
+					if (TransitionViewableResource)
+						ResourceState = TransitionViewableResource->TrackedAccess.GetSubresourceState(SubResourceIndex);
+
+					TransitionInfo->AccessAfter = ResourceState;
+				}
+			}
+
+			check(TransitionInfo->AccessBefore != ERHIAccess::Unknown);
+			check(TransitionInfo->AccessAfter != ERHIAccess::Unknown);
+		}
+
+		// Iterate over all tracked Resources to update their TrackedState with the latest state
+		for (TPair<FRHIViewableResource*, TRHIResourceState<FRHICmdListBaseArrayAllocator>>& Pair : TrackedResources)
+		{
+			FRHIViewableResource* TrackedResource = Pair.Key;
+			TRHIResourceState<FRHICmdListBaseArrayAllocator>* TrackedResourceState = &Pair.Value;
+			TrackedResource->TrackedAccess.Merge(*TrackedResourceState);
+		}
+
+		TransitionInfoToPatch.Empty();
+		TrackedResources.Empty();
+	}
+}
\ No newline at end of file
diff --git a/Engine/Source/Runtime/RHI/Private/RHIResources.cpp b/Engine/Source/Runtime/RHI/Private/RHIResources.cpp
--- a/Engine/Source/Runtime/RHI/Private/RHIResources.cpp
+++ b/Engine/Source/Runtime/RHI/Private/RHIResources.cpp
@@ -121,7 +121,7 @@ bool FRHIResource::Bypass()
 }
 
 FRHITexture::FRHITexture(const FRHITextureCreateDesc& InDesc)
-	: FRHIViewableResource(RRT_Texture, InDesc.InitialState)
+	: FRHIViewableResource(RRT_Texture, InDesc.InitialState, InDesc.GetSubresourceCount())
 #if ENABLE_RHI_VALIDATION
 	, RHIValidation::FTextureResource(InDesc)
 #endif
diff --git a/Engine/Source/Runtime/RHI/Private/RHIStrings.cpp b/Engine/Source/Runtime/RHI/Private/RHIStrings.cpp
--- a/Engine/Source/Runtime/RHI/Private/RHIStrings.cpp
+++ b/Engine/Source/Runtime/RHI/Private/RHIStrings.cpp
@@ -505,6 +505,7 @@ FString GetResourceTransitionFlagsName(EResourceTransitionFlags Flags)
 		case EResourceTransitionFlags::MaintainCompression: return TEXT("MaintainCompression");
 		case EResourceTransitionFlags::Discard:				return TEXT("Discard");
 		case EResourceTransitionFlags::Clear:				return TEXT("Clear");
+		case EResourceTransitionFlags::IgnoreAfterState:	return TEXT("IgnoreAfterState");
 		}
 	});
 }
diff --git a/Engine/Source/Runtime/RHI/Private/RHITransition.cpp b/Engine/Source/Runtime/RHI/Private/RHITransition.cpp
--- a/Engine/Source/Runtime/RHI/Private/RHITransition.cpp
+++ b/Engine/Source/Runtime/RHI/Private/RHITransition.cpp
@@ -21,12 +21,46 @@ FRHIViewableResource* GetViewableResource(const FRHITransitionInfo& Info)
 	return nullptr;
 }
 
+void FRHITransition::CleanupInfoData() const
+{
+	checkSlow(CreationInfoData != nullptr);
+	CreationInfoData->~FRHITransitionCreateInfoData();
+	FRHICmdListBaseLinearAllocator::Free(CreationInfoData);
+	CreationInfoData = nullptr;
+}
+
+void FRHITransition::CreatePlatformTransition() const
+{
+	checkSlow(GRHIGlobals.EnableHighLevelTransitionPatching);
+	if (IsPlatformLevelCreated() == false)
+	{
+		GDynamicRHI->RHICreateTransition(const_cast<FRHITransition*>(this), CreationInfoData->GetCreationInfo());
+		EnumAddFlags(Flags, ETransitionFlag::PlatformLevelCreated);
+		CleanupInfoData();
+	}
+}
+
+void FRHITransition::CreatePlatformTransitionThreadSafe() const
+{
+	checkSlow(GRHIGlobals.EnableHighLevelTransitionPatching);
+	
+	UE::TUniqueLock Lock(TransitionCreationMutex);
+
+	if (IsPlatformLevelCreated() == false)
+	{
+		GDynamicRHI->RHICreateTransition(const_cast<FRHITransition*>(this), CreationInfoData->GetCreationInfo());
+		EnumAddFlags(Flags, ETransitionFlag::PlatformLevelCreated);
+		CleanupInfoData();
+	}
+}
+
 void FRHITransition::Cleanup() const
 {
 	FRHITransition* Transition = const_cast<FRHITransition*>(this);
 	RHIReleaseTransition(Transition);
+	checkSlow(CreationInfoData == nullptr);
 
 	// Explicit destruction of the transition.
 	Transition->~FRHITransition();
-	FConcurrentLinearAllocator::Free(Transition);
-}
+	FRHICmdListBaseLinearAllocator::Free(Transition);
+}
\ No newline at end of file
diff --git a/Engine/Source/Runtime/RHI/Private/RHIUtilities.cpp b/Engine/Source/Runtime/RHI/Private/RHIUtilities.cpp
--- a/Engine/Source/Runtime/RHI/Private/RHIUtilities.cpp
+++ b/Engine/Source/Runtime/RHI/Private/RHIUtilities.cpp
@@ -668,6 +668,8 @@ ERHIAccess RHIGetDefaultResourceState(ETextureCreateFlags InUsage, bool bInHasIn
 			ResourceState = ERHIAccess::ShadingRateSource;
 		}
 	}
+	
+	check(ResourceState != ERHIAccess::Unknown);
 
 	return ResourceState;
 }
diff --git a/Engine/Source/Runtime/RHI/Private/RHIValidation.cpp b/Engine/Source/Runtime/RHI/Private/RHIValidation.cpp
--- a/Engine/Source/Runtime/RHI/Private/RHIValidation.cpp
+++ b/Engine/Source/Runtime/RHI/Private/RHIValidation.cpp
@@ -21,6 +21,12 @@ bool GRHIValidationEnabled = false;
 
 bool GRHIValidateBufferSourceCopy = true;
 
+bool GRHIValidationPrintHumanReadableCallStack = false;
+
+// Define the number of stack frames to capture
+const int32 NumStackFrames = 30;
+const uint32 IgnoreStackCount = 2; // Ignore the call to the function itself and the log
+
 // When set to 1, callstack for each uniform buffer allocation will be tracked 
 // (slow and leaks memory, but can be handy to find the location where an invalid
 // allocation has been made)
@@ -45,6 +51,7 @@ namespace RHIValidation
 	// Returns an array of resource names parsed from the "-RHIValidationLog" command line switch.
 	// RHI validation logging is automatically enabled for resources whose debug names match those in this list.
 	// Multiple values are comma separated, e.g. -RHIValidationLog="SceneDepthZ,GBufferA"
+	// Use the additional -RHIValidationLogStack arg to enable printing the resolved symbols of the callstack in the log
 	static TArray<FString> const& GetAutoLogResourceNames()
 	{
 		struct FInit
@@ -67,6 +74,8 @@ namespace RHIValidation
 					ResourceNames.TrimStartAndEndInline();
 					Strings.Add(ResourceNames);
 				}
+				
+				GRHIValidationPrintHumanReadableCallStack = FParse::Param(FCommandLine::Get(), TEXT("-RHIValidationLogStack"));
 			}
 		} static Init;
 
@@ -671,7 +680,7 @@ void FValidationRHI::RHICreateTransition(FRHITransition* Transition, const FRHIT
 		FState NextState = FState(Info.AccessAfter, DstPipelines);
 
 		BeginOps.Emplace(FOperation::BeginTransitionResource(Identity, PreviousState, NextState, Info.Flags, CreateInfo.Flags, nullptr));
-		EndOps  .Emplace(FOperation::EndTransitionResource(Identity, PreviousState, NextState, nullptr));
+		EndOps.Emplace(FOperation::EndTransitionResource(Identity, PreviousState, NextState, Info.Flags, nullptr));
 	}
 
 	if (bDoTrace)
@@ -1237,16 +1246,68 @@ namespace RHIValidation
 			*GetRHIAccessName(TrackedState.Access));
 	}
 
+	static FString ResolveAndFormatCallstack (uint64* Trace)
+	{
+		FString Callstack;
+		ANSICHAR Buffer[1024];
+
+		for (uint32 Idx = IgnoreStackCount; Idx < NumStackFrames; Idx++)
+		{
+			Buffer[0] = '\0'; // Clear the buffer
+			const ANSICHAR* TrimmedBuffer = Buffer;
+
+			// Resolve the program counter to a human-readable string
+			const bool bFoundSymbol = FPlatformStackWalk::ProgramCounterToHumanReadableString(
+				Idx,
+				Trace[Idx],
+				Buffer,
+				sizeof(Buffer)
+			);
+
+			if (bFoundSymbol)
+			{
+				// Find the "!" delimiter to trim the module and address
+				const ANSICHAR* BufferAfterModuleAndAddress = FCStringAnsi::Strstr(Buffer, "!");
+				if (BufferAfterModuleAndAddress)
+				{
+					TrimmedBuffer = BufferAfterModuleAndAddress + 1;
+				}
+			}
+
+			if (FCStringAnsi::Strstr(Buffer, "UnknownFunction") == nullptr)
+			{
+				// Append the trimmed symbol information to the call stack string
+				Callstack += FString(TrimmedBuffer) + TEXT("\r\n");
+			}
+		}
+
+		return Callstack;
+	}
+
+
 	static inline FString GetReasonString_BeginBacktrace(void* CreateTrace, void* BeginTrace)
 	{
 		if (CreateTrace || BeginTrace)
 		{
-			return FString::Printf(
+			if (GRHIValidationPrintHumanReadableCallStack)
+			{
+				return FString::Printf(
+					TEXT("    --- Callstack backtraces for the transition which has not been completed:\n")
+					TEXT("        RHICreateTransition: %s\n")
+					TEXT("        RHIBeginTransitions: %s\n"),
+					*ResolveAndFormatCallstack((uint64*)CreateTrace),
+					*ResolveAndFormatCallstack((uint64*)BeginTrace));
+
+			}
+			else
+			{
+				return FString::Printf(
 				TEXT("    --- Callstack backtraces for the transition which has not been completed (resolve in the Watch window):\n")
 				TEXT("        RHICreateTransition: (void**)0x%p,32\n")
 				TEXT("        RHIBeginTransitions: (void**)0x%p,32\n"),
 				CreateTrace,
 				BeginTrace);
+			}
 		}
 		else
 		{
@@ -1258,12 +1319,25 @@ namespace RHIValidation
 	{
 		if (Trace)
 		{
-			return FString::Printf(
-				TEXT("    --- Callstack backtrace for %s operation (resolve in the Watch window):\n")
-				TEXT("        %s: (void**)0x%p,32\n"),
-				OperationPrefix,
-				TracePrefix,
-				Trace);
+			if (GRHIValidationPrintHumanReadableCallStack)
+			{
+				return FString::Printf(
+					TEXT("    --- Callstack backtrace for %s operation:\n")
+					TEXT("        %s: %s\n"),
+					OperationPrefix,
+					TracePrefix,
+					*ResolveAndFormatCallstack((uint64*)Trace));
+			}
+			else
+			{
+				return FString::Printf(
+					TEXT("    --- Callstack backtrace for %s operation (resolve in the Watch window):\n")
+					TEXT("        %s: (void**)0x%p,32\n"),
+					OperationPrefix,
+					TracePrefix,
+					Trace);
+
+			}
 		}
 		else
 		{
@@ -1615,23 +1689,49 @@ namespace RHIValidation
 
 		if (CreateTrace)
 		{
-			FPlatformMisc::LowLevelOutputDebugStringf(TEXT("\n%s: Type: %s, %s, CreateTrace: 0x%p, %sTrace: 0x%p, %s\n"),
-				*GetResourceDebugName(Resource, SubresourceIndex),
-				Type,
-				LogStr,
-				CreateTrace,
-				TracePrefix,
-				Trace,
-				*BreadcrumbMessage);
+			if (GRHIValidationPrintHumanReadableCallStack)
+			{
+				FPlatformMisc::LowLevelOutputDebugStringf(TEXT("\n%s: Type: %s, %s, \nCreateTrace: %s\n, %sTrace: %s\n, %s\n"),
+					*GetResourceDebugName(Resource, SubresourceIndex),
+					Type,
+					LogStr,
+					*ResolveAndFormatCallstack((uint64*)CreateTrace),
+					TracePrefix,
+					*ResolveAndFormatCallstack((uint64*)Trace),
+					*BreadcrumbMessage);
+			}
+			else
+			{
+				FPlatformMisc::LowLevelOutputDebugStringf(TEXT("\n%s: Type: %s, %s, CreateTrace: 0x%p, %sTrace: 0x%p, %s\n"),
+					*GetResourceDebugName(Resource, SubresourceIndex),
+					Type,
+					LogStr,
+					CreateTrace,
+					TracePrefix,
+					Trace,
+					*BreadcrumbMessage);
+			}
 		}
 		else
 		{
-			FPlatformMisc::LowLevelOutputDebugStringf(TEXT("\n%s: Type: %s, %s, Trace: 0x%p, %s\n"),
-				*GetResourceDebugName(Resource, SubresourceIndex),
-				Type,
-				LogStr,
-				Trace,
-				*BreadcrumbMessage);
+			if (GRHIValidationPrintHumanReadableCallStack)
+			{
+				FPlatformMisc::LowLevelOutputDebugStringf(TEXT("\n%s: Type: %s, %s, \nTrace: %s\n, %s\n"),
+					*GetResourceDebugName(Resource, SubresourceIndex),
+					Type,
+					LogStr,
+					*ResolveAndFormatCallstack((uint64*)Trace),
+					*BreadcrumbMessage);
+			}
+			else
+			{
+				FPlatformMisc::LowLevelOutputDebugStringf(TEXT("\n%s: Type: %s, %s, Trace: 0x%p, %s\n"),
+					*GetResourceDebugName(Resource, SubresourceIndex),
+					Type,
+					LogStr,
+					Trace,
+					*BreadcrumbMessage);
+			}
 		}
 
 		return Trace;
@@ -1640,7 +1740,7 @@ namespace RHIValidation
 	void FTransientState::Acquire(FResource* Resource, void* CreateTrace)
 	{
 		RHI_VALIDATION_CHECK(bTransient, *GetReasonString_AcquireNonTransient(Resource));
-		RHI_VALIDATION_CHECK(Status == EStatus::None, *GetReasonString_DuplicateAcquireTransient(Resource, AcquireBacktrace, CreateTrace));
+		RHI_VALIDATION_CHECK(Status != EStatus::Acquired, *GetReasonString_DuplicateAcquireTransient(Resource, AcquireBacktrace, CreateTrace));
 		Status = EStatus::Acquired;
 
 		if (!AcquireBacktrace)
@@ -1685,7 +1785,7 @@ namespace RHIValidation
 
 		// Acquire should validate whether ResourceAfter is transient. We assume it is here.
 		RHI_VALIDATION_CHECK(TransientStateBefore.bTransient, *GetReasonString_AliasingOverlapNonTransient(ResourceBefore, ResourceAfter));
-		RHI_VALIDATION_CHECK(TransientStateBefore.IsDiscarded(), *GetReasonString_AliasingOverlapNonDiscarded(ResourceBefore, ResourceAfter, CreateTrace));
+		//RHI_VALIDATION_CHECK(TransientStateBefore.IsDiscarded(), *GetReasonString_AliasingOverlapNonDiscarded(ResourceBefore, ResourceAfter, CreateTrace));
 
 		if (ResourceBefore->LoggingMode != ELoggingMode::None)
 		{
@@ -1781,16 +1881,19 @@ namespace RHIValidation
 				return Previous == Next;
 			};
 
-			if (CurrentStateFromRHI.Access == ERHIAccess::Unknown)
-			{
-				RHI_VALIDATION_CHECK(Resource->TrackedAccess == State.Previous.Access && HasMatchingPipelines(State.Previous.Pipelines, CurrentStateFromRHI.Pipelines),
-					*GetReasonString_IncorrectPreviousTrackedState(Resource, SubresourceIndex, State.Previous, CurrentStateFromRHI.Pipelines));
-			}
-			else
+			if (EnumHasAnyFlags(NewFlags, EResourceTransitionFlags::IgnoreAfterState) == false)
 			{
-				// Check the current RHI state passed in matches the tracked state for the resource.
-				RHI_VALIDATION_CHECK(CurrentStateFromRHI.Access == State.Previous.Access && HasMatchingPipelines(State.Previous.Pipelines, CurrentStateFromRHI.Pipelines),
-					*GetReasonString_IncorrectPreviousExplicitState(Resource, SubresourceIndex, State.Previous, CurrentStateFromRHI));
+				if (CurrentStateFromRHI.Access == ERHIAccess::Unknown)
+				{
+					RHI_VALIDATION_CHECK(Resource->TrackedAccess == State.Previous.Access && HasMatchingPipelines(State.Previous.Pipelines, CurrentStateFromRHI.Pipelines),
+						*GetReasonString_IncorrectPreviousTrackedState(Resource, SubresourceIndex, State.Previous, CurrentStateFromRHI.Pipelines));
+				}
+				else
+				{
+					// Check the current RHI state passed in matches the tracked state for the resource.
+					RHI_VALIDATION_CHECK(CurrentStateFromRHI.Access == State.Previous.Access && HasMatchingPipelines(State.Previous.Pipelines, CurrentStateFromRHI.Pipelines),
+						*GetReasonString_IncorrectPreviousExplicitState(Resource, SubresourceIndex, State.Previous, CurrentStateFromRHI));
+				}
 			}
 		}
 
@@ -1798,24 +1901,29 @@ namespace RHIValidation
 		// @todo: this check is not particularly useful at the moment, as there are many unnecessary resource transitions.
 		//RHI_VALIDATION_CHECK(CurrentState != TargetState, *GetReasonString_UnnecessaryTransition(Resource, SubresourceIndex, CurrentState));
 
-		// Update the tracked state once all pipes have begun.
-		State.Previous = TargetState;
-		State.Current = TargetState;
-		State.Flags = NewFlags;
-		State.CreateTransitionBacktrace = CreateTrace;
-		State.BeginTransitionBacktrace = BeginTrace;
-		State.bUsedWithAllUAVsOverlap = false;
-		State.bUsedWithExplicitUAVsOverlap = false;
-		State.bTransitioning = true;
-
-		// Replicate the state to other pipes that are not part of the begin pipe mask.
-		for (ERHIPipeline OtherPipeline : MakeFlagsRange(ERHIPipeline::All & ~CurrentStateFromRHI.Pipelines))
+		if (EnumHasAnyFlags(NewFlags, EResourceTransitionFlags::IgnoreAfterState) == false)
 		{
-			States[OtherPipeline] = State;
+			State.bTransitioning = true;
+			State.Flags = NewFlags;
+			State.BeginTransitionBacktrace = BeginTrace;
+			// Update the tracked state once all pipes have begun.
+			State.Previous = TargetState;
+			State.Current = TargetState;
+			State.CreateTransitionBacktrace = CreateTrace;
+			State.bUsedWithAllUAVsOverlap = false;
+			State.bUsedWithExplicitUAVsOverlap = false;
+
+
+			// Replicate the state to other pipes that are not part of the begin pipe mask.
+			for (ERHIPipeline OtherPipeline : MakeFlagsRange(ERHIPipeline::All & ~CurrentStateFromRHI.Pipelines))
+			{
+				States[OtherPipeline] = State;
+			}
 		}
+		
 	}
 
-	void FSubresourceState::EndTransition(FResource* Resource, FSubresourceIndex const& SubresourceIndex, const FState& CurrentStateFromRHI, const FState& TargetState, ERHIPipeline ExecutingPipeline, uint64 ExecutingPipelineFenceValue, void* CreateTrace)
+	void FSubresourceState::EndTransition(FResource* Resource, FSubresourceIndex const& SubresourceIndex, const FState& CurrentStateFromRHI, const FState& TargetState, EResourceTransitionFlags NewFlags, ERHIPipeline ExecutingPipeline, uint64 ExecutingPipelineFenceValue, void* CreateTrace)
 	{
 		if (Resource->LoggingMode != ELoggingMode::None
 #if LOG_UNNAMED_RESOURCES
@@ -1830,22 +1938,25 @@ namespace RHIValidation
 			));
 		}
 
-		FPipelineState& State = States[ExecutingPipeline];
+		if (EnumHasAnyFlags(NewFlags, EResourceTransitionFlags::IgnoreAfterState) == false)
+		{
+			FPipelineState& State = States[ExecutingPipeline];
 
-		// Check that we aren't ending a transition that never began.
-		RHI_VALIDATION_CHECK(State.bTransitioning, TEXT("Unsolicited resource end transition call."));
-		State.bTransitioning = false;
-		State.BeginTransitionBacktrace = nullptr;
+			// Check that we aren't ending a transition that never began.
+			RHI_VALIDATION_CHECK(State.bTransitioning, TEXT("Unsolicited resource end transition call."));
+			State.bTransitioning = false;
+			State.BeginTransitionBacktrace = nullptr;
 
-		// Check that the end matches the begin.
-		RHI_VALIDATION_CHECK(TargetState == State.Current, *GetReasonString_MismatchedEndTransition(Resource, SubresourceIndex, State.Current, TargetState));
+			// Check that the end matches the begin.
+			RHI_VALIDATION_CHECK(TargetState == State.Current, *GetReasonString_MismatchedEndTransition(Resource, SubresourceIndex, State.Current, TargetState));
 
-		// Replicate the state to other pipes that are not part of the end pipe mask.
-		for (ERHIPipeline OtherPipeline : MakeFlagsRange(ERHIPipeline::All))
-		{
-			if (!EnumHasAnyFlags(TargetState.Pipelines, OtherPipeline))
+			// Replicate the state to other pipes that are not part of the end pipe mask.
+			for (ERHIPipeline OtherPipeline : MakeFlagsRange(ERHIPipeline::All))
 			{
-				States[OtherPipeline] = State;
+				if (!EnumHasAnyFlags(TargetState.Pipelines, OtherPipeline))
+				{
+					States[OtherPipeline] = State;
+				}
 			}
 		}
 
@@ -2145,6 +2256,7 @@ namespace RHIValidation
 					SubresourceIndex,
 					Data_EndTransition.PreviousState,
 					Data_EndTransition.NextState,
+					Data_EndTransition.Flags,
 					Queue.Pipeline,
 					Queue.FenceValue,
 					Data_EndTransition.CreateBacktrace);
@@ -2369,9 +2481,8 @@ namespace RHIValidation
 	void* CaptureBacktrace()
 	{
 		// Back traces will leak. Don't leave this turned on.
-		const uint32 MaxDepth = 32;
-		uint64* Backtrace = new uint64[MaxDepth];
-		FPlatformStackWalk::CaptureStackBackTrace(Backtrace, MaxDepth);
+		uint64* Backtrace = new uint64[NumStackFrames];
+		FPlatformStackWalk::CaptureStackBackTrace(Backtrace, NumStackFrames);
 
 		return Backtrace;
 	}
diff --git a/Engine/Source/Runtime/RHI/Public/DynamicRHI.h b/Engine/Source/Runtime/RHI/Public/DynamicRHI.h
--- a/Engine/Source/Runtime/RHI/Public/DynamicRHI.h
+++ b/Engine/Source/Runtime/RHI/Public/DynamicRHI.h
@@ -1426,13 +1426,7 @@ FORCEINLINE class IRHICommandContext* RHIGetDefaultContext()
 
 RHI_API FRenderQueryPoolRHIRef RHICreateRenderQueryPool(ERenderQueryType QueryType, uint32 NumQueries = UINT32_MAX);
 
-FORCEINLINE const FRHITransition* RHICreateTransition(const FRHITransitionCreateInfo& CreateInfo)
-{
-	// Placement create the transition on the heap.
-	FRHITransition* Transition = new (FConcurrentLinearAllocator::Malloc(FRHITransition::GetTotalAllocationSize(), (uint32)FRHITransition::GetAlignment())) FRHITransition(CreateInfo.SrcPipelines, CreateInfo.DstPipelines);
-	GDynamicRHI->RHICreateTransition(Transition, CreateInfo);
-	return Transition;
-}
+RHI_API FRHITransition* RHICreateTransition(const FRHITransitionCreateInfo& CreateInfo);
 
 FORCEINLINE void RHIReleaseTransition(FRHITransition* Transition)
 {
diff --git a/Engine/Source/Runtime/RHI/Public/RHI.h b/Engine/Source/Runtime/RHI/Public/RHI.h
--- a/Engine/Source/Runtime/RHI/Public/RHI.h
+++ b/Engine/Source/Runtime/RHI/Public/RHI.h
@@ -26,9 +26,6 @@ class FResourceArrayInterface;
 class FResourceBulkDataInterface;
 class FRHICommandList;
 
-/** RHI Logging. */
-RHI_API DECLARE_LOG_CATEGORY_EXTERN(LogRHI,Log,VeryVerbose);
-
 /**
  * RHI configuration settings.
  */
diff --git a/Engine/Source/Runtime/RHI/Public/RHIAllocators.h b/Engine/Source/Runtime/RHI/Public/RHIAllocators.h
new file mode 100644
--- /dev/null
+++ b/Engine/Source/Runtime/RHI/Public/RHIAllocators.h
@@ -0,0 +1,25 @@
+// Copyright Epic Games, Inc. All Rights Reserved.
+
+/*=============================================================================
+	RHIAllocators.h: Render Hardware Interface allocators
+=============================================================================*/
+
+#pragma once
+
+#include "Experimental/ConcurrentLinearAllocator.h"
+
+// Same of FDefaultBlockAllocationTag but custom TagName
+struct FCommandListBaseBlockAllocationTag
+{
+	static constexpr uint32 BlockSize = 64 * 1024;		// Blocksize used to allocate from
+	static constexpr bool AllowOversizedBlocks = true;  // The allocator supports oversized Blocks and will store them in a separate Block with counter 1
+	static constexpr bool RequiresAccurateSize = true;  // GetAllocationSize returning the accurate size of the allocation otherwise it could be relaxed to return the size to the end of the Block
+	static constexpr bool InlineBlockAllocation = false;  // Inline or Noinline the BlockAllocation which can have an impact on Performance
+	static constexpr const char* TagName = "RHICommandListBaseAllocator";
+
+	using Allocator = TBlockAllocationLockFreeCache<BlockSize, FOsAllocator>;
+};
+
+using FRHICmdListBaseArrayAllocator = TConcurrentLinearArrayAllocator<FCommandListBaseBlockAllocationTag>;
+using FRHICmdListBaseSetAllocator = TConcurrentLinearSetAllocator<FCommandListBaseBlockAllocationTag>;
+using FRHICmdListBaseLinearAllocator = TConcurrentLinearAllocator<FCommandListBaseBlockAllocationTag>;
\ No newline at end of file
diff --git a/Engine/Source/Runtime/RHI/Public/RHICommandList.h b/Engine/Source/Runtime/RHI/Public/RHICommandList.h
--- a/Engine/Source/Runtime/RHI/Public/RHICommandList.h
+++ b/Engine/Source/Runtime/RHI/Public/RHICommandList.h
@@ -30,6 +30,7 @@
 #include "ProfilingDebugging/CpuProfilerTrace.h"
 #include "RHIBreadcrumbs.h"
 #include "RHIGlobals.h"
+#include "RHIAllocators.h"
 #include "RHIShaderParameters.h"
 #include "RHITextureReference.h"
 #include "RHIResourceReplace.h"
@@ -74,6 +75,7 @@ class FComputePipelineState;
 class FGraphicsPipelineState;
 class FRayTracingPipelineState;
 class FWorkGraphPipelineState;
+struct FRHICommandSetTrackedAccess;
 
 DECLARE_STATS_GROUP(TEXT("RHICmdList"), STATGROUP_RHICMDLIST, STATCAT_Advanced);
 
@@ -706,6 +708,12 @@ private:
 	RHI_API void InvalidBufferFatalError(const FRHIResourceCreateInfo& CreateInfo, EBufferUsageFlags Usage, uint32 Stride);
 
 public:
+	RHI_API void TransitionInternal(TConstArrayView<FRHITransitionInfo> Infos, ERHITransitionCreateFlags CreateFlags = ERHITransitionCreateFlags::None);
+	FORCEINLINE_DEBUGGABLE void TransitionInternal(const FRHITransitionInfo& Info, ERHITransitionCreateFlags CreateFlags = ERHITransitionCreateFlags::None)
+	{
+		TransitionInternal(MakeArrayView(&Info, 1), CreateFlags);
+	}
+
 	RHI_API ERHIPipeline SwitchPipeline(ERHIPipeline Pipeline);
 
 	FORCEINLINE FRHIGPUMask GetGPUMask() const { return PersistentState.CurrentGPUMask; }
@@ -846,9 +854,16 @@ public:
 			FRHITextureCreateDesc NewCreateDesc(CreateDesc);
 			NewCreateDesc.SetInitialState(RHIGetDefaultResourceState(CreateDesc.Flags, CreateDesc.BulkData != nullptr));
 
+#if RHI_TRANSITIONS_HL_DEBUG
+				UE_LOG(LogRHI, Display, TEXT("Resource: %s Creation State %s"), CreateDesc.DebugName, *GetRHIAccessName(NewCreateDesc.InitialState));
+#endif
+
 			return GDynamicRHI->RHICreateTexture(*this, NewCreateDesc);
 		}
 
+#if RHI_TRANSITIONS_HL_DEBUG
+		UE_LOG(LogRHI, Display, TEXT("Resource: %s Creation State %s"), CreateDesc.DebugName, *GetRHIAccessName(CreateDesc.InitialState));
+#endif
 		return GDynamicRHI->RHICreateTexture(*this, CreateDesc);
 	}
 
@@ -1128,6 +1143,12 @@ public:
 		return GRHIGlobals.NeedsShaderUnbinds;
 	}
 
+	// Returns true if the RHI needs extra transitions
+	bool NeedsExtraTransitions() const
+	{
+		return GRHIGlobals.NeedExtraTransitions;
+	}
+
 	FRHIBatchedShaderParametersAllocator* CreateBatchedShaderParameterAllocator(ERHIBatchedShaderParameterAllocatorPageSize PageSize)
 	{
 		return new (MemManager) FRHIBatchedShaderParametersAllocator(ShaderParameterState.AllocatorsRoot, *this, PageSize);
@@ -1336,6 +1357,7 @@ protected:
 		uint8 bInsideOcclusionQueryBatch		: 1;
 		uint8 bRecursive						: 1;
 		uint8 bImmediate						: 1;
+		uint8 bAllowResourceStateTracking		: 1;
 
 		FRHIGPUMask CurrentGPUMask;
 		FRHIGPUMask InitialGPUMask;
@@ -1355,13 +1377,14 @@ protected:
 		TStaticArray<void*, MAX_NUM_GPUS> QueryBatchData_Timestamp { InPlace, nullptr };
 		TStaticArray<void*, MAX_NUM_GPUS> QueryBatchData_Occlusion { InPlace, nullptr };
 
-		FPersistentState(FRHIGPUMask InInitialGPUMask, bool bInImmediate = false)
+		FPersistentState(FRHIGPUMask InInitialGPUMask, bool bInImmediate = false, bool bTrackResources = true)
 			: bHasFragmentDensityAttachment(0)
 			, bInsideRenderPass(0)
 			, bInsideComputePass(0)
 			, bInsideOcclusionQueryBatch(0)
 			, bRecursive(0)
 			, bImmediate(bInImmediate)
+			, bAllowResourceStateTracking(bTrackResources)
 			, CurrentGPUMask(InInitialGPUMask)
 			, InitialGPUMask(InInitialGPUMask)
 		{}
@@ -1411,12 +1434,15 @@ public:
 		}
 	}
 
+	FORCEINLINE_DEBUGGABLE bool NeedTransitionsPatching() { return (PersistentState.bAllowResourceStateTracking) && (GRHIGlobals.EnableHighLevelTransitionPatching); }
+
 private:
 	FRHICommandListBase(FPersistentState const& InPersistentState);
 
 	// Replays recorded commands. Used internally, do not call directly.
 	RHI_EXECUTE_API void Execute();
 
+	friend class FRHIScopedResourceBarrier;
 	friend class FRHICommandListExecutor;
 	friend class FRHICommandListIterator;
 	friend class FRHICommandListScopedFlushAndExecute;
@@ -1432,6 +1458,32 @@ private:
 #if WITH_RHI_BREADCRUMBS
 	friend bool IRHIComputeContext::ShouldEmitBreadcrumbs() const;
 #endif
+
+	struct FTransitionPatchInfo
+	{
+		FTransitionPatchInfo(FRHITransitionInfo* InTransitionInfo, ERHIAccess InState) :
+			TransitionInfo(InTransitionInfo),
+			State(InState) {}
+
+		FRHITransitionInfo* TransitionInfo;
+		ERHIAccess			State;
+	};
+
+	// Using FRHICmdListBaseArrayAllocator since elements Lifetime for all the following containers are 1 frame, between Record and Submit
+	// Track every resource state change happening during the command list Record
+	TMap<FRHIViewableResource*, TRHIResourceState<FRHICmdListBaseArrayAllocator>, FRHICmdListBaseSetAllocator> TrackedResources;
+	// Transitions with Unknown or Discard (when needed) that will patched at Dispatch time
+	TArray<FTransitionPatchInfo, FRHICmdListBaseArrayAllocator> TransitionInfoToPatch;
+	
+	// This will scan the array and keep track of transitions that needs patching later at Dispatch time
+	RHI_API void TrackResourceState(TArrayView<const FRHITransition*> Transitions);
+	
+	// Called at Dispatch time, in order of submission, it's when the Unknown and Discard transitions can be patched with the real state
+	void UpdateTransitionsAndResourceState();
+	
+	// Bypass needs to patch Transitions on the spot
+	inline void PatchTransitionsInPlace(FRHITransitionInfo* TransitionInfo, const ERHIPipeline SrcPipeline, const ERHIPipeline DestPipeline);
+	inline void PatchTransitionsInPlace(TArrayView<const FRHITransition*> Transitions);
 };
 
 #if WITH_RHI_BREADCRUMBS
@@ -2833,9 +2885,22 @@ public:
 		extern RHI_API TAutoConsoleVariable<int32> GProfileGPUTransitions;
 		RHI_BREADCRUMB_EVENT_CONDITIONAL(*this, GProfileGPUTransitions.GetValueOnAnyThread() != 0, "RHIBeginTransitions");
 #endif
-
 		if (Bypass())
 		{
+			if (NeedTransitionsPatching())
+			{
+				PatchTransitionsInPlace(Transitions);
+			}
+
+			//Create the Platform Transition for each FRHITransition
+			if (GRHIGlobals.EnableHighLevelTransitionPatching)
+			{
+				for (const FRHITransition* Transition : Transitions)
+				{
+					Transition->CreatePlatformTransition();
+				}
+			}
+
 			GetComputeContext().RHIBeginTransitions(Transitions);
 
 			for (const FRHITransition* Transition : Transitions)
@@ -2845,10 +2910,10 @@ public:
 		}
 		else
 		{
+			TrackResourceState(Transitions);
 			// Copy the transition array into the command list
 			FRHITransition** DstTransitionArray = (FRHITransition**)Alloc(sizeof(FRHITransition*) * Transitions.Num(), alignof(FRHITransition*));
 			FMemory::Memcpy(DstTransitionArray, Transitions.GetData(), sizeof(FRHITransition*) * Transitions.Num());
-
 			ALLOC_COMMAND(FRHICommandBeginTransitions)(MakeArrayView((const FRHITransition**)DstTransitionArray, Transitions.Num()));
 		}
 	}
@@ -2859,9 +2924,15 @@ public:
 		extern RHI_API TAutoConsoleVariable<int32> GProfileGPUTransitions;
 		RHI_BREADCRUMB_EVENT_CONDITIONAL(*this, GProfileGPUTransitions.GetValueOnAnyThread() != 0, "RHIEndTransitions");
 #endif
-
 		if (Bypass())
 		{
+
+#if UE_BUILD_DEBUG
+			for (const FRHITransition* Transition : Transitions)
+			{
+				checkSlow(Transition->IsPlatformLevelCreated());
+			}
+#endif
 			GetComputeContext().RHIEndTransitions(Transitions);
 
 			for (const FRHITransition* Transition : Transitions)
@@ -2904,6 +2975,11 @@ public:
 
 	FORCEINLINE_DEBUGGABLE void SetTrackedAccess(TArrayView<const FRHITrackedAccessInfo> Infos)
 	{
+		if (NeedTransitionsPatching())
+		{
+			return;
+		}
+
 		bUsesSetTrackedAccess = true;
 
 		if (Bypass())
@@ -3280,6 +3356,40 @@ public:
 	FORCEINLINE_DEBUGGABLE void TransferResources(TConstArrayView<FTransferResourceParams> Params)
 	{
 #if WITH_MGPU
+		FRHIGPUMask PrevGPUMask = GetGPUMask();
+
+		if (NeedsExtraTransitions())
+		{
+			for (const FTransferResourceParams& Param : Params)
+			{
+				FRHIGPUMask GPUMaskSrc = FRHIGPUMask::FromIndex(Param.SrcGPUIndex);
+				FRHIGPUMask GPUMaskDest = FRHIGPUMask::FromIndex(Param.DestGPUIndex);
+
+				if (Param.Texture)
+				{
+					SetGPUMask(GPUMaskSrc);
+					TransitionInternal(FRHITransitionInfo(Param.Texture.GetReference(), ERHIAccess::Unknown, ERHIAccess::CopySrc, EResourceTransitionFlags::IgnoreAfterState),
+						ERHITransitionCreateFlags::AllowDuringRenderPass);
+
+					SetGPUMask(GPUMaskDest);
+					TransitionInternal(FRHITransitionInfo(Param.Texture.GetReference(), ERHIAccess::Unknown, ERHIAccess::CopyDest, EResourceTransitionFlags::IgnoreAfterState),
+						ERHITransitionCreateFlags::AllowDuringRenderPass);
+
+				}
+				else
+				{
+					SetGPUMask(GPUMaskSrc);
+					TransitionInternal(FRHITransitionInfo(Param.Buffer.GetReference(), ERHIAccess::Unknown, ERHIAccess::CopySrc, EResourceTransitionFlags::IgnoreAfterState),
+						ERHITransitionCreateFlags::AllowDuringRenderPass);
+					SetGPUMask(GPUMaskDest);
+					TransitionInternal(FRHITransitionInfo(Param.Buffer.GetReference(), ERHIAccess::Unknown, ERHIAccess::CopyDest, EResourceTransitionFlags::IgnoreAfterState),
+						ERHITransitionCreateFlags::AllowDuringRenderPass);
+				}
+			}
+		}
+
+		SetGPUMask(PrevGPUMask);
+		
 		if (Bypass())
 		{
 			GetComputeContext().RHITransferResources(Params);
@@ -3288,6 +3398,40 @@ public:
 		{
 			ALLOC_COMMAND(FRHICommandTransferResources)(AllocArray(Params));
 		}
+
+		if (NeedsExtraTransitions())
+		{
+			for (const FTransferResourceParams& Param : Params)
+			{
+				FRHIGPUMask GPUMaskSrc = FRHIGPUMask::FromIndex(Param.SrcGPUIndex);
+				FRHIGPUMask GPUMaskDest = FRHIGPUMask::FromIndex(Param.DestGPUIndex);
+
+				if (Param.Texture)
+				{
+					SetGPUMask(GPUMaskSrc);
+					TransitionInternal(FRHITransitionInfo(Param.Texture.GetReference(), ERHIAccess::CopySrc, ERHIAccess::Unknown, EResourceTransitionFlags::IgnoreAfterState),
+						ERHITransitionCreateFlags::AllowDuringRenderPass);
+
+					SetGPUMask(GPUMaskDest);
+					TransitionInternal(FRHITransitionInfo(Param.Texture.GetReference(), ERHIAccess::CopyDest, ERHIAccess::Unknown, EResourceTransitionFlags::IgnoreAfterState),
+						ERHITransitionCreateFlags::AllowDuringRenderPass);
+
+				}
+				else
+				{
+					SetGPUMask(GPUMaskSrc);
+					TransitionInternal(FRHITransitionInfo(Param.Buffer.GetReference(), ERHIAccess::CopySrc, ERHIAccess::Unknown, EResourceTransitionFlags::IgnoreAfterState),
+						ERHITransitionCreateFlags::AllowDuringRenderPass);
+
+					SetGPUMask(GPUMaskDest);
+					TransitionInternal(FRHITransitionInfo(Param.Buffer.GetReference(), ERHIAccess::CopyDest, ERHIAccess::Unknown, EResourceTransitionFlags::IgnoreAfterState),
+						ERHITransitionCreateFlags::AllowDuringRenderPass);
+				}
+			}
+		}
+
+		SetGPUMask(PrevGPUMask);
+
 #endif // WITH_MGPU
 	}
 
@@ -3322,6 +3466,31 @@ public:
 	FORCEINLINE_DEBUGGABLE void CrossGPUTransfer(TConstArrayView<FTransferResourceParams> Params, TConstArrayView<FCrossGPUTransferFence*> PreTransfer, TConstArrayView<FCrossGPUTransferFence*> PostTransfer)
 	{
 #if WITH_MGPU
+		FRHIGPUMask PrevGPUMask = GetGPUMask();
+
+		if (NeedsExtraTransitions())
+		{
+			for (const FTransferResourceParams& Param : Params)
+			{
+				FRHIGPUMask GPUMask = FRHIGPUMask::FromIndex(Param.SrcGPUIndex);
+				SetGPUMask(GPUMask);
+				
+				if (Param.Texture)
+				{
+					
+					TransitionInternal(FRHITransitionInfo(Param.Texture.GetReference(), ERHIAccess::Unknown, ERHIAccess::CopySrc, EResourceTransitionFlags::IgnoreAfterState),
+						ERHITransitionCreateFlags::AllowDuringRenderPass);
+				}
+				else
+				{
+					TransitionInternal(FRHITransitionInfo(Param.Buffer.GetReference(), ERHIAccess::Unknown, ERHIAccess::CopySrc, EResourceTransitionFlags::IgnoreAfterState),
+						ERHITransitionCreateFlags::AllowDuringRenderPass);
+				}
+			}
+		}
+
+		SetGPUMask(PrevGPUMask);
+
 		if (Bypass())
 		{
 			GetComputeContext().RHICrossGPUTransfer(Params, PreTransfer, PostTransfer);
@@ -3330,12 +3499,60 @@ public:
 		{
 			ALLOC_COMMAND(FRHICommandCrossGPUTransfer)(AllocArray(Params), AllocArray(PreTransfer), AllocArray(PostTransfer));
 		}
+
+		if (NeedsExtraTransitions())
+		{
+			for (const FTransferResourceParams& Param : Params)
+			{
+				FRHIGPUMask GPUMask = FRHIGPUMask::FromIndex(Param.SrcGPUIndex);
+				SetGPUMask(GPUMask);
+
+				if (Param.Texture)
+				{
+					TransitionInternal(FRHITransitionInfo(Param.Texture.GetReference(), ERHIAccess::CopySrc, ERHIAccess::Unknown, EResourceTransitionFlags::IgnoreAfterState),
+						ERHITransitionCreateFlags::AllowDuringRenderPass);
+				}
+				else
+				{
+					TransitionInternal(FRHITransitionInfo(Param.Buffer.GetReference(), ERHIAccess::CopySrc, ERHIAccess::Unknown, EResourceTransitionFlags::IgnoreAfterState),
+						ERHITransitionCreateFlags::AllowDuringRenderPass);
+				}
+			}
+		}
+
+		SetGPUMask(PrevGPUMask);
+
 #endif // WITH_MGPU
 	}
 
 	FORCEINLINE_DEBUGGABLE void CrossGPUTransferSignal(TConstArrayView<FTransferResourceParams> Params, TConstArrayView<FCrossGPUTransferFence*> PreTransfer)
 	{
 #if WITH_MGPU
+
+		FRHIGPUMask PrevGPUMask = GetGPUMask();
+
+		if (NeedsExtraTransitions())
+		{
+			for (const FTransferResourceParams& Param : Params)
+			{
+				FRHIGPUMask GPUMask = FRHIGPUMask::FromIndex(Param.DestGPUIndex);
+				SetGPUMask(GPUMask);
+
+				if (Param.Texture)
+				{
+					TransitionInternal(FRHITransitionInfo(Param.Texture.GetReference(), ERHIAccess::Unknown, ERHIAccess::CopyDest, EResourceTransitionFlags::IgnoreAfterState), 
+						ERHITransitionCreateFlags::AllowDuringRenderPass);
+				}
+				else
+				{
+					TransitionInternal(FRHITransitionInfo(Param.Buffer.GetReference(), ERHIAccess::Unknown, ERHIAccess::CopyDest, EResourceTransitionFlags::IgnoreAfterState), 
+						ERHITransitionCreateFlags::AllowDuringRenderPass);
+				}
+			}
+		}
+
+		SetGPUMask(PrevGPUMask);
+		
 		if (Bypass())
 		{
 			GetComputeContext().RHICrossGPUTransferSignal(Params, PreTransfer);
@@ -3344,6 +3561,28 @@ public:
 		{
 			ALLOC_COMMAND(FRHICommandCrossGPUTransferSignal)(AllocArray(Params), AllocArray(PreTransfer));
 		}
+
+		if (NeedsExtraTransitions())
+		{
+			for (const FTransferResourceParams& Param : Params)
+			{
+				FRHIGPUMask GPUMask = FRHIGPUMask::FromIndex(Param.DestGPUIndex);
+				SetGPUMask(GPUMask);
+
+				if (Param.Texture)
+				{
+					TransitionInternal(FRHITransitionInfo(Param.Texture.GetReference(), ERHIAccess::CopyDest, ERHIAccess::Unknown, EResourceTransitionFlags::IgnoreAfterState), 
+						ERHITransitionCreateFlags::AllowDuringRenderPass);
+				}
+				else
+				{
+					TransitionInternal(FRHITransitionInfo(Param.Buffer.GetReference(), ERHIAccess::CopyDest, ERHIAccess::Unknown, EResourceTransitionFlags::IgnoreAfterState), 
+						ERHITransitionCreateFlags::AllowDuringRenderPass);
+				}
+			}
+		}
+
+		SetGPUMask(PrevGPUMask);
 #endif // WITH_MGPU
 	}
 
@@ -3366,6 +3605,11 @@ public:
 
 	FORCEINLINE_DEBUGGABLE void BuildAccelerationStructures(TConstArrayView<FRayTracingGeometryBuildParams> Params, const FRHIBufferRange& ScratchBufferRange)
 	{
+		if (NeedsExtraTransitions())
+		{
+			Transition(FRHITransitionInfo(ScratchBufferRange.Buffer, ERHIAccess::Unknown, ERHIAccess::UAVCompute));
+		}
+
 		if (Bypass())
 		{
 			GetComputeContext().RHIBuildAccelerationStructures(Params, ScratchBufferRange);
diff --git a/Engine/Source/Runtime/RHI/Public/RHICommandList.inl b/Engine/Source/Runtime/RHI/Public/RHICommandList.inl
--- a/Engine/Source/Runtime/RHI/Public/RHICommandList.inl
+++ b/Engine/Source/Runtime/RHI/Public/RHICommandList.inl
@@ -224,3 +224,75 @@ inline void TRHILambdaCommandMultiPipe<RHICmdListType, LAMBDA>::ExecuteAndDestru
 	Lambda(static_cast<FRHIContextArray const&>(Contexts));
 	Lambda.~LAMBDA();
 }
+
+FORCEINLINE_DEBUGGABLE void FRHICommandListBase::PatchTransitionsInPlace(FRHITransitionInfo* TransitionInfo, const ERHIPipeline SrcPipeline, const ERHIPipeline DestPipeline)
+{
+	if ((TransitionInfo->Type == FRHITransitionInfo::EType::Buffer) && (GRHIGlobals.SupportImplicitBufferStateTracking))
+	{
+		// This RHI supports Buffers state tracking at platform level, so we don't need to track them
+		return;
+	}
+
+	FRHISubresourceRange SubResourceIndex = *TransitionInfo;
+	FRHIViewableResource* TransitionViewableResource = GetViewableResource(*TransitionInfo);
+
+	// Both Unknown is not a supported case
+	check(!((TransitionInfo->AccessBefore == ERHIAccess::Unknown) && (TransitionInfo->AccessAfter == ERHIAccess::Unknown)));
+	// Discard to Unknown is not a supported case
+	check(!((TransitionInfo->AccessBefore == ERHIAccess::Discard) && (TransitionInfo->AccessAfter == ERHIAccess::Unknown)));
+
+	if (TransitionViewableResource)
+	{
+		// Some RHI will not transition to initial state on After = Discard when running on the Async queue
+		if (NeedsExtraTransitions())
+		{
+			if ((TransitionInfo->AccessAfter == ERHIAccess::Discard)
+				&& ((SrcPipeline == ERHIPipeline::AsyncCompute) || (SrcPipeline == ERHIPipeline::All))
+				&& (DestPipeline == ERHIPipeline::AsyncCompute))
+			{
+				TransitionInfo->Flags |= EResourceTransitionFlags::IgnoreAfterState;
+			}
+		}
+
+		const bool bNeedUpdateAfterState = !((EnumHasAnyFlags(TransitionInfo->Flags, EResourceTransitionFlags::IgnoreAfterState)) || (TransitionInfo->AccessAfter == ERHIAccess::Unknown));
+		// Some RHI needs extra information on transition from Discard
+		const bool bBeforeIsDiscard = NeedsExtraTransitions() && (TransitionInfo->AccessBefore == ERHIAccess::Discard);
+
+		if (bBeforeIsDiscard)
+		{
+			TransitionInfo->DiscardAccessBefore = TransitionViewableResource->TrackedAccess.GetSubresourceState(SubResourceIndex);
+		}
+		else if (TransitionInfo->AccessBefore == ERHIAccess::Unknown)
+		{
+			TransitionInfo->AccessBefore = TransitionViewableResource->TrackedAccess.GetSubresourceState(SubResourceIndex);
+		}
+		else if (TransitionInfo->AccessAfter == ERHIAccess::Unknown)
+		{
+			TransitionInfo->AccessAfter = TransitionViewableResource->TrackedAccess.GetSubresourceState(SubResourceIndex);
+		}
+
+		if (bNeedUpdateAfterState)
+		{
+			TransitionViewableResource->TrackedAccess.SetSubresourceState(SubResourceIndex, TransitionInfo->AccessAfter);
+			TransitionViewableResource->TrackedAccess.CheckAllSubresourceSame();
+		}
+	}
+
+	check(TransitionInfo->AccessBefore != ERHIAccess::Unknown);
+	check(TransitionInfo->AccessAfter != ERHIAccess::Unknown);
+}
+
+
+FORCEINLINE_DEBUGGABLE void FRHICommandListBase::PatchTransitionsInPlace(TArrayView<const FRHITransition*> Transitions)
+{
+	check(Bypass());
+
+	for (const FRHITransition* Transition : Transitions)
+	{
+		FRHITransitionCreateInfoData* TransitionCreateInfoData = Transition->CreationInfoData;
+		for (int32 i = 0; i < TransitionCreateInfoData->TransitionInfos.Num(); i++)
+		{
+			PatchTransitionsInPlace((FRHITransitionInfo*)&TransitionCreateInfoData->TransitionInfos[i], TransitionCreateInfoData->SrcPipelines, TransitionCreateInfoData->DstPipelines);
+		}
+	}
+}
diff --git a/Engine/Source/Runtime/RHI/Public/RHICommandListCommandExecutes.inl b/Engine/Source/Runtime/RHI/Public/RHICommandListCommandExecutes.inl
--- a/Engine/Source/Runtime/RHI/Public/RHICommandListCommandExecutes.inl
+++ b/Engine/Source/Runtime/RHI/Public/RHICommandListCommandExecutes.inl
@@ -372,6 +372,17 @@ void FRHICommandResummarizeHTile::Execute(FRHICommandListBase& CmdList)
 void FRHICommandBeginTransitions::Execute(FRHICommandListBase& CmdList)
 {
 	RHISTAT(BeginTransitions);
+
+	// Create the platform level transitions
+	if (GRHIGlobals.EnableHighLevelTransitionPatching)
+	{
+		for (const FRHITransition* Transition : Transitions)
+		{
+			Transition->CreatePlatformTransitionThreadSafe();
+			checkSlow(Transition->IsPlatformLevelCreated());
+		}
+	}
+
 	INTERNAL_DECORATOR_COMPUTE(RHIBeginTransitions)(Transitions);
 
 	for (const FRHITransition* Transition : Transitions)
@@ -383,6 +394,17 @@ void FRHICommandBeginTransitions::Execute(FRHICommandListBase& CmdList)
 void FRHICommandEndTransitions::Execute(FRHICommandListBase& CmdList)
 {
 	RHISTAT(EndTransitions);
+
+	// Create the platform level transitions
+	if (GRHIGlobals.EnableHighLevelTransitionPatching)
+	{
+		for (const FRHITransition* Transition : Transitions)
+		{
+			Transition->CreatePlatformTransitionThreadSafe();
+			checkSlow(Transition->IsPlatformLevelCreated());
+		}
+	}
+
 	INTERNAL_DECORATOR_COMPUTE(RHIEndTransitions)(Transitions);
 
 	for (const FRHITransition* Transition : Transitions)
@@ -395,12 +417,18 @@ void FRHICommandResourceTransition::Execute(FRHICommandListBase& CmdList)
 {
 	RHISTAT(ResourceTransition);
 
+	// Create FRHITransition from FRHITransitionHighLevel with Platform information
+	if (GRHIGlobals.EnableHighLevelTransitionPatching)
+	{
+		Transition->CreatePlatformTransitionThreadSafe();
+	}
+
+	checkSlow(Transition->IsPlatformLevelCreated());
+
 	INTERNAL_DECORATOR_COMPUTE(RHIBeginTransitions)(MakeArrayView((const FRHITransition**)&Transition, 1));
 	INTERNAL_DECORATOR_COMPUTE(RHIEndTransitions)(MakeArrayView((const FRHITransition**)&Transition, 1));
 
-	// Manual release. No need to free, the instance was allocated on the RHICmdList stack.
-	GDynamicRHI->RHIReleaseTransition(Transition);
-	Transition->~FRHITransition();
+	Transition->Cleanup();
 }
 
 void FRHICommandSetTrackedAccess::Execute(FRHICommandListBase& CmdList)
diff --git a/Engine/Source/Runtime/RHI/Public/RHIContext.h b/Engine/Source/Runtime/RHI/Public/RHIContext.h
--- a/Engine/Source/Runtime/RHI/Public/RHIContext.h
+++ b/Engine/Source/Runtime/RHI/Public/RHIContext.h
@@ -487,7 +487,7 @@ public:
 	inline ERHIAccess GetTrackedAccess(const FRHIViewableResource* Resource) const
 	{
 		check(Resource);
-		return Resource->TrackedAccess;
+		return Resource->TrackedAccess.GetResourceState();
 	}
 
 	virtual void* RHIGetNativeCommandBuffer() { return nullptr; }
diff --git a/Engine/Source/Runtime/RHI/Public/RHIDefinitions.h b/Engine/Source/Runtime/RHI/Public/RHIDefinitions.h
--- a/Engine/Source/Runtime/RHI/Public/RHIDefinitions.h
+++ b/Engine/Source/Runtime/RHI/Public/RHIDefinitions.h
@@ -13,11 +13,15 @@
 #include "Misc/EnumClassFlags.h"
 #include "ProfilingDebugging/CsvProfilerConfig.h" // TODO Move defines into RHIDefinitions
 #include "UObject/NameTypes.h"
+#include "Experimental/ConcurrentLinearAllocator.h"
 
 #ifndef USE_STATIC_SHADER_PLATFORM_ENUMS
 #define USE_STATIC_SHADER_PLATFORM_ENUMS 0
 #endif
 
+/** RHI Logging. */
+RHI_API DECLARE_LOG_CATEGORY_EXTERN(LogRHI, Log, VeryVerbose);
+
 /** Alignment of the shader parameters struct is required to be 16-byte boundaries. */
 #define SHADER_PARAMETER_STRUCT_ALIGNMENT 16
 
@@ -1321,8 +1325,10 @@ enum class ERHITransitionCreateFlags
 	// so should use a partial flush rather than a fence as this is more optimal.
 	NoSplit = 1 << 1,
 
-	BeginSimpleMode
+	// Indicates that the transition is allowed to happen during a RenderPass
+	AllowDuringRenderPass = 1 << 2,
 };
+
 ENUM_CLASS_FLAGS(ERHITransitionCreateFlags);
 
 enum class EResourceTransitionFlags
@@ -1332,6 +1338,7 @@ enum class EResourceTransitionFlags
 	MaintainCompression = 1 << 0, // Specifies that the transition should not decompress the resource, allowing us to read a compressed resource directly in its compressed state.
 	Discard				= 1 << 1, // Specifies that the data in the resource should be discarded during the transition - used for transient resource acquire when the resource will be fully overwritten
 	Clear				= 1 << 2, // Specifies that the data in the resource should be cleared during the transition - used for transient resource acquire when the resource might not be fully overwritten
+	IgnoreAfterState	= 1 << 3, // The transition will be patched for Unknown if needed but the After State will be ignored by the tracker
 
 	Last = Clear,
 	Mask = (Last << 1) - 1
@@ -1540,3 +1547,11 @@ struct FShaderCodeValidationUBSize
 	uint16 BindPoint;
 	uint32 Size;
 };
+
+#if UE_BUILD_DEBUG
+#define RHI_TRANSITIONS_HL_DEBUG	1
+#else
+#define RHI_TRANSITIONS_HL_DEBUG	0
+#endif
+
+#define RHI_UnknownSubResourceCount (0)
\ No newline at end of file
diff --git a/Engine/Source/Runtime/RHI/Public/RHIGlobals.h b/Engine/Source/Runtime/RHI/Public/RHIGlobals.h
--- a/Engine/Source/Runtime/RHI/Public/RHIGlobals.h
+++ b/Engine/Source/Runtime/RHI/Public/RHIGlobals.h
@@ -686,6 +686,15 @@ struct FRHIGlobals
 
 	/** true if the RHI supports shader execution reordering. */
 	bool SupportsShaderExecutionReordering = false;
+
+	/** True if RHI can automatically handle resource transition for Buffers type, therefore tracking at high level is not needed*/
+	bool SupportImplicitBufferStateTracking = false;
+
+	/** True if RHI ask for patching of Unknown transition states in the RHI high level */
+	bool EnableHighLevelTransitionPatching = false;
+
+	/** True if RHI needs the transitions from the high level for COPYSRC/COPYDEST*/
+	bool NeedExtraTransitions = false;
 };
 
 extern RHI_API FRHIGlobals GRHIGlobals;
diff --git a/Engine/Source/Runtime/RHI/Public/RHIResources.h b/Engine/Source/Runtime/RHI/Public/RHIResources.h
--- a/Engine/Source/Runtime/RHI/Public/RHIResources.h
+++ b/Engine/Source/Runtime/RHI/Public/RHIResources.h
@@ -1306,19 +1306,30 @@ public:
 		TrackedAccess = Access;
 	}
 
+	ERHIAccess GetTrackedAccess_Unsafe()
+	{
+		return TrackedAccess.GetResourceState();
+	}
+
 	FName GetName() const
 	{
 		return Name;
 	}
 
+	/** @return The number of SubResources in the TrackedState. */
+	uint16 GetTrackedSubResourceCount() const
+	{
+		return TrackedAccess.SubResourceCount;
+	}
+
 #if ENABLE_RHI_VALIDATION
 	virtual RHIValidation::FResource* GetValidationTrackerResource() = 0;
 #endif
 
 protected:
-	FRHIViewableResource(ERHIResourceType InResourceType, ERHIAccess InAccess)
+	FRHIViewableResource(ERHIResourceType InResourceType, ERHIAccess InAccess, uint16 SubResourceCount)
 		: FRHIResource(InResourceType)
-		, TrackedAccess(InAccess)
+		, TrackedAccess(InAccess, SubResourceCount)
 	{}
 
 	void TakeOwnership(FRHIViewableResource& Other)
@@ -1334,8 +1345,9 @@ protected:
 	FName Name;
 
 private:
-	ERHIAccess TrackedAccess;
+	TRHIResourceState<FDefaultAllocator> TrackedAccess;
 
+	friend class FRHICommandListBase;
 	friend class FRHIComputeCommandList;
 	friend class IRHIComputeContext;
 };
@@ -1379,7 +1391,7 @@ class FRHIBuffer : public FRHIViewableResource
 public:
 	/** Initialization constructor. */
 	FRHIBuffer(FRHIBufferDesc const& InDesc)
-		: FRHIViewableResource(RRT_Buffer, ERHIAccess::Unknown /* TODO (RemoveUnknowns): Use InitialAccess from descriptor after refactor. */)
+		: FRHIViewableResource(RRT_Buffer, ERHIAccess::Unknown /* TODO (RemoveUnknowns): Use InitialAccess from descriptor after refactor. */,  1)
 		, Desc(InDesc)
 	{}
 
@@ -1684,6 +1696,14 @@ struct FRHITextureDesc
 		return CalcMemorySizeEstimate(FirstMipIndex, NumMips - 1);
 	}
 
+	uint16 GetSubresourceCount() const
+	{
+		uint16 NumFaces = ((Dimension == ETextureDimension::TextureCube) || (Dimension == ETextureDimension::TextureCubeArray)) ? 6 : 1;
+		uint16 NumPlanes = IsStencilFormat(Format) ? 2 : 1;
+
+		return (ArraySize * (NumMips * NumFaces) * NumPlanes);
+	}
+
 private:
 	RHI_API static bool Validate(const FRHITextureDesc& Desc, const TCHAR* Name, bool bFatal);
 };
@@ -2067,7 +2087,7 @@ private:
 	friend class FRHITextureReference;
 	/** Constructor for texture references */
 	FRHITexture(ERHIResourceType InResourceType)
-		: FRHIViewableResource(InResourceType, ERHIAccess::Unknown)
+		: FRHIViewableResource(InResourceType, ERHIAccess::Unknown, RHI_UnknownSubResourceCount)
 	{
 		check(InResourceType == RRT_TextureReference);
 	}
@@ -2263,6 +2283,7 @@ public:
 	 */
 	virtual class FRHICustomPresent* GetCustomPresent() const { return nullptr; }
 
+	virtual FRHITexture* GetOptionalSDRBackBuffer(FRHITexture* BackBuffer) const { return nullptr; }
 
 	/**
 	 * Ticks the viewport on the Game thread
diff --git a/Engine/Source/Runtime/RHI/Public/RHITransition.h b/Engine/Source/Runtime/RHI/Public/RHITransition.h
--- a/Engine/Source/Runtime/RHI/Public/RHITransition.h
+++ b/Engine/Source/Runtime/RHI/Public/RHITransition.h
@@ -2,26 +2,30 @@
 
 #pragma once
 
+#include "Async/Mutex.h"
+#include "Async/UniqueLock.h"
 #include "Containers/ArrayView.h"
 #include "Misc/Optional.h"
 #include "RHIDefinitions.h"
+#include "RHIAllocators.h"
 #include "RHIAccess.h"
 #include "RHIPipeline.h"
 #include "RHIValidationCommon.h"
 
+
 // The size in bytes of the storage required by the platform RHI for each resource transition.
 extern RHI_API uint64 GRHITransitionPrivateData_SizeInBytes;
 extern RHI_API uint64 GRHITransitionPrivateData_AlignInBytes;
 
 struct FRHISubresourceRange
 {
-	static const uint32 kDepthPlaneSlice = 0;
-	static const uint32 kStencilPlaneSlice = 1;
-	static const uint32 kAllSubresources = TNumericLimits<uint32>::Max();
+	static const uint16 kDepthPlaneSlice = 0;
+	static const uint16 kStencilPlaneSlice = 1;
+	static const uint16 kAllSubresources = TNumericLimits<uint16>::Max();
 
-	uint32 MipIndex = kAllSubresources;
-	uint32 ArraySlice = kAllSubresources;
-	uint32 PlaneSlice = kAllSubresources;
+	uint16 MipIndex = kAllSubresources;
+	uint16 ArraySlice = kAllSubresources;
+	uint16 PlaneSlice = kAllSubresources;
 
 	FRHISubresourceRange() = default;
 
@@ -29,6 +33,20 @@ struct FRHISubresourceRange
 		uint32 InMipIndex,
 		uint32 InArraySlice,
 		uint32 InPlaneSlice)
+	{
+		checkSlow(InMipIndex <= TNumericLimits<uint16>::Max());
+		checkSlow(InArraySlice <= TNumericLimits<uint16>::Max());
+		checkSlow(InPlaneSlice <= TNumericLimits<uint16>::Max());
+
+		MipIndex = (uint16)InMipIndex;
+		ArraySlice = (uint16)InArraySlice;
+		PlaneSlice = (uint16)InPlaneSlice;
+	}
+
+	FRHISubresourceRange(
+		uint16 InMipIndex,
+		uint16 InArraySlice,
+		uint16 InPlaneSlice)
 		: MipIndex(InMipIndex)
 		, ArraySlice(InArraySlice)
 		, PlaneSlice(InPlaneSlice)
@@ -54,6 +72,11 @@ struct FRHISubresourceRange
 		return IsAllMips() && IsAllArraySlices() && IsAllPlaneSlices();
 	}
 
+	inline bool IsZeroIndexResource() const
+	{
+		return (MipIndex == 0) && (ArraySlice == 0) && (PlaneSlice == 0);
+	}
+
 	inline bool IgnoreDepthPlane() const
 	{
 		return PlaneSlice == kStencilPlaneSlice;
@@ -115,6 +138,7 @@ struct FRHITransitionInfo : public FRHISubresourceRange
 
 	ERHIAccess AccessBefore = ERHIAccess::Unknown;
 	ERHIAccess AccessAfter = ERHIAccess::Unknown;
+	ERHIAccess DiscardAccessBefore = ERHIAccess::Unknown;
 	EResourceTransitionFlags Flags = EResourceTransitionFlags::None;
 	TOptional<FRHICommitResourceInfo> CommitInfo;
 
@@ -191,6 +215,7 @@ struct FRHITransitionInfo : public FRHISubresourceRange
 		return Resource == RHS.Resource
 			&& Type == RHS.Type
 			&& AccessBefore == RHS.AccessBefore
+			&& DiscardAccessBefore == RHS.DiscardAccessBefore
 			&& AccessAfter == RHS.AccessAfter
 			&& Flags == RHS.Flags
 			&& FRHISubresourceRange::operator==(RHS);
@@ -356,6 +381,295 @@ struct FRHITransientAliasingInfo
 	}
 };
 
+//==================================================================================================================================
+// FRHIResourceState
+// Tracking of per-resource or per-subresource state
+//==================================================================================================================================
+template <typename TSecondaryAllocator>
+class TRHIResourceState
+{
+public:
+	struct FSubResourceElement
+	{
+		FSubResourceElement() {}
+		FSubResourceElement(const FRHISubresourceRange& Range, ERHIAccess State)
+		: SubRange(Range)
+		, SubState(State)
+		{}
+
+		FRHISubresourceRange SubRange;
+		ERHIAccess SubState;
+	};
+
+	explicit TRHIResourceState(uint16 SubResourceCount) 
+		: ResourceState(ERHIAccess::Unknown)
+		, UnsetSubResourceState(ERHIAccess::Unknown)
+		, SubResourceCount(SubResourceCount)
+		, bAllSubresourcesSame(1)
+		, bCmdListTracking(0)
+	{}
+
+	explicit TRHIResourceState(ERHIAccess inState, uint16 SubResourceCount, bool CmdListTracking = false) 
+		: ResourceState(inState)
+		, UnsetSubResourceState(ERHIAccess::Unknown)
+		, SubResourceCount(SubResourceCount)
+		, bAllSubresourcesSame(1)
+		, bCmdListTracking(CmdListTracking)
+	{}
+
+	explicit TRHIResourceState(const FRHISubresourceRange& Range, ERHIAccess State, uint16 SubResourceCount, bool CmdListTracking = false)
+		: ResourceState(ERHIAccess::Unknown)
+		, UnsetSubResourceState(ERHIAccess::Unknown)
+		, SubResourceCount(SubResourceCount)
+		, bAllSubresourcesSame(1)
+		, bCmdListTracking(CmdListTracking)
+	{
+		SetSubresourceState(Range, State);
+	}
+
+	~TRHIResourceState() {}
+
+	friend class TRHIResourceState<FRHICmdListBaseArrayAllocator>;
+	
+	//Using for set "global" resource state from cmdlist states
+	inline void Merge(const TRHIResourceState<FRHICmdListBaseArrayAllocator>& Other)
+	{
+		check(Other.bCmdListTracking);
+
+		if (Other.bAllSubresourcesSame)
+		{ 
+			check(Other.ResourceState != ERHIAccess::Unknown)
+			check(Other.SubResourcesState.IsEmpty() == true);
+			SetResourceState(Other.ResourceState);
+		}
+		else
+		{
+			for (int32 i = 0; i < Other.SubResourcesState.Num(); i++)
+			{
+				SetSubresourceState(Other.SubResourcesState[i].SubRange, Other.SubResourcesState[i].SubState);
+			}
+		}
+
+		//Compaction of subresource tracking in to resource tracking, might need to move it at explicit calls to avoid perf repercussions
+		if (bAllSubresourcesSame == false)
+		{
+			if (SubResourcesState.Num() == SubResourceCount)
+			{
+				ERHIAccess AllSubResState = SubResourcesState[0].SubState;
+
+				for (int32 i = 1; i < SubResourcesState.Num(); i++)
+				{
+					if (SubResourcesState[i].SubState != AllSubResState)
+					{
+						AllSubResState = ERHIAccess::Unknown;
+						break;
+					}
+				}
+
+				if (AllSubResState != ERHIAccess::Unknown)
+				{
+					SetResourceState(AllSubResState);
+				}
+			}
+		}
+	}
+
+	inline bool AreAllSubresourcesSame() const { return bAllSubresourcesSame && (ResourceState != ERHIAccess::Unknown); }
+	inline bool CheckAllSubresourceSame()
+	{
+		// already marked same?
+		if (bAllSubresourcesSame)
+		{
+			return true;
+		}
+		
+		if (SubResourcesState.Num() < SubResourceCount)
+		{
+			return false;
+		}
+
+		ERHIAccess State = SubResourcesState[0].SubState;
+
+		// All subresources must be individually checked
+		for (int32 i = 1; i < SubResourcesState.Num(); i++)
+		{
+			if (SubResourcesState[i].SubState != State)
+			{
+				return false;
+			}
+		}
+
+		SetResourceState(State);
+
+		return true;
+	}
+
+	inline ERHIAccess GetResourceState() const
+	{
+		return GetSubresourceState(FRHISubresourceRange(FRHISubresourceRange::kAllSubresources, FRHISubresourceRange::kAllSubresources, FRHISubresourceRange::kAllSubresources));
+	}
+	
+	inline ERHIAccess GetSubresourceState(const FRHISubresourceRange& SubresourceIndex) const
+	{
+		checkSlow(SubResourceCount != RHI_UnknownSubResourceCount);
+
+		if (bAllSubresourcesSame) 
+		{
+			return ResourceState;
+		}
+		else
+		{
+			if (SubresourceIndex.IsWholeResource())
+			{
+				//trying to get the status of the whole resource while subresource are in different state each one
+				if (SubResourcesState.Num() == SubResourceCount)
+				{
+					ERHIAccess AllSubResState = SubResourcesState[0].SubState;
+
+					for (int32 i = 1; i < SubResourcesState.Num(); i++)
+					{
+						if (SubResourcesState[i].SubState != AllSubResState)
+						{
+							AllSubResState = ERHIAccess::Unknown;
+							break;
+						}
+					}
+
+					check(AllSubResState != ERHIAccess::Unknown);
+					return AllSubResState;
+				}
+
+				// Edge case, only resource 0 is set and we are asking for the whole resource state
+				if ((SubResourcesState.Num() == 1) && (SubResourcesState[0].SubRange.IsZeroIndexResource()))
+				{
+					return SubResourcesState[0].SubState;
+				}
+
+				// This should not happen, we are asking for the whole resource state but subresources are in different state within each other
+				check(false);
+				return ERHIAccess::Unknown;
+			}
+
+			for (int32 i = 0; i < SubResourcesState.Num(); i++)
+			{
+				if (SubResourcesState[i].SubRange == SubresourceIndex)
+				{
+					return SubResourcesState[i].SubState;
+				}
+			}
+
+			//trying to get the status of one subresources that we didn't set yet, return the Unset value
+			if (SubResourcesState.Num() < SubResourceCount)
+			{
+				return UnsetSubResourceState;
+			}
+
+			//resource wasn't there 
+			return ERHIAccess::Unknown;
+		}
+	}
+
+	inline void SetResourceState(ERHIAccess State)
+	{
+		SetSubresourceState(FRHISubresourceRange(FRHISubresourceRange::kAllSubresources, FRHISubresourceRange::kAllSubresources, FRHISubresourceRange::kAllSubresources), State);
+	}
+	
+	inline void SetSubresourceState(const FRHISubresourceRange& SubresourceIndex, ERHIAccess State)
+	{
+		checkSlow(SubResourceCount != RHI_UnknownSubResourceCount);
+		checkSlow(IsValidAccess(State));
+
+		// If setting all subresources
+		if ((SubresourceIndex.IsWholeResource()) || (SubResourceCount == 1))
+		{
+			if (SubresourceIndex.IsWholeResource() == false)
+				check(SubresourceIndex.IsZeroIndexResource());
+
+			bAllSubresourcesSame = 1;
+			ResourceState = State;
+			SubResourcesState.Empty(SubResourceCount);
+		}
+		else
+		{
+			check(SubresourceIndex.MipIndex * SubresourceIndex.ArraySlice * SubresourceIndex.PlaneSlice <= SubResourceCount);
+			
+			SubResourcesState.Reserve(SubResourceCount);
+			
+			// If state was previously tracked on a per-resource level, then transition to per-subresource tracking
+			if (bAllSubresourcesSame)
+			{
+				for (int32 i = 0; i < SubResourcesState.Num(); i++)
+				{
+					SubResourcesState[i].SubState = ResourceState;
+				}
+
+				bAllSubresourcesSame = 0;
+
+				// State is now tracked per-subresource, m_ResourceState still would be the not set subresources state
+				UnsetSubResourceState = ResourceState;
+				ResourceState = ERHIAccess::Unknown;
+			}
+
+			bool bFound = false;
+			for (int32 i = 0; i < SubResourcesState.Num(); i++)
+			{
+				if (SubResourcesState[i].SubRange == SubresourceIndex)
+				{
+					bFound = true;
+					SubResourcesState[i].SubState = State;
+					break;
+				}
+			}
+
+			if (bFound == false)
+			{
+				SubResourcesState.Push(FSubResourceElement(SubresourceIndex, State));
+				check(SubResourcesState.Num() <= SubResourceCount);
+			}
+
+		}
+	}
+
+	//Operators
+	inline TRHIResourceState& operator= (ERHIAccess State)
+	{
+		check(bAllSubresourcesSame == 1);
+		SetResourceState(State);
+		return *this;
+	}
+
+	inline bool operator == (const ERHIAccess Other) const
+	{
+		check(bAllSubresourcesSame == 1);
+		return Other == ResourceState;
+	}
+
+	bool IsValid()
+	{
+		if (bAllSubresourcesSame == false)
+		{
+			if (SubResourcesState.IsEmpty())
+				return false;
+		}
+
+		return true;
+	}
+
+	// SubresourceState[0...n] status of the subresources if any
+	TArray<FSubResourceElement, TSizedInlineAllocator<8, 32, TSecondaryAllocator> > SubResourcesState; //todo test if a sparse container would perform better
+	// ResourceState is the status of the whole resource
+	ERHIAccess ResourceState;
+	// Status of unset SubResources
+	ERHIAccess UnsetSubResourceState;
+	// Total number of SubResources
+	uint16 SubResourceCount;
+	// Set to 1 if ResourceState is valid.  In this case, all subresources have the same state
+	// Set to 0 if SubresourceState is valid.  In this case, each subresources may have a different state (or may be unknown)
+	uint8 bAllSubresourcesSame : 1;
+	// Set to 1 if the State is used on Command list tracking  
+	uint8 bCmdListTracking : 1;
+};
+
 struct FRHITransitionCreateInfo
 {
 	FRHITransitionCreateInfo() = default;
@@ -378,6 +692,33 @@ struct FRHITransitionCreateInfo
 	ERHITransitionCreateFlags Flags = ERHITransitionCreateFlags::None;
 	TArrayView<const FRHITransitionInfo> TransitionInfos;
 	TArrayView<const FRHITransientAliasingInfo> AliasingInfos;
+
+	inline bool operator == (FRHITransitionCreateInfo const& RHS) const
+	{
+		if (SrcPipelines != RHS.SrcPipelines)
+			return false;
+
+		if (DstPipelines != RHS.DstPipelines)
+			return false;
+
+		if (Flags != RHS.Flags)
+			return false;
+
+		for (int32 i = 0; i < RHS.TransitionInfos.Num(); i++)
+		{
+			if (TransitionInfos[i] != RHS.TransitionInfos[i])
+				return false;
+		}
+
+		for (int32 i = 0; i < RHS.AliasingInfos.Num(); i++)
+		{
+			if (AliasingInfos[i] != RHS.AliasingInfos[i])
+				return false;
+		}
+
+		return true;
+	}
+
 };
 
 struct FRHITrackedAccessInfo
@@ -393,6 +734,88 @@ struct FRHITrackedAccessInfo
 	ERHIAccess Access = ERHIAccess::Unknown;
 };
 
+struct FRHITransition;
+RHI_API FRHITransition* RHICreateTransition(const FRHITransitionCreateInfo& CreateInfo);
+
+
+enum class ETransitionFlag : uint8
+{
+	None = 0,
+	AllowDuringRenderPass = 1 << 0,			// Specific transition are happening during RenderPass and are reported by the validation layer, we need to allow some of them
+	PlatformLevelCreated  = 1 << 1			// Platform level has been created on this Transition
+};
+
+// Struct used to store creation info and create the opaque platform dependent FRHITransition after tracking and patching
+struct FRHITransitionCreateInfoData
+{
+private:
+	// Prevent copying and moving. Only pointers to these structures are allowed.
+	FRHITransitionCreateInfoData(const FRHITransitionCreateInfoData&) = delete;
+	FRHITransitionCreateInfoData(FRHITransitionCreateInfoData&&) = delete;
+
+public:
+	FRHITransitionCreateInfoData(const FRHITransitionCreateInfo& CreationInfo)
+	{
+		SrcPipelines = CreationInfo.SrcPipelines;
+		DstPipelines = CreationInfo.DstPipelines;
+		m_Flags = CreationInfo.Flags;
+
+		TransitionInfos.SetNumUninitialized(CreationInfo.TransitionInfos.Num());
+		for (int32 i = 0; i < CreationInfo.TransitionInfos.Num(); i++)
+		{
+			TransitionInfos[i] = CreationInfo.TransitionInfos[i];
+		}
+
+		AliasingInfos.SetNumUninitialized(CreationInfo.AliasingInfos.Num());
+		for (int32 i = 0; i < CreationInfo.AliasingInfos.Num(); i++)
+		{
+			AliasingInfos[i] = CreationInfo.AliasingInfos[i];
+		}
+
+		// Deep copy of AlisianInfo used for validation
+		if (GRHIValidationEnabled)
+		{
+			uint32 OffsetDataIndex = 0;
+			uint32 AllOverlapNum = 0;
+			for (FRHITransientAliasingInfo& AliasingInfo : AliasingInfos)
+			{
+				AllOverlapNum += AliasingInfo.Overlaps.Num();
+			}
+
+			TransientAliasingOverlapData.Reserve(AllOverlapNum);
+			if (AllOverlapNum)
+			{
+				for (FRHITransientAliasingInfo& AliasingInfo : AliasingInfos)
+				{
+					if (!AliasingInfo.Overlaps.IsEmpty())
+					{
+						TransientAliasingOverlapData.Append(AliasingInfo.Overlaps.GetData(), AliasingInfo.Overlaps.Num());
+						AliasingInfo.Overlaps = TArrayView<const FRHITransientAliasingOverlap>(&TransientAliasingOverlapData[OffsetDataIndex], AliasingInfo.Overlaps.Num());
+						OffsetDataIndex += AliasingInfo.Overlaps.Num();
+					}
+				}
+			}
+		}
+	}
+
+	virtual ~FRHITransitionCreateInfoData()
+	{}
+
+	inline const FRHITransitionCreateInfo GetCreationInfo() const
+	{
+		return FRHITransitionCreateInfo(SrcPipelines, DstPipelines, m_Flags, MakeArrayView(TransitionInfos), MakeArrayView(AliasingInfos));
+	}
+
+	ERHIPipeline SrcPipelines = ERHIPipeline::None;
+	ERHIPipeline DstPipelines = ERHIPipeline::None;
+	mutable ERHITransitionCreateFlags m_Flags = ERHITransitionCreateFlags::None;
+	TArray<FRHITransitionInfo, FRHICmdListBaseArrayAllocator>		 TransitionInfos;
+	TArray<FRHITransientAliasingInfo, FRHICmdListBaseArrayAllocator> AliasingInfos;
+
+	// We need to copy aliasing overlap information 
+	TArray<FRHITransientAliasingOverlap, FRHICmdListBaseArrayAllocator> TransientAliasingOverlapData;
+};
+
 // Opaque data structure used to represent a pending resource transition in the RHI.
 struct FRHITransition
 #if ENABLE_RHI_VALIDATION
@@ -415,14 +838,19 @@ public:
 		return const_cast<FRHITransition*>(this)->GetPrivateData<T>();
 	}
 
+	inline bool AllowInRenderingPass() const { return EnumHasAnyFlags(Flags, ETransitionFlag::AllowDuringRenderPass); }
+
+	inline bool IsPlatformLevelCreated() const { return EnumHasAnyFlags(Flags, ETransitionFlag::PlatformLevelCreated); }
+
 private:
 	// Prevent copying and moving. Only pointers to these structures are allowed.
 	FRHITransition(const FRHITransition&) = delete;
 	FRHITransition(FRHITransition&&) = delete;
 
 	// Private constructor. Memory for transitions is allocated manually with extra space at the tail of the structure for RHI use.
-	FRHITransition(ERHIPipeline SrcPipelines, ERHIPipeline DstPipelines)
+	FRHITransition(ERHIPipeline SrcPipelines, ERHIPipeline DstPipelines, ETransitionFlag NewFlags = ETransitionFlag::None)
 		: State(int8(int32(SrcPipelines) | (int32(DstPipelines) << int32(ERHIPipeline::Num))))
+		, Flags(NewFlags)
 #if DO_CHECK || USING_CODE_ANALYSIS
 		, AllowedSrc(SrcPipelines)
 		, AllowedDst(DstPipelines)
@@ -433,7 +861,8 @@ private:
 	{}
 
 	// Give private access to specific functions/RHI commands that need to allocate or control transitions.
-	friend const FRHITransition* RHICreateTransition(const FRHITransitionCreateInfo&);
+	friend RHI_API FRHITransition* RHICreateTransition(const FRHITransitionCreateInfo& CreateInfo);
+	friend class FRHICommandListBase;
 	friend class FRHIComputeCommandList;
 	friend struct FRHICommandBeginTransitions;
 	friend struct FRHICommandEndTransitions;
@@ -449,6 +878,9 @@ private:
 	{
 		return FMath::Max((uint64)alignof(FRHITransition), GRHITransitionPrivateData_AlignInBytes);
 	}
+	
+	RHI_API void CreatePlatformTransition() const;
+	RHI_API void CreatePlatformTransitionThreadSafe() const;
 
 	inline void MarkBegin(ERHIPipeline Pipeline) const
 	{
@@ -477,11 +909,18 @@ private:
 			Cleanup();
 		}
 	}
+	
+	RHI_API void CleanupInfoData() const;
 
 	RHI_API void Cleanup() const;
 
 	mutable int8 State;
 	static_assert((int32(ERHIPipeline::Num) * 2) < (sizeof(State) * 8), "Not enough bits to hold pipeline state.");
+	mutable ETransitionFlag Flags;
+
+	// Deep copy of FRHITransitionCreateInfo used to track resource state and patch Unknown transition
+	mutable FRHITransitionCreateInfoData* CreationInfoData = nullptr;
+	mutable UE::FMutex TransitionCreationMutex;
 
 #if DO_CHECK || USING_CODE_ANALYSIS
 	mutable ERHIPipeline AllowedSrc;
diff --git a/Engine/Source/Runtime/RHI/Public/RHIValidationCommon.h b/Engine/Source/Runtime/RHI/Public/RHIValidationCommon.h
--- a/Engine/Source/Runtime/RHI/Public/RHIValidationCommon.h
+++ b/Engine/Source/Runtime/RHI/Public/RHIValidationCommon.h
@@ -190,7 +190,7 @@ namespace RHIValidation
 		TRHIPipelineArray<FPipelineState> States;
 
 		void BeginTransition   (FResource* Resource, FSubresourceIndex const& SubresourceIndex, const FState& CurrentStateFromRHI, const FState& TargetState, EResourceTransitionFlags NewFlags, ERHITransitionCreateFlags CreateFlags, ERHIPipeline Pipeline, const TRHIPipelineArray<uint64>& PipelineMaxAwaitedFenceValues, void* CreateTrace);
-		void EndTransition     (FResource* Resource, FSubresourceIndex const& SubresourceIndex, const FState& CurrentStateFromRHI, const FState& TargetState, ERHIPipeline Pipeline, uint64 PipelineFenceValue, void* CreateTrace);
+		void EndTransition     (FResource* Resource, FSubresourceIndex const& SubresourceIndex, const FState& CurrentStateFromRHI, const FState& TargetState, EResourceTransitionFlags NewFlags, ERHIPipeline Pipeline, uint64 PipelineFenceValue, void* CreateTrace);
 		void Assert            (FResource* Resource, FSubresourceIndex const& SubresourceIndex, const FState& RequiredState, bool bAllowAllUAVsOverlap);
 		void AssertTracked     (FResource* Resource, FSubresourceIndex const& SubresourceIndex, const FState& TrackedState);
 		void SpecificUAVOverlap(FResource* Resource, FSubresourceIndex const& SubresourceIndex, ERHIPipeline Pipeline, bool bAllow);
@@ -270,6 +270,7 @@ namespace RHIValidation
 
 		FTransientState(ERHIAccess InitialAccess)
 			: bTransient(InitialAccess == ERHIAccess::Discard)
+			, Status((InitialAccess == ERHIAccess::Discard) ? EStatus::Discarded : EStatus::None)
 		{}
 
 		void* AcquireBacktrace = nullptr;
@@ -568,6 +569,7 @@ namespace RHIValidation
 				FResourceIdentity Identity;
 				FState PreviousState;
 				FState NextState;
+				EResourceTransitionFlags Flags;
 				void* CreateBacktrace;
 			} Data_EndTransition;
 
@@ -664,7 +666,7 @@ namespace RHIValidation
 			return MoveTemp(Op);
 		}
 
-		static inline FOperation EndTransitionResource(FResourceIdentity Identity, FState PreviousState, FState NextState, void* CreateBacktrace)
+		static inline FOperation EndTransitionResource(FResourceIdentity Identity, FState PreviousState, FState NextState, EResourceTransitionFlags Flags, void* CreateBacktrace)
 		{
 			for (ERHIPipeline Pipeline : MakeFlagsRange(NextState.Pipelines))
 			{
@@ -676,6 +678,7 @@ namespace RHIValidation
 			Op.Data_EndTransition.Identity = Identity;
 			Op.Data_EndTransition.PreviousState = PreviousState;
 			Op.Data_EndTransition.NextState = NextState;
+			Op.Data_EndTransition.Flags = Flags;
 			Op.Data_EndTransition.CreateBacktrace = CreateBacktrace;
 			return MoveTemp(Op);
 		}
@@ -979,7 +982,7 @@ namespace RHIValidation
 			// This function exists due to the implicit transitions that RHI functions make (e.g. RHICopyToResolveTarget).
 			// It should be removed when we eventually remove all implicit transitions from the RHI.
 			AddOp(FOperation::BeginTransitionResource(Identity, PreviousState, NextState, Flags, ERHITransitionCreateFlags::None, nullptr));
-			AddOp(FOperation::EndTransitionResource(Identity, PreviousState, NextState, nullptr));
+			AddOp(FOperation::EndTransitionResource(Identity, PreviousState, NextState, Flags, nullptr));
 		}
 
 		void AllUAVsOverlap(bool bAllow)
diff --git a/Engine/Source/Runtime/RHI/Public/RHIValidationContext.h b/Engine/Source/Runtime/RHI/Public/RHIValidationContext.h
--- a/Engine/Source/Runtime/RHI/Public/RHIValidationContext.h
+++ b/Engine/Source/Runtime/RHI/Public/RHIValidationContext.h
@@ -362,6 +362,22 @@ public:
 	virtual void RHIBuildAccelerationStructures(TConstArrayView<FRayTracingGeometryBuildParams> Params, const FRHIBufferRange& ScratchBufferRange) override final
 	{
 		// #yuriy_todo: explicit transitions and state validation for BLAS
+		for (const FRayTracingGeometryBuildParams& P : Params)
+		{
+			const FRayTracingGeometryInitializer& Initializer = P.Geometry->GetInitializer();
+
+			if (Initializer.IndexBuffer)
+			{
+				Tracker->Assert(Initializer.IndexBuffer->GetWholeResourceIdentity(), ERHIAccess::SRVCompute);
+			}
+
+			for (const FRayTracingGeometrySegment& Segment : Initializer.Segments)
+			{
+				const FBufferRHIRef& RHIVertexBuffer = Segment.VertexBuffer;
+				Tracker->Assert(Segment.VertexBuffer->GetWholeResourceIdentity(), ERHIAccess::SRVCompute);
+			}
+		}
+
 		RHIContext->RHIBuildAccelerationStructures(Params, ScratchBufferRange);
 	}
 
@@ -372,6 +388,17 @@ public:
 		{
 			Tracker->Assert(SceneBuildParams.Scene->GetWholeResourceIdentity(), ERHIAccess::BVHWrite);
 		}
+
+		if (SceneBuildParams.InstanceBuffer)
+		{
+			Tracker->Assert(SceneBuildParams.InstanceBuffer->GetWholeResourceIdentity(), ERHIAccess::SRVCompute);
+		}
+
+		if (SceneBuildParams.ScratchBuffer)
+		{
+			Tracker->Assert(SceneBuildParams.ScratchBuffer->GetWholeResourceIdentity(), ERHIAccess::UAVCompute);
+		}
+		
 		RHIContext->RHIBuildAccelerationStructure(SceneBuildParams);
 	}
 
@@ -620,10 +647,13 @@ public:
 
 	virtual void RHIBeginTransitions(TArrayView<const FRHITransition*> Transitions) override final
 	{
-		ensureMsgf(!State.bInsideBeginRenderPass, TEXT("Beginning a transition within a renderpass is not supported!"));
-
 		for (const FRHITransition* Transition : Transitions)
 		{
+			if (Transition->AllowInRenderingPass() == false)
+			{
+				ensureMsgf(!State.bInsideBeginRenderPass, TEXT("Ending a transition within a renderpass is not supported!"));
+			}
+
 			Tracker->AddOps(Transition->PendingAliases);
 		}
 
@@ -642,10 +672,13 @@ public:
 
 	virtual void RHIEndTransitions(TArrayView<const FRHITransition*> Transitions) override final
 	{
-		ensureMsgf(!State.bInsideBeginRenderPass, TEXT("Ending a transition within a renderpass is not supported!"));
-
 		for (const FRHITransition* Transition : Transitions)
 		{
+			if (Transition->AllowInRenderingPass() == false)
+			{
+				ensureMsgf(!State.bInsideBeginRenderPass, TEXT("Ending a transition within a renderpass is not supported!"));
+			}
+
 			Tracker->AddOps(Transition->PendingWaits[GetPipeline()]);
 		}
 
@@ -1165,6 +1198,18 @@ public:
 	virtual void RHIBuildAccelerationStructures(TConstArrayView<FRayTracingGeometryBuildParams> Params, const FRHIBufferRange& ScratchBufferRange) override final
 	{
 		// #yuriy_todo: explicit transitions and state validation for BLAS
+		for (const FRayTracingGeometryBuildParams& P : Params)
+		{
+			const FRayTracingGeometryInitializer& Initializer = P.Geometry->GetInitializer();
+
+//			TODO: Add ERHIAccess::SRVCompute assert for Initializer.IndexBuffer, still some resources in the wrong state eg. FHairCardIndexBuffer
+
+			for (const FRayTracingGeometrySegment& Segment : Initializer.Segments)
+			{
+				Tracker->Assert(Segment.VertexBuffer->GetWholeResourceIdentity(), ERHIAccess::SRVCompute);
+			}
+		}
+
 		RHIContext->RHIBuildAccelerationStructures(Params, ScratchBufferRange);
 	}
 
@@ -1175,6 +1220,17 @@ public:
 		{
 			Tracker->Assert(SceneBuildParams.Scene->GetWholeResourceIdentity(), ERHIAccess::BVHWrite);
 		}
+
+		if (SceneBuildParams.InstanceBuffer)
+		{
+			Tracker->Assert(SceneBuildParams.InstanceBuffer->GetWholeResourceIdentity(), ERHIAccess::SRVCompute);
+		}
+
+		if (SceneBuildParams.ScratchBuffer)
+		{
+			Tracker->Assert(SceneBuildParams.ScratchBuffer->GetWholeResourceIdentity(), ERHIAccess::UAVCompute);
+		}
+
 		RHIContext->RHIBuildAccelerationStructure(SceneBuildParams);
 	}
 
diff --git a/Engine/Source/Runtime/RenderCore/Private/GlobalRenderResources.cpp b/Engine/Source/Runtime/RenderCore/Private/GlobalRenderResources.cpp
--- a/Engine/Source/Runtime/RenderCore/Private/GlobalRenderResources.cpp
+++ b/Engine/Source/Runtime/RenderCore/Private/GlobalRenderResources.cpp
@@ -428,7 +428,8 @@ public:
 		const FRHITextureCreateDesc Desc =
 			FRHITextureCreateDesc::CreateCube(TEXT("SolidColorCube"), 1, PixelFormat)
 			.SetFlags(ETextureCreateFlags::ShaderResource)
-			.SetClassName(TEXT("FSolidColorTextureCube"));
+			.SetClassName(TEXT("FSolidColorTextureCube"))
+			.SetInitialState(ERHIAccess::SRVMask);
 
 		FTextureRHIRef TextureCube = RHICreateTexture(Desc);
 		TextureRHI = TextureCube;
@@ -483,7 +484,8 @@ public:
 		const FRHITextureCreateDesc Desc =
 			FRHITextureCreateDesc::CreateCubeArray(TEXT("BlackCubeArray"), 1, 1, PF_B8G8R8A8)
 			.SetFlags(ETextureCreateFlags::ShaderResource)
-			.SetClassName(TEXT("FBlackCubeArrayTexture"));
+			.SetClassName(TEXT("FBlackCubeArrayTexture"))
+			.SetInitialState(ERHIAccess::SRVMask);
 
 		// Create the texture RHI.
 		TextureRHI = RHICreateTexture(Desc);
diff --git a/Engine/Source/Runtime/SlateRHIRenderer/Private/SlateRHIRenderer.cpp b/Engine/Source/Runtime/SlateRHIRenderer/Private/SlateRHIRenderer.cpp
--- a/Engine/Source/Runtime/SlateRHIRenderer/Private/SlateRHIRenderer.cpp
+++ b/Engine/Source/Runtime/SlateRHIRenderer/Private/SlateRHIRenderer.cpp
@@ -927,8 +927,33 @@ void FSlateRHIRenderer::PresentWindow_RenderThread(FRHICommandListImmediate& RHI
 		UEngine::SetPresentLatencyMarkerStart(CurrentFrameCounter);
 	});
 
+	FRHITexture* OptionalSDRBuffer = nullptr;
+
 	RHICmdList.BeginDrawingViewport(DrawPassOutputs.ViewportRHI, FTextureRHIRef());
+
+	if (GRHIGlobals.NeedExtraTransitions)
+	{
+		RHICmdList.Transition(FRHITransitionInfo(DrawPassOutputs.OutputTextureRHI, ERHIAccess::Unknown, ERHIAccess::Present));
+
+		OptionalSDRBuffer = DrawPassOutputs.ViewportRHI->GetOptionalSDRBackBuffer(DrawPassOutputs.OutputTextureRHI);
+
+		if (OptionalSDRBuffer)
+		{
+			RHICmdList.Transition(FRHITransitionInfo(OptionalSDRBuffer, ERHIAccess::Unknown, ERHIAccess::Present));
+		}
+	}
+
 	RHICmdList.EndDrawingViewport(DrawPassOutputs.ViewportRHI, true, DrawPassInputs.bLockToVsync);
+	
+	if (GRHIGlobals.NeedExtraTransitions)
+	{
+		RHICmdList.Transition(FRHITransitionInfo(DrawPassOutputs.OutputTextureRHI, ERHIAccess::Present, ERHIAccess::SRVMask));
+
+		if (OptionalSDRBuffer)
+		{
+			RHICmdList.Transition(FRHITransitionInfo(OptionalSDRBuffer, ERHIAccess::Present, ERHIAccess::SRVMask));
+		}
+	}
 
 	RHICmdList.EnqueueLambda([CurrentFrameCounter = GFrameCounterRenderThread](FRHICommandListImmediate& InRHICmdList)
 	{
