Minor RHI cleanups
- Set OwnerName in FRHITexture constructor.
- Inlining a couple "Name" functions.
- Don't pass FName by reference.

diff --git a/Engine/Source/Runtime/RHI/Private/RHI.cpp b/Engine/Source/Runtime/RHI/Private/RHI.cpp
--- a/Engine/Source/Runtime/RHI/Private/RHI.cpp
+++ b/Engine/Source/Runtime/RHI/Private/RHI.cpp
@@ -505,22 +505,6 @@ bool operator==(const FBlendStateInitializerRHI::FRenderTarget& A, const FBlendS
 	return bSame;
 }
 
-FName FRHIResource::GetOwnerName() const
-{
-#if RHI_ENABLE_RESOURCE_INFO
-	return OwnerName;
-#else
-	return NAME_None;
-#endif
-}
-
-void FRHIResource::SetOwnerName(const FName& InOwnerName)
-{
-#if RHI_ENABLE_RESOURCE_INFO
-	OwnerName = InOwnerName;
-#endif
-}
-
 #if RHI_ENABLE_RESOURCE_INFO
 
 static FCriticalSection GRHIResourceTrackingCriticalSection;
diff --git a/Engine/Source/Runtime/RHI/Private/RHIResources.cpp b/Engine/Source/Runtime/RHI/Private/RHIResources.cpp
--- a/Engine/Source/Runtime/RHI/Private/RHIResources.cpp
+++ b/Engine/Source/Runtime/RHI/Private/RHIResources.cpp
@@ -128,9 +128,10 @@ FRHITexture::FRHITexture(const FRHITextureCreateDesc& InDesc)
 	, TextureDesc(InDesc)
 {
 	SetName(InDesc.DebugName);
+	SetOwnerName(InDesc.OwnerName);
 }
 
-void FRHITexture::SetName(const FName& InName)
+void FRHITexture::SetName(FName InName)
 {
 	Name = InName;
 
diff --git a/Engine/Source/Runtime/RHI/Public/RHIResources.h b/Engine/Source/Runtime/RHI/Public/RHIResources.h
--- a/Engine/Source/Runtime/RHI/Public/RHIResources.h
+++ b/Engine/Source/Runtime/RHI/Public/RHIResources.h
@@ -106,8 +106,20 @@ public:
 
 	inline ERHIResourceType GetType() const { return ResourceType; }
 
-	RHI_API FName GetOwnerName() const;
-	RHI_API void SetOwnerName(const FName& InOwnerName);
+	FName GetOwnerName() const
+	{
+#if RHI_ENABLE_RESOURCE_INFO
+		return OwnerName;
+#else
+		return NAME_None;
+#endif
+	}
+	void SetOwnerName(FName InOwnerName)
+	{
+#if RHI_ENABLE_RESOURCE_INFO
+		OwnerName = InOwnerName;
+#endif
+	}
 
 #if RHI_ENABLE_RESOURCE_INFO
 	// Get resource info if available.
@@ -2028,7 +2040,7 @@ public:
 		return LastRenderTime.GetLastRenderTime();
 	}
 
-	RHI_API void SetName(const FName& InName);
+	RHI_API void SetName(FName InName);
 
 	///
 	/// Deprecated functions
