[Insights]
  - Added tracing for the new GPU Profiler (UE_TRACE_GPU_PROFILER_ENABLED). Requires RHI_NEW_GPU_PROFILER=1.
  - Changed GpuChannel to be declared in the global namespace (instead of GpuProfilerTrace namespace). Now GPU channel is used also by the new GPU Profiler.
  - Deprecated all TRACE_GPUPROFILER_* macros from the old GPU profiler. These macros were unused and some of them were not even matching correctly the FGpuProfilerTrace API.

diff --git a/Engine/Source/Runtime/RHI/Private/GPUProfiler.cpp b/Engine/Source/Runtime/RHI/Private/GPUProfiler.cpp
--- a/Engine/Source/Runtime/RHI/Private/GPUProfiler.cpp
+++ b/Engine/Source/Runtime/RHI/Private/GPUProfiler.cpp
@@ -1278,12 +1278,14 @@ namespace UE::RHI::GPUProfiler
 						FFrameState& FrameState = Frames.FindOrAdd(FrameBoundary.FrameNumber);
 						FrameState.Emplace(Queue, MoveTemp(QueueState.Timestamps));
 
+					#if WITH_RHI_BREADCRUMBS
 						// Reinsert timestamp streams for the current active stats on 
 						// this queue, since these got moved into the frame state.
 						for (auto& [Stat, RefCount] : QueueState.ActiveStats)
 						{
 							QueueState.Timestamps.Stats.FindOrAdd(Stat);
 						}
+					#endif
 
 						if (FrameState.Num() == QueueStates.Num())
 						{
diff --git a/Engine/Source/Runtime/RHI/Private/GpuProfilerTrace.cpp b/Engine/Source/Runtime/RHI/Private/GpuProfilerTrace.cpp
--- a/Engine/Source/Runtime/RHI/Private/GpuProfilerTrace.cpp
+++ b/Engine/Source/Runtime/RHI/Private/GpuProfilerTrace.cpp
@@ -10,6 +10,18 @@
 #include "Trace/Detail/Important/ImportantLogScope.h"
 #include "Trace/Detail/LogScope.h"
 
+// Tracing for the new GPU Profiler
+#ifndef UE_TRACE_GPU_PROFILER_ENABLED
+#define UE_TRACE_GPU_PROFILER_ENABLED UE_TRACE_ENABLED && RHI_NEW_GPU_PROFILER && !UE_BUILD_SHIPPING
+#endif
+
+// Both (old and new) GPU profilers uses same "GPU" trace channel.
+#if GPUPROFILERTRACE_ENABLED || UE_TRACE_GPU_PROFILER_ENABLED
+	UE_TRACE_CHANNEL_EXTERN(GpuChannel, RHI_API)
+	UE_TRACE_CHANNEL_DEFINE(GpuChannel)
+#endif
+
+// The old GpuProfilerTrace is deprecated in UE 5.6
 #if GPUPROFILERTRACE_ENABLED
 
 namespace GpuProfilerTrace
@@ -43,15 +55,13 @@ FGpuTraceFrame GCurrentFrame;
 
 static TSet<uint32> GEventNames;
 
-UE_TRACE_CHANNEL_EXTERN(GpuChannel, RHI_API)
-UE_TRACE_CHANNEL_DEFINE(GpuChannel)
-
+// deprecated in UE 5.6
 UE_TRACE_EVENT_BEGIN(GpuProfiler, EventSpec, NoSync|Important)
 	UE_TRACE_EVENT_FIELD(uint32, EventType)
 	UE_TRACE_EVENT_FIELD(uint16[], Name)
 UE_TRACE_EVENT_END()
 
-// GPU Index 0
+// GPU Index 0 - deprecated in UE 5.6
 UE_TRACE_EVENT_BEGIN(GpuProfiler, Frame)
 	UE_TRACE_EVENT_FIELD(uint64, CalibrationBias)
 	UE_TRACE_EVENT_FIELD(uint64, TimestampBase)
@@ -59,7 +69,7 @@ UE_TRACE_EVENT_BEGIN(GpuProfiler, Frame)
 	UE_TRACE_EVENT_FIELD(uint8[], Data)
 UE_TRACE_EVENT_END()
 
-// GPU Index 1
+// GPU Index 1 - deprecated in UE 5.6
 UE_TRACE_EVENT_BEGIN(GpuProfiler, Frame2)
 	UE_TRACE_EVENT_FIELD(uint64, CalibrationBias)
 	UE_TRACE_EVENT_FIELD(uint64, TimestampBase)
@@ -256,4 +266,196 @@ void FGpuProfilerTrace::Deinitialize()
 	}
 }
 
-#endif
+#endif // GPUPROFILERTRACE_ENABLED
+
+// Tracing for the new GPU Profiler
+#if UE_TRACE_GPU_PROFILER_ENABLED
+
+namespace UE::RHI::GPUProfiler
+{
+
+UE_TRACE_EVENT_BEGIN(GpuProfiler, Init, NoSync | Important)
+	UE_TRACE_EVENT_FIELD(uint8, Version)
+UE_TRACE_EVENT_END()
+
+UE_TRACE_EVENT_BEGIN(GpuProfiler, QueueSpec, NoSync | Important)
+	UE_TRACE_EVENT_FIELD(uint32, QueueId)
+	UE_TRACE_EVENT_FIELD(UE::Trace::WideString, TypeString)
+UE_TRACE_EVENT_END()
+
+UE_TRACE_EVENT_BEGIN(GpuProfiler, EventFrameBoundary, NoSync)
+	UE_TRACE_EVENT_FIELD(uint32, QueueId)
+	UE_TRACE_EVENT_FIELD(uint32, FrameNumber)
+UE_TRACE_EVENT_END()
+
+#if WITH_RHI_BREADCRUMBS
+UE_TRACE_EVENT_BEGIN(GpuProfiler, EventBeginBreadcrumb, NoSync)
+	UE_TRACE_EVENT_FIELD(uint32, QueueId)
+	UE_TRACE_EVENT_FIELD(uint64, GPUTimestampTOP)
+	UE_TRACE_EVENT_FIELD(UE::Trace::WideString, Name)
+	//TODO: long name format + serialized args
+	//UE_TRACE_EVENT_FIELD(UE::Trace::WideString, LongNameFormat)
+	//UE_TRACE_EVENT_FIELD(uint8[], LongNameArgs)
+UE_TRACE_EVENT_END()
+
+UE_TRACE_EVENT_BEGIN(GpuProfiler, EventEndBreadcrumb, NoSync)
+	UE_TRACE_EVENT_FIELD(uint32, QueueId)
+	UE_TRACE_EVENT_FIELD(uint64, GPUTimestampBOP)
+	UE_TRACE_EVENT_FIELD(UE::Trace::WideString, Name) // for debugging
+UE_TRACE_EVENT_END()
+#endif // WITH_RHI_BREADCRUMBS
+
+UE_TRACE_EVENT_BEGIN(GpuProfiler, EventBeginWork, NoSync)
+	UE_TRACE_EVENT_FIELD(uint32, QueueId)
+	UE_TRACE_EVENT_FIELD(uint64, GPUTimestampTOP)
+	UE_TRACE_EVENT_FIELD(uint64, CPUTimestamp)
+UE_TRACE_EVENT_END()
+
+UE_TRACE_EVENT_BEGIN(GpuProfiler, EventEndWork, NoSync)
+	UE_TRACE_EVENT_FIELD(uint32, QueueId)
+	UE_TRACE_EVENT_FIELD(uint64, GPUTimestampBOP)
+UE_TRACE_EVENT_END()
+
+UE_TRACE_EVENT_BEGIN(GpuProfiler, EventStats, NoSync)
+	UE_TRACE_EVENT_FIELD(uint32, QueueId)
+	UE_TRACE_EVENT_FIELD(uint32, NumDraws)
+	UE_TRACE_EVENT_FIELD(uint32, NumPrimitives)
+UE_TRACE_EVENT_END()
+
+struct FGPUProfilerSink_Trace final : public FEventSink
+{
+	bool bInitialized = false;
+
+	void InitializeQueues(TConstArrayView<FQueue> Queues) override
+	{
+		if (!UE_TRACE_CHANNELEXPR_IS_ENABLED(GpuChannel))
+		{
+			//TODO: Should we still initialize even if GPU channel is off?
+			// (ex.: to keep a copy of all queues and to emit Init and QueueSpec later when GPU channel is on)
+			return;
+		}
+
+		if (bInitialized)
+		{
+			//TODO: Is it valid for InitializeQueues to be called multiple times?
+			//return;
+		}
+
+		bInitialized = true;
+
+		UE_TRACE_LOG(GpuProfiler, Init, GpuChannel)
+			<< Init.Version(2);
+
+		for (FQueue const& Queue : Queues)
+		{
+			FStringView TypeString(Queue.GetTypeString());
+			UE_TRACE_LOG(GpuProfiler, QueueSpec, GpuChannel, TypeString.Len() * sizeof(TCHAR))
+				<< QueueSpec.QueueId(Queue.Value)
+				<< QueueSpec.TypeString(TypeString.GetData(), TypeString.Len());
+		}
+	}
+
+	void ProcessEvents(FQueue Queue, FEventStream const& EventStream) override
+	{
+		if (!UE_TRACE_CHANNELEXPR_IS_ENABLED(GpuChannel))
+		{
+			return;
+		}
+
+		if (!bInitialized)
+		{
+			//TODO: Can we query now Queues and initialize now?
+			return;
+		}
+
+		for (FEvent const* Event : EventStream)
+		{
+			switch (Event->GetType())
+			{
+				case FEvent::EType::FrameBoundary:
+				{
+					FEvent::FFrameBoundary const& FrameBoundary = Event->Value.Get<FEvent::FFrameBoundary>();
+					UE_TRACE_LOG(GpuProfiler, EventFrameBoundary, GpuChannel)
+						<< EventFrameBoundary.QueueId(Queue.Value)
+						<< EventFrameBoundary.FrameNumber(FrameBoundary.FrameNumber);
+					break;
+				}
+
+#if WITH_RHI_BREADCRUMBS
+				case FEvent::EType::BeginBreadcrumb:
+				{
+					FEvent::FBeginBreadcrumb const& BeginBreadcrumb = Event->Value.Get<FEvent::FBeginBreadcrumb>();
+					FRHIBreadcrumb::FBuffer Buffer;
+					const TCHAR* Name = BeginBreadcrumb.Breadcrumb->Name.GetTCHAR(Buffer);
+					UE_TRACE_LOG(GpuProfiler, EventBeginBreadcrumb, GpuChannel)
+						<< EventBeginBreadcrumb.QueueId(Queue.Value)
+						<< EventBeginBreadcrumb.GPUTimestampTOP(BeginBreadcrumb.GPUTimestampTOP)
+						<< EventBeginBreadcrumb.Name(Name);
+					break;
+				}
+
+				case FEvent::EType::EndBreadcrumb:
+				{
+					FEvent::FEndBreadcrumb const& EndBreadcrumb = Event->Value.Get<FEvent::FEndBreadcrumb>();
+					FRHIBreadcrumb::FBuffer Buffer;
+					const TCHAR* Name = EndBreadcrumb.Breadcrumb->Name.GetTCHAR(Buffer);
+					UE_TRACE_LOG(GpuProfiler, EventEndBreadcrumb, GpuChannel)
+						<< EventEndBreadcrumb.QueueId(Queue.Value)
+						<< EventEndBreadcrumb.GPUTimestampBOP(EndBreadcrumb.GPUTimestampBOP)
+						<< EventEndBreadcrumb.Name(Name); // for debugging
+					break;
+				}
+#endif // WITH_RHI_BREADCRUMBS
+
+				case FEvent::EType::BeginWork:
+				{
+					FEvent::FBeginWork const& BeginWork = Event->Value.Get<FEvent::FBeginWork>();
+					UE_TRACE_LOG(GpuProfiler, EventBeginWork, GpuChannel)
+						<< EventBeginWork.QueueId(Queue.Value)
+						<< EventBeginWork.GPUTimestampTOP(BeginWork.GPUTimestampTOP)
+						<< EventBeginWork.CPUTimestamp(BeginWork.CPUTimestamp);
+					break;
+				}
+
+				case FEvent::EType::EndWork:
+				{
+					FEvent::FEndWork const& EndWork = Event->Value.Get<FEvent::FEndWork>();
+					UE_TRACE_LOG(GpuProfiler, EventEndWork, GpuChannel)
+						<< EventEndWork.QueueId(Queue.Value)
+						<< EventEndWork.GPUTimestampBOP(EndWork.GPUTimestampBOP);
+					break;
+				}
+
+				case FEvent::EType::Stats:
+				{
+					FEvent::FStats const& Stats = Event->Value.Get<FEvent::FStats>();
+					UE_TRACE_LOG(GpuProfiler, EventStats, GpuChannel)
+						<< EventStats.QueueId(Queue.Value)
+						<< EventStats.NumDraws(Stats.NumDraws)
+						<< EventStats.NumPrimitives(Stats.NumPrimitives);
+					break;
+				}
+
+				case FEvent::EType::SignalFence:
+				case FEvent::EType::WaitFence:
+				case FEvent::EType::Flip:
+				case FEvent::EType::VSync:
+				{
+					//TODO
+					break;
+				}
+
+				default:
+				{
+					unimplemented();
+					break;
+				}
+			}
+		}
+	}
+
+} GGPUProfilerSink_Trace;
+
+} // namespace UE::RHI::GPUProfiler
+
+#endif // UE_TRACE_GPU_PROFILER_ENABLED
diff --git a/Engine/Source/Runtime/RHI/Public/GpuProfilerTrace.h b/Engine/Source/Runtime/RHI/Public/GpuProfilerTrace.h
--- a/Engine/Source/Runtime/RHI/Public/GpuProfilerTrace.h
+++ b/Engine/Source/Runtime/RHI/Public/GpuProfilerTrace.h
@@ -7,7 +7,9 @@
 
 #if !defined(GPUPROFILERTRACE_ENABLED)
 #if UE_TRACE_ENABLED && !UE_BUILD_SHIPPING
-#define GPUPROFILERTRACE_ENABLED 1
+// The trace implementation for the old GPU Profiler (FGpuProfilerTrace API) is deprecated in UE 5.6.
+// The new GPU Profiler has a different trace implementation (see UE_TRACE_GPU_PROFILER_ENABLED and FGPUProfilerSink_Trace).
+#define GPUPROFILERTRACE_ENABLED !RHI_NEW_GPU_PROFILER
 #else
 #define GPUPROFILERTRACE_ENABLED 0
 #endif
@@ -47,47 +49,14 @@ struct FGpuProfilerTrace
 
 #undef GPU_TRACE_ARG
 
-#define TRACE_GPUPROFILER_DEFINE_EVENT_TYPE(Name) \
-	FGpuProfilerTrace::FEventType PREPROCESSOR_JOIN(__GGpuProfilerEventType, Name)(TEXT(#Name));
-
-#define TRACE_GPUPROFILER_DECLARE_EVENT_TYPE_EXTERN(Name) \
-	extern FGpuProfilerTrace::FEventType PREPROCESSOR_JOIN(__GGpuProfilerEventType, Name);
-
-#define TRACE_GPUPROFILER_EVENT_TYPE(Name) \
-	&PREPROCESSOR_JOIN(__GGpuProfilerEventType, Name)
-
-#define TRACE_GPUPROFILER_BEGIN_FRAME() \
-	FGpuProfilerTrace::BeginFrame();
-
-#define TRACE_GPUPROFILER_BEGIN_EVENT(EventType, FrameNumber, TimestampMicroseconds) \
-	FGpuProfilerTrace::BeginEvent(EventType, FrameNumber, TimestampMicroseconds);
-
-#define TRACE_GPUPROFILER_END_EVENT(TimestampMicroseconds) \
-	FGpuProfilerTrace::EndEvent(TimestampMicroseconds);
-
-#define TRACE_GPUPROFILER_END_FRAME() \
-	FGpuProfilerTrace::EndFrame();
-
-#define TRACE_GPUPROFILER_DEINITIALIZE() \
-	FGpuProfilerTrace::Deinitialize();
-
-#else
-
-struct FGpuProfilerTrace
-{
-	struct FEventType
-	{
-		FEventType(const TCHAR* Name) {};
-	};
-};
-
-#define TRACE_GPUPROFILER_DEFINE_EVENT_TYPE(...)
-#define TRACE_GPUPROFILER_DECLARE_EVENT_TYPE_EXTERN(...)
-#define TRACE_GPUPROFILER_EVENT_TYPE(...) nullptr
-#define TRACE_GPUPROFILER_BEGIN_FRAME(...)
-#define TRACE_GPUPROFILER_BEGIN_EVENT(...)
-#define TRACE_GPUPROFILER_END_EVENT(...)
-#define TRACE_GPUPROFILER_END_FRAME(...)
-#define TRACE_GPUPROFILER_DEINITIALIZE(...)
-
-#endif
+#endif // GPUPROFILERTRACE_ENABLED
+
+// Deprecated macros
+#define TRACE_GPUPROFILER_DEFINE_EVENT_TYPE(...)         UE_DEPRECATED_MACRO(5.6, "TRACE_GPUPROFILER_DEFINE_EVENT_TYPE has been deprecated and it is non functional. Use standard RHI breadcrumb events instead.")
+#define TRACE_GPUPROFILER_DECLARE_EVENT_TYPE_EXTERN(...) UE_DEPRECATED_MACRO(5.6, "TRACE_GPUPROFILER_DECLARE_EVENT_TYPE_EXTERN has been deprecated and it is non functional. Use standard RHI breadcrumb events instead.")
+#define TRACE_GPUPROFILER_EVENT_TYPE(...)                UE_DEPRECATED_MACRO(5.6, "TRACE_GPUPROFILER_EVENT_TYPE has been deprecated and it is non functional. Use standard RHI breadcrumb events instead.") nullptr
+#define TRACE_GPUPROFILER_BEGIN_FRAME(...)               UE_DEPRECATED_MACRO(5.6, "TRACE_GPUPROFILER_BEGIN_FRAME has been deprecated and it is non functional. Use standard RHI breadcrumb events instead.")
+#define TRACE_GPUPROFILER_BEGIN_EVENT(...)               UE_DEPRECATED_MACRO(5.6, "TRACE_GPUPROFILER_BEGIN_EVENT has been deprecated and it is non functional. Use standard RHI breadcrumb events instead.")
+#define TRACE_GPUPROFILER_END_EVENT(...)                 UE_DEPRECATED_MACRO(5.6, "TRACE_GPUPROFILER_END_EVENT has been deprecated and it is non functional. Use standard RHI breadcrumb events instead.")
+#define TRACE_GPUPROFILER_END_FRAME(...)                 UE_DEPRECATED_MACRO(5.6, "TRACE_GPUPROFILER_END_FRAME has been deprecated and it is non functional. Use standard RHI breadcrumb events instead.")
+#define TRACE_GPUPROFILER_DEINITIALIZE(...)              UE_DEPRECATED_MACRO(5.6, "TRACE_GPUPROFILER_DEINITIALIZE has been deprecated and it is non functional. Use standard RHI breadcrumb events instead.")
diff --git a/Engine/Source/Runtime/RenderCore/Private/ProfilingDebugging/RealtimeGPUProfiler.cpp b/Engine/Source/Runtime/RenderCore/Private/ProfilingDebugging/RealtimeGPUProfiler.cpp
--- a/Engine/Source/Runtime/RenderCore/Private/ProfilingDebugging/RealtimeGPUProfiler.cpp
+++ b/Engine/Source/Runtime/RenderCore/Private/ProfilingDebugging/RealtimeGPUProfiler.cpp
@@ -845,14 +845,13 @@ void FRealtimeGPUProfiler::Cleanup()
 
 	RenderQueryPool.SafeRelease();
 	
-	TRACE_GPUPROFILER_DEINITIALIZE();
+#if GPUPROFILERTRACE_ENABLED
+	FGpuProfilerTrace::Deinitialize();
+#endif
 }
 
 #if UE_TRACE_ENABLED
-namespace GpuProfilerTrace
-{
-	UE_TRACE_CHANNEL_EXTERN(GpuChannel, RHI_API)
-}
+UE_TRACE_CHANNEL_EXTERN(GpuChannel, RHI_API)
 #endif
 
 static std::atomic<bool> GAreGPUStatsEnabled{ false };
@@ -873,7 +872,7 @@ void LatchAreGPUStatsEnabled()
 
 #if GPUPROFILERTRACE_ENABLED
 		// Force GPU profiler on if Unreal Insights is running
-		if (UE_TRACE_CHANNELEXPR_IS_ENABLED(GpuProfilerTrace::GpuChannel))
+		if (UE_TRACE_CHANNELEXPR_IS_ENABLED(GpuChannel))
 		{
 			return true;
 		}
