RHI Cleanup
- Removing RHIFinalizeAsyncReallocateTexture2D and RHICancelAsyncReallocateTexture2D along with their _RenderThread variants.
- Deprecating the CommandList callers of those methods.

diff --git a/Engine/Source/Runtime/Apple/MetalRHI/Private/MetalDynamicRHI.h b/Engine/Source/Runtime/Apple/MetalRHI/Private/MetalDynamicRHI.h
--- a/Engine/Source/Runtime/Apple/MetalRHI/Private/MetalDynamicRHI.h
+++ b/Engine/Source/Runtime/Apple/MetalRHI/Private/MetalDynamicRHI.h
@@ -122,8 +122,6 @@ public:
 	virtual FTextureRHIRef RHIAsyncCreateTexture2D(uint32 SizeX, uint32 SizeY, uint8 Format, uint32 NumMips, ETextureCreateFlags Flags, ERHIAccess InResourceState, void** InitialMipData, uint32 NumInitialMips, const TCHAR* DebugName, FGraphEventRef& OutCompletionEvent) final override;
 	virtual uint32 RHIComputeMemorySize(FRHITexture* TextureRHI) final override;
 	virtual FTextureRHIRef RHIAsyncReallocateTexture2D(FRHITexture* Texture2D, int32 NewMipCount, int32 NewSizeX, int32 NewSizeY, FThreadSafeCounter* RequestStatus) final override;
-	virtual ETextureReallocationStatus RHIFinalizeAsyncReallocateTexture2D(FRHITexture* Texture2D, bool bBlockUntilCompleted) final override;
-	virtual ETextureReallocationStatus RHICancelAsyncReallocateTexture2D(FRHITexture* Texture2D, bool bBlockUntilCompleted) final override;
 	virtual void* RHILockTexture2D(FRHITexture* Texture, uint32 MipIndex, EResourceLockMode LockMode, uint32& DestStride, bool bLockWithinMiptail, uint64* OutLockedByteCount = nullptr) final override;
 	virtual void RHIUnlockTexture2D(FRHITexture* Texture, uint32 MipIndex, bool bLockWithinMiptail) final override;
 	virtual void* RHILockTexture2DArray(FRHITexture* Texture, uint32 TextureIndex, uint32 MipIndex, EResourceLockMode LockMode, uint32& DestStride, bool bLockWithinMiptail) final override;
@@ -166,8 +164,6 @@ public:
 	virtual void RHISubmitCommandLists(FRHISubmitCommandListsArgs&& Args) final override;
 
 	virtual FTextureRHIRef AsyncReallocateTexture2D_RenderThread(class FRHICommandListImmediate& RHICmdList, FRHITexture* Texture2D, int32 NewMipCount, int32 NewSizeX, int32 NewSizeY, FThreadSafeCounter* RequestStatus) final override;
-	virtual ETextureReallocationStatus FinalizeAsyncReallocateTexture2D_RenderThread(class FRHICommandListImmediate& RHICmdList, FRHITexture* Texture2D, bool bBlockUntilCompleted) final override;
-	virtual ETextureReallocationStatus CancelAsyncReallocateTexture2D_RenderThread(class FRHICommandListImmediate& RHICmdList, FRHITexture* Texture2D, bool bBlockUntilCompleted) final override;
 
 	virtual void* LockTexture2D_RenderThread(class FRHICommandListImmediate& RHICmdList, FRHITexture* Texture, uint32 MipIndex, EResourceLockMode LockMode, uint32& DestStride, bool bLockWithinMiptail, bool bNeedsDefaultRHIFlush = true, uint64* OutLockedByteCount = nullptr) final override;
 	virtual void UnlockTexture2D_RenderThread(class FRHICommandListImmediate& RHICmdList, FRHITexture* Texture, uint32 MipIndex, bool bLockWithinMiptail, bool bNeedsDefaultRHIFlush = true) final override;
diff --git a/Engine/Source/Runtime/Apple/MetalRHI/Private/MetalTexture.cpp b/Engine/Source/Runtime/Apple/MetalRHI/Private/MetalTexture.cpp
--- a/Engine/Source/Runtime/Apple/MetalRHI/Private/MetalTexture.cpp
+++ b/Engine/Source/Runtime/Apple/MetalRHI/Private/MetalTexture.cpp
@@ -1410,18 +1410,6 @@ FTextureRHIRef FMetalDynamicRHI::AsyncReallocateTexture2D_RenderThread(class FRH
 	return this->RHIAsyncReallocateTexture2D(Texture2D, NewMipCount, NewSizeX, NewSizeY, RequestStatus);
 }
 
-ETextureReallocationStatus FMetalDynamicRHI::FinalizeAsyncReallocateTexture2D_RenderThread(class FRHICommandListImmediate& RHICmdList, FRHITexture* Texture2D, bool bBlockUntilCompleted)
-{
-	// No need to flush - does nothing
-	return this->RHIFinalizeAsyncReallocateTexture2D(Texture2D, bBlockUntilCompleted);
-}
-
-ETextureReallocationStatus FMetalDynamicRHI::CancelAsyncReallocateTexture2D_RenderThread(class FRHICommandListImmediate& RHICmdList, FRHITexture* Texture2D, bool bBlockUntilCompleted)
-{
-	// No need to flush - does nothing
-	return this->RHICancelAsyncReallocateTexture2D(Texture2D, bBlockUntilCompleted);
-}
-
 FTextureRHIRef FMetalDynamicRHI::RHIAsyncReallocateTexture2D(FRHITexture* OldTextureRHI, int32 NewMipCount, int32 NewSizeX, int32 NewSizeY, FThreadSafeCounter* RequestStatus)
 {
     MTL_SCOPED_AUTORELEASE_POOL;
@@ -1482,16 +1470,6 @@ FTextureRHIRef FMetalDynamicRHI::RHIAsyncReallocateTexture2D(FRHITexture* OldTex
     return NewTexture;
 }
 
-ETextureReallocationStatus FMetalDynamicRHI::RHIFinalizeAsyncReallocateTexture2D(FRHITexture* Texture2D, bool bBlockUntilCompleted )
-{
-	return TexRealloc_Succeeded;
-}
-
-ETextureReallocationStatus FMetalDynamicRHI::RHICancelAsyncReallocateTexture2D(FRHITexture* Texture2D, bool bBlockUntilCompleted )
-{
-	return TexRealloc_Failed;
-}
-
 void* FMetalDynamicRHI::LockTexture2D_RenderThread(class FRHICommandListImmediate& RHICmdList, FRHITexture* Texture, uint32 MipIndex, EResourceLockMode LockMode, uint32& DestStride, bool bLockWithinMiptail, bool bNeedsDefaultRHIFlush, uint64* OutLockedByteCount)
 {
     MTL_SCOPED_AUTORELEASE_POOL;
diff --git a/Engine/Source/Runtime/D3D12RHI/Private/D3D12RHIPrivate.h b/Engine/Source/Runtime/D3D12RHI/Private/D3D12RHIPrivate.h
--- a/Engine/Source/Runtime/D3D12RHI/Private/D3D12RHIPrivate.h
+++ b/Engine/Source/Runtime/D3D12RHI/Private/D3D12RHIPrivate.h
@@ -340,8 +340,6 @@ public:
 	virtual FTextureRHIRef RHICreateTexture(FRHICommandListBase& RHICmdList, const FRHITextureCreateDesc& CreateDesc) override;
 	virtual uint32 RHIComputeMemorySize(FRHITexture* TextureRHI) final override;
 	virtual FTextureRHIRef RHIAsyncReallocateTexture2D(FRHITexture* Texture2D, int32 NewMipCount, int32 NewSizeX, int32 NewSizeY, FThreadSafeCounter* RequestStatus) override;
-	virtual ETextureReallocationStatus RHIFinalizeAsyncReallocateTexture2D(FRHITexture* Texture2D, bool bBlockUntilCompleted) final override;
-	virtual ETextureReallocationStatus RHICancelAsyncReallocateTexture2D(FRHITexture* Texture2D, bool bBlockUntilCompleted) final override;
 	virtual void* RHILockTexture2D(FRHITexture* Texture, uint32 MipIndex, EResourceLockMode LockMode, uint32& DestStride, bool bLockWithinMiptail, uint64* OutLockedByteCount = nullptr) final override;
 	virtual void RHIUnlockTexture2D(FRHITexture* Texture, uint32 MipIndex, bool bLockWithinMiptail) final override;
 	virtual void* RHILockTexture2DArray(FRHITexture* Texture, uint32 TextureIndex, uint32 MipIndex, EResourceLockMode LockMode, uint32& DestStride, bool bLockWithinMiptail) final override;
@@ -455,14 +453,6 @@ public:
 	virtual void UnlockTexture2DArray_RenderThread(class FRHICommandListImmediate& RHICmdList, FRHITexture* Texture, uint32 ArrayIndex, uint32 MipIndex, bool bLockWithinMiptail) override final;
 
 	virtual FTextureRHIRef AsyncReallocateTexture2D_RenderThread(class FRHICommandListImmediate& RHICmdList, FRHITexture* Texture2D, int32 NewMipCount, int32 NewSizeX, int32 NewSizeY, FThreadSafeCounter* RequestStatus);
-	virtual ETextureReallocationStatus FinalizeAsyncReallocateTexture2D_RenderThread(class FRHICommandListImmediate& RHICmdList, FRHITexture* Texture2D, bool bBlockUntilCompleted)
-	{
-		return RHIFinalizeAsyncReallocateTexture2D(Texture2D, bBlockUntilCompleted);
-	}
-	virtual ETextureReallocationStatus CancelAsyncReallocateTexture2D_RenderThread(class FRHICommandListImmediate& RHICmdList, FRHITexture* Texture2D, bool bBlockUntilCompleted)
-	{
-		return RHICancelAsyncReallocateTexture2D(Texture2D, bBlockUntilCompleted);
-	}
 
 #if D3D12_RHI_RAYTRACING
 
diff --git a/Engine/Source/Runtime/D3D12RHI/Private/D3D12Texture.cpp b/Engine/Source/Runtime/D3D12RHI/Private/D3D12Texture.cpp
--- a/Engine/Source/Runtime/D3D12RHI/Private/D3D12Texture.cpp
+++ b/Engine/Source/Runtime/D3D12RHI/Private/D3D12Texture.cpp
@@ -1177,17 +1177,6 @@ FTextureRHIRef FD3D12DynamicRHI::RHIAsyncReallocateTexture2D(FRHITexture* Textur
 	return nullptr;
 }
 
-ETextureReallocationStatus FD3D12DynamicRHI::RHIFinalizeAsyncReallocateTexture2D(FRHITexture* Texture2D, bool bBlockUntilCompleted)
-{
-	return TexRealloc_Succeeded;
-}
-
-ETextureReallocationStatus FD3D12DynamicRHI::RHICancelAsyncReallocateTexture2D(FRHITexture* Texture2D, bool bBlockUntilCompleted)
-{
-	return TexRealloc_Succeeded;
-}
-
-
 ///////////////////////////////////////////////////////////////////////////////////////////
 // FD3D12Texture
 ///////////////////////////////////////////////////////////////////////////////////////////
diff --git a/Engine/Source/Runtime/Engine/Private/Streaming/Texture2DMipAllocator_AsyncReallocate.cpp b/Engine/Source/Runtime/Engine/Private/Streaming/Texture2DMipAllocator_AsyncReallocate.cpp
--- a/Engine/Source/Runtime/Engine/Private/Streaming/Texture2DMipAllocator_AsyncReallocate.cpp
+++ b/Engine/Source/Runtime/Engine/Private/Streaming/Texture2DMipAllocator_AsyncReallocate.cpp
@@ -62,8 +62,6 @@ bool FTexture2DMipAllocator_AsyncReallocate::AllocateMips(
 	{
 		const bool bFlushRHIThread = CVarFlushRHIThreadOnSTreamingTextureLocks.GetValueOnAnyThread() > 0;
 
-		RHIFinalizeAsyncReallocateTexture2D(IntermediateTextureRHI, true);
-
 		OutMipInfos.AddDefaulted(CurrentFirstLODIdx);
 
 		for (int32 MipIndex = PendingFirstLODIdx; MipIndex < CurrentFirstLODIdx; ++MipIndex)
diff --git a/Engine/Source/Runtime/Engine/Private/Streaming/Texture2DStreamIn_DDC_AsyncReallocate.cpp b/Engine/Source/Runtime/Engine/Private/Streaming/Texture2DStreamIn_DDC_AsyncReallocate.cpp
--- a/Engine/Source/Runtime/Engine/Private/Streaming/Texture2DStreamIn_DDC_AsyncReallocate.cpp
+++ b/Engine/Source/Runtime/Engine/Private/Streaming/Texture2DStreamIn_DDC_AsyncReallocate.cpp
@@ -35,7 +35,6 @@ void FTexture2DStreamIn_DDC_AsyncReallocate::LockMips(const FContext& Context)
 	DECLARE_SCOPE_CYCLE_COUNTER(TEXT("FTexture2DStreamIn_DDC_AsyncReallocate::LockMips"), STAT_Texture2DStreamInDDCAsyncReallocate_LockMips, STATGROUP_StreamingDetails);
 	check(Context.CurrentThread == TT_Render);
 
-	RHIFinalizeAsyncReallocateTexture2D(IntermediateTextureRHI, true);
 	DoLockNewMips(Context);
 
 	PushTask(Context, TT_Async, SRA_UPDATE_CALLBACK(LoadMips), TT_Render, SRA_UPDATE_CALLBACK(Cancel));
diff --git a/Engine/Source/Runtime/Engine/Private/Streaming/Texture2DStreamIn_DerivedData.cpp b/Engine/Source/Runtime/Engine/Private/Streaming/Texture2DStreamIn_DerivedData.cpp
--- a/Engine/Source/Runtime/Engine/Private/Streaming/Texture2DStreamIn_DerivedData.cpp
+++ b/Engine/Source/Runtime/Engine/Private/Streaming/Texture2DStreamIn_DerivedData.cpp
@@ -216,7 +216,6 @@ void FTexture2DStreamIn_DerivedData_AsyncReallocate::LockMips(const FContext& Co
 		STAT_Texture2DStreamIn_DerivedData_AsyncReallocate_LockMips, STATGROUP_StreamingDetails);
 	check(Context.CurrentThread == TT_Render);
 
-	RHIFinalizeAsyncReallocateTexture2D(IntermediateTextureRHI, true);
 	DoLockNewMips(Context);
 
 	PushTask(Context, TT_Async, SRA_UPDATE_CALLBACK(LoadMips), TT_Render, SRA_UPDATE_CALLBACK(Cancel));
diff --git a/Engine/Source/Runtime/Engine/Private/Streaming/Texture2DStreamIn_IO_AsyncReallocate.cpp b/Engine/Source/Runtime/Engine/Private/Streaming/Texture2DStreamIn_IO_AsyncReallocate.cpp
--- a/Engine/Source/Runtime/Engine/Private/Streaming/Texture2DStreamIn_IO_AsyncReallocate.cpp
+++ b/Engine/Source/Runtime/Engine/Private/Streaming/Texture2DStreamIn_IO_AsyncReallocate.cpp
@@ -33,7 +33,6 @@ void FTexture2DStreamIn_IO_AsyncReallocate::LockMips(const FContext& Context)
 	DECLARE_SCOPE_CYCLE_COUNTER(TEXT("FTexture2DStreamIn_IO_AsyncReallocate::LockMips"), STAT_Texture2DStreamInIOAsyncReallocate_LockMips, STATGROUP_StreamingDetails);
 	check(Context.CurrentThread == TT_Render);
 
-	RHIFinalizeAsyncReallocateTexture2D(IntermediateTextureRHI, true);
 	DoLockNewMips(Context);
 
 	PushTask(Context, TT_Async, SRA_UPDATE_CALLBACK(LoadMips), TT_Render, SRA_UPDATE_CALLBACK(Cancel));
diff --git a/Engine/Source/Runtime/Engine/Private/Streaming/Texture2DStreamOut_AsyncCreate.cpp b/Engine/Source/Runtime/Engine/Private/Streaming/Texture2DStreamOut_AsyncCreate.cpp
--- a/Engine/Source/Runtime/Engine/Private/Streaming/Texture2DStreamOut_AsyncCreate.cpp
+++ b/Engine/Source/Runtime/Engine/Private/Streaming/Texture2DStreamOut_AsyncCreate.cpp
@@ -76,7 +76,6 @@ void FTexture2DStreamOut_AsyncCreate::Finalize(const FContext& Context)
 			IntermediateTextureRHI);
 	}
 
-	RHIFinalizeAsyncReallocateTexture2D(IntermediateTextureRHI, true);
 	DoFinishUpdate(Context);
 }
 
diff --git a/Engine/Source/Runtime/Engine/Private/Streaming/Texture2DStreamOut_AsyncReallocate.cpp b/Engine/Source/Runtime/Engine/Private/Streaming/Texture2DStreamOut_AsyncReallocate.cpp
--- a/Engine/Source/Runtime/Engine/Private/Streaming/Texture2DStreamOut_AsyncReallocate.cpp
+++ b/Engine/Source/Runtime/Engine/Private/Streaming/Texture2DStreamOut_AsyncReallocate.cpp
@@ -37,7 +37,6 @@ void FTexture2DStreamOut_AsyncReallocate::Finalize(const FContext& Context)
 	DECLARE_SCOPE_CYCLE_COUNTER(TEXT("FTexture2DStreamOut_AsyncReallocate::Finalize"), STAT_Texture2DStreamOutAsyncReallocate_Finalize, STATGROUP_StreamingDetails);
 	check(Context.CurrentThread == TT_Render);
 
-	RHIFinalizeAsyncReallocateTexture2D(IntermediateTextureRHI, true);
 	DoFinishUpdate(Context);
 }
 
diff --git a/Engine/Source/Runtime/NullDrv/Public/NullRHI.h b/Engine/Source/Runtime/NullDrv/Public/NullRHI.h
--- a/Engine/Source/Runtime/NullDrv/Public/NullRHI.h
+++ b/Engine/Source/Runtime/NullDrv/Public/NullRHI.h
@@ -215,14 +215,6 @@ public:
 
 		return new FNullTexture(Desc);
 	}
-	virtual ETextureReallocationStatus RHIFinalizeAsyncReallocateTexture2D(FRHITexture* Texture2D, bool bBlockUntilCompleted) final override
-	{ 
-		return TexRealloc_Succeeded; 
-	}
-	virtual ETextureReallocationStatus RHICancelAsyncReallocateTexture2D(FRHITexture* Texture2D, bool bBlockUntilCompleted) final override
-	{ 
-		return TexRealloc_Succeeded; 
-	}
 	virtual void* RHILockTexture2D(FRHITexture* Texture, uint32 MipIndex, EResourceLockMode LockMode, uint32& DestStride, bool bLockWithinMiptail, uint64* OutLockedByteCount) final override
 	{ 
 		DestStride = 0; 
diff --git a/Engine/Source/Runtime/OpenGLDrv/Private/OpenGLDrv.h b/Engine/Source/Runtime/OpenGLDrv/Private/OpenGLDrv.h
--- a/Engine/Source/Runtime/OpenGLDrv/Private/OpenGLDrv.h
+++ b/Engine/Source/Runtime/OpenGLDrv/Private/OpenGLDrv.h
@@ -352,8 +352,6 @@ public:
 	virtual void RHIGenerateMips(FRHITexture* Texture) final override;
 	virtual uint32 RHIComputeMemorySize(FRHITexture* TextureRHI) final override;
 	virtual FTextureRHIRef RHIAsyncReallocateTexture2D(FRHITexture* Texture2D, int32 NewMipCount, int32 NewSizeX, int32 NewSizeY, FThreadSafeCounter* RequestStatus) final override;
-	virtual ETextureReallocationStatus RHIFinalizeAsyncReallocateTexture2D(FRHITexture* Texture2D, bool bBlockUntilCompleted) final override;
-	virtual ETextureReallocationStatus RHICancelAsyncReallocateTexture2D(FRHITexture* Texture2D, bool bBlockUntilCompleted) final override;
 	virtual void* RHILockTexture2D(FRHITexture* Texture, uint32 MipIndex, EResourceLockMode LockMode, uint32& DestStride, bool bLockWithinMiptail, uint64* OutLockedByteCount = nullptr) final override;
 	virtual void RHIUnlockTexture2D(FRHITexture* Texture, uint32 MipIndex, bool bLockWithinMiptail) final override;
 	virtual void* RHILockTexture2DArray(FRHITexture* Texture, uint32 TextureIndex, uint32 MipIndex, EResourceLockMode LockMode, uint32& DestStride, bool bLockWithinMiptail) final override;
@@ -652,16 +650,6 @@ public:
 
 	virtual FTextureRHIRef AsyncReallocateTexture2D_RenderThread(class FRHICommandListImmediate& RHICmdList, FRHITexture* Texture2D, int32 NewMipCount, int32 NewSizeX, int32 NewSizeY, FThreadSafeCounter* RequestStatus) final override;
 
-	virtual ETextureReallocationStatus FinalizeAsyncReallocateTexture2D_RenderThread(class FRHICommandListImmediate& RHICmdList, FRHITexture* Texture2D, bool bBlockUntilCompleted) final override
-	{
-		return this->RHIFinalizeAsyncReallocateTexture2D(Texture2D, bBlockUntilCompleted);
-	}
-
-	virtual ETextureReallocationStatus CancelAsyncReallocateTexture2D_RenderThread(class FRHICommandListImmediate& RHICmdList, FRHITexture* Texture2D, bool bBlockUntilCompleted) final override
-	{
-		return this->RHICancelAsyncReallocateTexture2D(Texture2D, bBlockUntilCompleted);
-	}
-
 	virtual FSamplerStateRHIRef RHICreateSamplerState(const FSamplerStateInitializerRHI& Initializer) final override;
 
 	virtual FGraphicsPipelineStateRHIRef RHICreateGraphicsPipelineState(const FGraphicsPipelineStateInitializer& Initializer) override
diff --git a/Engine/Source/Runtime/OpenGLDrv/Private/OpenGLTexture.cpp b/Engine/Source/Runtime/OpenGLDrv/Private/OpenGLTexture.cpp
--- a/Engine/Source/Runtime/OpenGLDrv/Private/OpenGLTexture.cpp
+++ b/Engine/Source/Runtime/OpenGLDrv/Private/OpenGLTexture.cpp
@@ -1463,33 +1463,6 @@ FTextureRHIRef FOpenGLDynamicRHI::RHIAsyncReallocateTexture2D(FRHITexture* Textu
 	return NewTexture;
 }
 
-/**
- * Returns the status of an ongoing or completed texture reallocation:
- *	TexRealloc_Succeeded	- The texture is ok, reallocation is not in progress.
- *	TexRealloc_Failed		- The texture is bad, reallocation is not in progress.
- *	TexRealloc_InProgress	- The texture is currently being reallocated async.
- *
- * @param Texture2D		- Texture to check the reallocation status for
- * @return				- Current reallocation status
- */
-ETextureReallocationStatus FOpenGLDynamicRHI::RHIFinalizeAsyncReallocateTexture2D(FRHITexture* Texture2D, bool bBlockUntilCompleted )
-{
-	return TexRealloc_Succeeded;
-}
-
-/**
- * Cancels an async reallocation for the specified texture.
- * This should be called for the new texture, not the original.
- *
- * @param Texture				Texture to cancel
- * @param bBlockUntilCompleted	If true, blocks until the cancellation is fully completed
- * @return						Reallocation status
- */
-ETextureReallocationStatus FOpenGLDynamicRHI::RHICancelAsyncReallocateTexture2D(FRHITexture* Texture2D, bool bBlockUntilCompleted )
-{
-	return TexRealloc_Succeeded;
-}
-
 void* FOpenGLDynamicRHI::RHILockTexture2D(FRHITexture* TextureRHI,uint32 MipIndex,EResourceLockMode LockMode,uint32& DestStride,bool bLockWithinMiptail, uint64* OutLockedByteCount)
 {
 	FOpenGLTexture* Texture = ResourceCast(TextureRHI);
diff --git a/Engine/Source/Runtime/RHI/Private/RHICommandList.cpp b/Engine/Source/Runtime/RHI/Private/RHICommandList.cpp
--- a/Engine/Source/Runtime/RHI/Private/RHICommandList.cpp
+++ b/Engine/Source/Runtime/RHI/Private/RHICommandList.cpp
@@ -1964,20 +1964,6 @@ FTextureRHIRef FDynamicRHI::AsyncReallocateTexture2D_RenderThread(class FRHIComm
 	return GDynamicRHI->RHIAsyncReallocateTexture2D(Texture2D, NewMipCount, NewSizeX, NewSizeY, RequestStatus);
 }
 
-ETextureReallocationStatus FDynamicRHI::FinalizeAsyncReallocateTexture2D_RenderThread(class FRHICommandListImmediate& RHICmdList, FRHITexture* Texture2D, bool bBlockUntilCompleted)
-{
-	CSV_SCOPED_TIMING_STAT(RHITStalls, FinalizeAsyncReallocateTexture2D_RenderThread);
-	FScopedRHIThreadStaller StallRHIThread(RHICmdList);
-	return GDynamicRHI->RHIFinalizeAsyncReallocateTexture2D(Texture2D, bBlockUntilCompleted);
-}
-
-ETextureReallocationStatus FDynamicRHI::CancelAsyncReallocateTexture2D_RenderThread(class FRHICommandListImmediate& RHICmdList, FRHITexture* Texture2D, bool bBlockUntilCompleted)
-{
-	CSV_SCOPED_TIMING_STAT(RHITStalls, CancelAsyncReallocateTexture2D_RenderThread);
-	FScopedRHIThreadStaller StallRHIThread(RHICmdList);
-	return GDynamicRHI->RHICancelAsyncReallocateTexture2D(Texture2D, bBlockUntilCompleted);
-}
-
 FUpdateTexture3DData FDynamicRHI::RHIBeginUpdateTexture3D(FRHICommandListBase& RHICmdList, FRHITexture* Texture, uint32 MipIndex, const struct FUpdateTextureRegion3D& UpdateRegion)
 {
 	check(IsInParallelRenderingThread());
diff --git a/Engine/Source/Runtime/RHI/Public/DynamicRHI.h b/Engine/Source/Runtime/RHI/Public/DynamicRHI.h
--- a/Engine/Source/Runtime/RHI/Public/DynamicRHI.h
+++ b/Engine/Source/Runtime/RHI/Public/DynamicRHI.h
@@ -478,31 +478,6 @@ public:
 	// NP: Note that no RHI currently implements this as an async call, we should simplify the API.
 	virtual FTextureRHIRef RHIAsyncReallocateTexture2D(FRHITexture* Texture2D, int32 NewMipCount, int32 NewSizeX, int32 NewSizeY, FThreadSafeCounter* RequestStatus) = 0;
 
-	/**
-	* Finalizes an async reallocation request.
-	* If bBlockUntilCompleted is false, it will only poll the status and finalize if the reallocation has completed.
-	*
-	* @param Texture2D				- Texture to finalize the reallocation for
-	* @param bBlockUntilCompleted	- Whether the function should block until the reallocation has completed
-	* @return						- Current reallocation status:
-	*	TexRealloc_Succeeded	Reallocation succeeded
-	*	TexRealloc_Failed		Reallocation failed
-	*	TexRealloc_InProgress	Reallocation is still in progress, try again later
-	*/
-	// FlushType: Wait RHI Thread
-	virtual ETextureReallocationStatus RHIFinalizeAsyncReallocateTexture2D(FRHITexture* Texture2D, bool bBlockUntilCompleted) = 0;
-
-	/**
-	* Cancels an async reallocation for the specified texture.
-	* This should be called for the new texture, not the original.
-	*
-	* @param Texture				Texture to cancel
-	* @param bBlockUntilCompleted	If true, blocks until the cancellation is fully completed
-	* @return						Reallocation status
-	*/
-	// FlushType: Wait RHI Thread
-	virtual ETextureReallocationStatus RHICancelAsyncReallocateTexture2D(FRHITexture* Texture2D, bool bBlockUntilCompleted) = 0;
-
 	/**
 	* Locks an RHI texture's mip-map for read/write operations on the CPU
 	* @param Texture - the RHI texture resource to lock, must not be 0
@@ -899,8 +874,6 @@ public:
 	virtual void RHIProcessDeleteQueue() {}
 
 	RHI_API virtual FTextureRHIRef AsyncReallocateTexture2D_RenderThread(class FRHICommandListImmediate& RHICmdList, FRHITexture* Texture2D, int32 NewMipCount, int32 NewSizeX, int32 NewSizeY, FThreadSafeCounter* RequestStatus);
-	RHI_API virtual ETextureReallocationStatus FinalizeAsyncReallocateTexture2D_RenderThread(class FRHICommandListImmediate& RHICmdList, FRHITexture* Texture2D, bool bBlockUntilCompleted);
-	RHI_API virtual ETextureReallocationStatus CancelAsyncReallocateTexture2D_RenderThread(class FRHICommandListImmediate& RHICmdList, FRHITexture* Texture2D, bool bBlockUntilCompleted);
 
 	RHI_API virtual void* LockTexture2D_RenderThread(class FRHICommandListImmediate& RHICmdList, FRHITexture* Texture, uint32 MipIndex, EResourceLockMode LockMode, uint32& DestStride, bool bLockWithinMiptail, bool bNeedsDefaultRHIFlush = true, uint64* OutLockedByteCount = nullptr);
 	RHI_API virtual void UnlockTexture2D_RenderThread(class FRHICommandListImmediate& RHICmdList, FRHITexture* Texture, uint32 MipIndex, bool bLockWithinMiptail, bool bNeedsDefaultRHIFlush = true);
diff --git a/Engine/Source/Runtime/RHI/Public/RHICommandList.h b/Engine/Source/Runtime/RHI/Public/RHICommandList.h
--- a/Engine/Source/Runtime/RHI/Public/RHICommandList.h
+++ b/Engine/Source/Runtime/RHI/Public/RHICommandList.h
@@ -4924,15 +4924,17 @@ public:
 		return GDynamicRHI->AsyncReallocateTexture2D_RenderThread(*this, Texture2D, NewMipCount, NewSizeX, NewSizeY, RequestStatus);
 	}
 	
+	UE_DEPRECATED(5.6, "FinalizeAsyncReallocateTexture2D is no longer implemented.")
 	FORCEINLINE ETextureReallocationStatus FinalizeAsyncReallocateTexture2D(FRHITexture* Texture2D, bool bBlockUntilCompleted)
 	{
 		LLM_SCOPE(ELLMTag::Textures);
-		return GDynamicRHI->FinalizeAsyncReallocateTexture2D_RenderThread(*this, Texture2D, bBlockUntilCompleted);
+		return TexRealloc_Succeeded;
 	}
 	
+	UE_DEPRECATED(5.6, "CancelAsyncReallocateTexture2D is no longer implemented.")
 	FORCEINLINE ETextureReallocationStatus CancelAsyncReallocateTexture2D(FRHITexture* Texture2D, bool bBlockUntilCompleted)
 	{
-		return GDynamicRHI->CancelAsyncReallocateTexture2D_RenderThread(*this, Texture2D, bBlockUntilCompleted);
+		return TexRealloc_Succeeded;
 	}
 	
 	FORCEINLINE void* LockTexture2D(FRHITexture* Texture, uint32 MipIndex, EResourceLockMode LockMode, uint32& DestStride, bool bLockWithinMiptail, bool bFlushRHIThread = true, uint64* OutLockedByteCount = nullptr)
@@ -5728,14 +5730,16 @@ FORCEINLINE FTextureRHIRef RHIAsyncReallocateTexture2D(FRHITexture* Texture2D, i
 	return FRHICommandListExecutor::GetImmediateCommandList().AsyncReallocateTexture2D(Texture2D, NewMipCount, NewSizeX, NewSizeY, RequestStatus);
 }
 
+UE_DEPRECATED(5.6, "RHIFinalizeAsyncReallocateTexture2D is no longer implemented.")
 FORCEINLINE ETextureReallocationStatus RHIFinalizeAsyncReallocateTexture2D(FRHITexture* Texture2D, bool bBlockUntilCompleted)
 {
-	return FRHICommandListExecutor::GetImmediateCommandList().FinalizeAsyncReallocateTexture2D(Texture2D, bBlockUntilCompleted);
+	return TexRealloc_Succeeded;
 }
 
+UE_DEPRECATED(5.6, "RHICancelAsyncReallocateTexture2D is no longer implemented.")
 FORCEINLINE ETextureReallocationStatus RHICancelAsyncReallocateTexture2D(FRHITexture* Texture2D, bool bBlockUntilCompleted)
 {
-	return FRHICommandListExecutor::GetImmediateCommandList().CancelAsyncReallocateTexture2D(Texture2D, bBlockUntilCompleted);
+	return TexRealloc_Succeeded;
 }
 
 FORCEINLINE void* RHILockTexture2D(FRHITexture* Texture, uint32 MipIndex, EResourceLockMode LockMode, uint32& DestStride, bool bLockWithinMiptail, bool bFlushRHIThread = true, uint64* OutLockedByteCount = nullptr)
diff --git a/Engine/Source/Runtime/RHI/Public/RHIValidation.h b/Engine/Source/Runtime/RHI/Public/RHIValidation.h
--- a/Engine/Source/Runtime/RHI/Public/RHIValidation.h
+++ b/Engine/Source/Runtime/RHI/Public/RHIValidation.h
@@ -592,37 +592,6 @@ public:
 		return NewTexture2D;
 	}
 
-	/**
-	* Finalizes an async reallocation request.
-	* If bBlockUntilCompleted is false, it will only poll the status and finalize if the reallocation has completed.
-	*
-	* @param Texture2D				- Texture to finalize the reallocation for
-	* @param bBlockUntilCompleted	- Whether the function should block until the reallocation has completed
-	* @return						- Current reallocation status:
-	*	TexRealloc_Succeeded	Reallocation succeeded
-	*	TexRealloc_Failed		Reallocation failed
-	*	TexRealloc_InProgress	Reallocation is still in progress, try again later
-	*/
-	// FlushType: Wait RHI Thread
-	virtual ETextureReallocationStatus RHIFinalizeAsyncReallocateTexture2D(FRHITexture* Texture2D, bool bBlockUntilCompleted) override final
-	{
-		return RHI->RHIFinalizeAsyncReallocateTexture2D(Texture2D, bBlockUntilCompleted);
-	}
-
-	/**
-	* Cancels an async reallocation for the specified texture.
-	* This should be called for the new texture, not the original.
-	*
-	* @param Texture				Texture to cancel
-	* @param bBlockUntilCompleted	If true, blocks until the cancellation is fully completed
-	* @return						Reallocation status
-	*/
-	// FlushType: Wait RHI Thread
-	virtual ETextureReallocationStatus RHICancelAsyncReallocateTexture2D(FRHITexture* Texture2D, bool bBlockUntilCompleted) override final
-	{
-		return RHI->RHICancelAsyncReallocateTexture2D(Texture2D, bBlockUntilCompleted);
-	}
-
 	/**
 	* Locks an RHI texture's mip-map for read/write operations on the CPU
 	* @param Texture - the RHI texture resource to lock, must not be 0
@@ -1113,16 +1082,6 @@ public:
 		return NewTexture2D;
 	}
 
-	virtual ETextureReallocationStatus FinalizeAsyncReallocateTexture2D_RenderThread(class FRHICommandListImmediate& RHICmdList, FRHITexture* Texture2D, bool bBlockUntilCompleted) override final
-	{
-		return RHI->FinalizeAsyncReallocateTexture2D_RenderThread(RHICmdList, Texture2D, bBlockUntilCompleted);
-	}
-
-	virtual ETextureReallocationStatus CancelAsyncReallocateTexture2D_RenderThread(class FRHICommandListImmediate& RHICmdList, FRHITexture* Texture2D, bool bBlockUntilCompleted) override final
-	{
-		return RHI->CancelAsyncReallocateTexture2D_RenderThread(RHICmdList, Texture2D, bBlockUntilCompleted);
-	}
-
 	virtual void* LockBuffer_BottomOfPipe(class FRHICommandListBase& RHICmdList, FRHIBuffer* Buffer, uint32 Offset, uint32 SizeRHI, EResourceLockMode LockMode) override final
 	{
 		RHI_VALIDATION_CHECK(LockMode != RLM_WriteOnly_NoOverwrite || GRHISupportsMapWriteNoOverwrite, TEXT("Using RLM_WriteOnly_NoOverwrite when the RHI doesn't support it."));
diff --git a/Engine/Source/Runtime/VulkanRHI/Private/VulkanTexture.cpp b/Engine/Source/Runtime/VulkanRHI/Private/VulkanTexture.cpp
--- a/Engine/Source/Runtime/VulkanRHI/Private/VulkanTexture.cpp
+++ b/Engine/Source/Runtime/VulkanRHI/Private/VulkanTexture.cpp
@@ -1205,16 +1205,6 @@ FTextureRHIRef FVulkanDynamicRHI::RHIAsyncReallocateTexture2D(FRHITexture* OldTe
 	return NewTexture;
 }
 
-ETextureReallocationStatus FVulkanDynamicRHI::RHIFinalizeAsyncReallocateTexture2D(FRHITexture* Texture2D, bool bBlockUntilCompleted)
-{
-	return TexRealloc_Succeeded;
-}
-
-ETextureReallocationStatus FVulkanDynamicRHI::RHICancelAsyncReallocateTexture2D(FRHITexture* Texture2D, bool bBlockUntilCompleted)
-{
-	return TexRealloc_Succeeded;
-}
-
 void* FVulkanDynamicRHI::RHILockTexture2D(FRHITexture* TextureRHI, uint32 MipIndex, EResourceLockMode LockMode, uint32& DestStride, bool bLockWithinMiptail, uint64* OutLockedByteCount)
 {
 	LLM_SCOPE_VULKAN(ELLMTagVulkan::VulkanTextures);
diff --git a/Engine/Source/Runtime/VulkanRHI/Public/VulkanDynamicRHI.h b/Engine/Source/Runtime/VulkanRHI/Public/VulkanDynamicRHI.h
--- a/Engine/Source/Runtime/VulkanRHI/Public/VulkanDynamicRHI.h
+++ b/Engine/Source/Runtime/VulkanRHI/Public/VulkanDynamicRHI.h
@@ -138,13 +138,7 @@ public:
 	virtual FTextureRHIRef RHIAsyncCreateTexture2D(uint32 SizeX, uint32 SizeY, uint8 Format, uint32 NumMips, ETextureCreateFlags Flags, ERHIAccess InResourceState, void** InitialMipData, uint32 NumInitialMips, const TCHAR* DebugName, FGraphEventRef& OutCompletionEvent) final override;
 	virtual uint32 RHIComputeMemorySize(FRHITexture* TextureRHI) final override;
 	virtual FTextureRHIRef RHIAsyncReallocateTexture2D(FRHITexture* Texture2D, int32 NewMipCount, int32 NewSizeX, int32 NewSizeY, FThreadSafeCounter* RequestStatus) final override;
-	virtual ETextureReallocationStatus RHIFinalizeAsyncReallocateTexture2D(FRHITexture* Texture2D, bool bBlockUntilCompleted) final override;
-	virtual ETextureReallocationStatus FinalizeAsyncReallocateTexture2D_RenderThread(class FRHICommandListImmediate& RHICmdList, FRHITexture* Texture2D, bool bBlockUntilCompleted) final override
-	{
-		return this->RHIFinalizeAsyncReallocateTexture2D(Texture2D, bBlockUntilCompleted);
-	}
 
-	virtual ETextureReallocationStatus RHICancelAsyncReallocateTexture2D(FRHITexture* Texture2D, bool bBlockUntilCompleted) final override;
 	virtual void* RHILockTexture2D(FRHITexture* Texture, uint32 MipIndex, EResourceLockMode LockMode, uint32& DestStride, bool bLockWithinMiptail, uint64* OutLockedByteCount = nullptr) final override;
 	virtual void RHIUnlockTexture2D(FRHITexture* Texture, uint32 MipIndex, bool bLockWithinMiptail) final override
 	{
@@ -210,7 +204,6 @@ public:
 	virtual bool RHIMatchPrecachePSOInitializers(const FGraphicsPipelineStateInitializer& LHS, const FGraphicsPipelineStateInitializer& RHS) final override;
 
 	virtual FTextureRHIRef AsyncReallocateTexture2D_RenderThread(class FRHICommandListImmediate& RHICmdList, FRHITexture* Texture2D, int32 NewMipCount, int32 NewSizeX, int32 NewSizeY, FThreadSafeCounter* RequestStatus) override final;
-	//virtual ETextureReallocationStatus CancelAsyncReallocateTexture2D_RenderThread(class FRHICommandListImmediate& RHICmdList, FRHITexture* Texture2D, bool bBlockUntilCompleted) override final;
 	
 	virtual void* LockTexture2D_RenderThread(class FRHICommandListImmediate& RHICmdList, FRHITexture* Texture, uint32 MipIndex, EResourceLockMode LockMode, uint32& DestStride, bool bLockWithinMiptail, bool bNeedsDefaultRHIFlush = true, uint64* OutLockedByteCount = nullptr) override final
 	{
diff --git a/Engine/Source/Runtime/Windows/D3D11RHI/Private/D3D11RHIPrivate.h b/Engine/Source/Runtime/Windows/D3D11RHI/Private/D3D11RHIPrivate.h
--- a/Engine/Source/Runtime/Windows/D3D11RHI/Private/D3D11RHIPrivate.h
+++ b/Engine/Source/Runtime/Windows/D3D11RHI/Private/D3D11RHIPrivate.h
@@ -453,8 +453,6 @@ public:
 	virtual uint32 RHIComputeMemorySize(FRHITexture* TextureRHI) final override;
 	virtual void RHIAsyncCopyTexture2DCopy(FRHITexture* NewTexture2DRHI, FRHITexture* Texture2DRHI, int32 NewMipCount, int32 NewSizeX, int32 NewSizeY, FThreadSafeCounter* RequestStatus);
 	virtual FTextureRHIRef RHIAsyncReallocateTexture2D(FRHITexture* Texture2D, int32 NewMipCount, int32 NewSizeX, int32 NewSizeY, FThreadSafeCounter* RequestStatus) final override;
-	virtual ETextureReallocationStatus RHIFinalizeAsyncReallocateTexture2D(FRHITexture* Texture2D, bool bBlockUntilCompleted) final override;
-	virtual ETextureReallocationStatus RHICancelAsyncReallocateTexture2D(FRHITexture* Texture2D, bool bBlockUntilCompleted) final override;
 	virtual void* RHILockTexture2D(FRHITexture* Texture, uint32 MipIndex, EResourceLockMode LockMode, uint32& DestStride, bool bLockWithinMiptail, uint64* OutLockedByteCount) final override;
 	virtual void RHIUnlockTexture2D(FRHITexture* Texture, uint32 MipIndex, bool bLockWithinMiptail) final override;
 	virtual void* RHILockTexture2DArray(FRHITexture* Texture, uint32 TextureIndex, uint32 MipIndex, EResourceLockMode LockMode, uint32& DestStride, bool bLockWithinMiptail) final override;
@@ -572,8 +570,6 @@ public:
 	// For example, deferring the actual creation to RHI thread by sending an RHI command.
 	// For D3D11, these functions mainly just remove RHIT stalls because ID3D11Device is thread safe.
 	virtual FTextureRHIRef AsyncReallocateTexture2D_RenderThread(class FRHICommandListImmediate& RHICmdList, FRHITexture* Texture2D, int32 NewMipCount, int32 NewSizeX, int32 NewSizeY, FThreadSafeCounter* RequestStatus) final override;
-	virtual ETextureReallocationStatus FinalizeAsyncReallocateTexture2D_RenderThread(class FRHICommandListImmediate& RHICmdList, FRHITexture* Texture2D, bool bBlockUntilCompleted) final override;
-	virtual ETextureReallocationStatus CancelAsyncReallocateTexture2D_RenderThread(class FRHICommandListImmediate& RHICmdList, FRHITexture* Texture2D, bool bBlockUntilCompleted) final override;
 	virtual void* LockTexture2D_RenderThread(class FRHICommandListImmediate& RHICmdList, FRHITexture* Texture, uint32 MipIndex, EResourceLockMode LockMode, uint32& DestStride, bool bLockWithinMiptail, bool bNeedsDefaultRHIFlush = true, uint64* OutLockedByteCount = nullptr) final override;
 	virtual void UnlockTexture2D_RenderThread(class FRHICommandListImmediate& RHICmdList, FRHITexture* Texture, uint32 MipIndex, bool bLockWithinMiptail, bool bNeedsDefaultRHIFlush = true) final override;
 	virtual void RHIEndUpdateTexture3D(FRHICommandListBase& RHICmdList, FUpdateTexture3DData& UpdateData) final override;
diff --git a/Engine/Source/Runtime/Windows/D3D11RHI/Private/D3D11Texture.cpp b/Engine/Source/Runtime/Windows/D3D11RHI/Private/D3D11Texture.cpp
--- a/Engine/Source/Runtime/Windows/D3D11RHI/Private/D3D11Texture.cpp
+++ b/Engine/Source/Runtime/Windows/D3D11RHI/Private/D3D11Texture.cpp
@@ -986,49 +986,6 @@ FTextureRHIRef FD3D11DynamicRHI::AsyncReallocateTexture2D_RenderThread(
 	return NewTexture2D;
 }
 
-/**
- * Returns the status of an ongoing or completed texture reallocation:
- *	TexRealloc_Succeeded	- The texture is ok, reallocation is not in progress.
- *	TexRealloc_Failed		- The texture is bad, reallocation is not in progress.
- *	TexRealloc_InProgress	- The texture is currently being reallocated async.
- *
- * @param Texture2D		- Texture to check the reallocation status for
- * @return				- Current reallocation status
- */
-ETextureReallocationStatus FD3D11DynamicRHI::RHIFinalizeAsyncReallocateTexture2D(FRHITexture* Texture2D, bool bBlockUntilCompleted )
-{
-	return TexRealloc_Succeeded;
-}
-
-ETextureReallocationStatus FD3D11DynamicRHI::FinalizeAsyncReallocateTexture2D_RenderThread(
-	class FRHICommandListImmediate& RHICmdList,
-	FRHITexture* Texture2D,
-	bool bBlockUntilCompleted)
-{
-	return RHIFinalizeAsyncReallocateTexture2D(Texture2D, bBlockUntilCompleted);
-}
-
-/**
- * Cancels an async reallocation for the specified texture.
- * This should be called for the new texture, not the original.
- *
- * @param Texture				Texture to cancel
- * @param bBlockUntilCompleted	If true, blocks until the cancellation is fully completed
- * @return						Reallocation status
- */
-ETextureReallocationStatus FD3D11DynamicRHI::RHICancelAsyncReallocateTexture2D(FRHITexture* Texture2D, bool bBlockUntilCompleted )
-{
-	return TexRealloc_Succeeded;
-}
-
-ETextureReallocationStatus FD3D11DynamicRHI::CancelAsyncReallocateTexture2D_RenderThread(
-	class FRHICommandListImmediate& RHICmdList,
-	FRHITexture* Texture2D,
-	bool bBlockUntilCompleted)
-{
-	return RHICancelAsyncReallocateTexture2D(Texture2D, bBlockUntilCompleted);
-}
-
 void* FD3D11Texture::Lock(FD3D11DynamicRHI* D3DRHI, uint32 MipIndex, uint32 ArrayIndex, EResourceLockMode LockMode, uint32& DestStride, bool bForceLockDeferred, uint64* OutLockedByteCount)
 {
 	check(!IsTexture3D()); // Only 2D texture locks are implemented
