Fixing -RHIValidationLogStack commnand line parsing

diff --git a/Engine/Source/Runtime/RHI/Private/RHIValidation.cpp b/Engine/Source/Runtime/RHI/Private/RHIValidation.cpp
--- a/Engine/Source/Runtime/RHI/Private/RHIValidation.cpp
+++ b/Engine/Source/Runtime/RHI/Private/RHIValidation.cpp
@@ -75,7 +75,7 @@ namespace RHIValidation
 					Strings.Add(ResourceNames);
 				}
 				
-				GRHIValidationPrintHumanReadableCallStack = FParse::Param(FCommandLine::Get(), TEXT("-RHIValidationLogStack"));
+				GRHIValidationPrintHumanReadableCallStack = FParse::Param(FCommandLine::Get(), TEXT("RHIValidationLogStack"));
 			}
 		} static Init;
 
