Include Feature Level for Mac in Editor

#tests Run the editor in SM5

diff --git a/Engine/Source/Runtime/RHI/Private/Apple/AppleDynamicRHI.cpp b/Engine/Source/Runtime/RHI/Private/Apple/AppleDynamicRHI.cpp
--- a/Engine/Source/Runtime/RHI/Private/Apple/AppleDynamicRHI.cpp
+++ b/Engine/Source/Runtime/RHI/Private/Apple/AppleDynamicRHI.cpp
@@ -143,7 +143,10 @@ static IDynamicRHIModule* LoadDynamicRHIModule(ERHIFeatureLevel::Type& Requested
 
 	ComputeRequestedFeatureLevel(Options, RequestedFeatureLevel);
 
-	FApp::SetGraphicsRHI(TEXT("Metal"));
+	FString FeatureLevelName;
+	GetFeatureLevelName(RequestedFeatureLevel, FeatureLevelName);
+	FApp::SetGraphicsRHI(FString::Printf(TEXT("Metal (%s)"), *FeatureLevelName));
+	
 	DynamicRHIModule = &FModuleManager::LoadModuleChecked<IDynamicRHIModule>(TEXT("MetalRHI"));
 
 	return DynamicRHIModule;
