Fix Mac Editor Feature Level not matching actual feature level

#tests Run the editor on a project supporting SM6 but on a system supporting only SM5

diff --git a/Engine/Source/Runtime/RHI/Private/Apple/AppleDynamicRHI.cpp b/Engine/Source/Runtime/RHI/Private/Apple/AppleDynamicRHI.cpp
--- a/Engine/Source/Runtime/RHI/Private/Apple/AppleDynamicRHI.cpp
+++ b/Engine/Source/Runtime/RHI/Private/Apple/AppleDynamicRHI.cpp
@@ -142,10 +142,6 @@ static IDynamicRHIModule* LoadDynamicRHIModule(ERHIFeatureLevel::Type& Requested
 	}
 
 	ComputeRequestedFeatureLevel(Options, RequestedFeatureLevel);
-
-	FString FeatureLevelName;
-	GetFeatureLevelName(RequestedFeatureLevel, FeatureLevelName);
-	FApp::SetGraphicsRHI(FString::Printf(TEXT("Metal (%s)"), *FeatureLevelName));
 	
 	DynamicRHIModule = &FModuleManager::LoadModuleChecked<IDynamicRHIModule>(TEXT("MetalRHI"));
 
@@ -166,6 +162,11 @@ FDynamicRHI* PlatformCreateDynamicRHI()
 	if (nullptr != (DynamicRHIModule = LoadDynamicRHIModule(RequestedFeatureLevel)))
 	{
 		DynamicRHI = DynamicRHIModule->CreateRHI(RequestedFeatureLevel);
+
+		// The Feature Level is known after creating the RHI
+		FString FeatureLevelName;
+		GetFeatureLevelName(GMaxRHIFeatureLevel, FeatureLevelName);
+		FApp::SetGraphicsRHI(FString::Printf(TEXT("Metal (%s)"), *FeatureLevelName));
 	}
 
 	return DynamicRHI;
