Set the sync flag for Gpu Profiler events.
Support late enabling for the Gpu Channel

diff --git a/Engine/Source/Runtime/RHI/Private/GpuProfilerTrace.cpp b/Engine/Source/Runtime/RHI/Private/GpuProfilerTrace.cpp
--- a/Engine/Source/Runtime/RHI/Private/GpuProfilerTrace.cpp
+++ b/Engine/Source/Runtime/RHI/Private/GpuProfilerTrace.cpp
@@ -283,13 +283,13 @@ UE_TRACE_EVENT_BEGIN(GpuProfiler, QueueSpec, NoSync | Important)
 	UE_TRACE_EVENT_FIELD(UE::Trace::WideString, TypeString)
 UE_TRACE_EVENT_END()
 
-UE_TRACE_EVENT_BEGIN(GpuProfiler, EventFrameBoundary, NoSync)
+UE_TRACE_EVENT_BEGIN(GpuProfiler, EventFrameBoundary)
 	UE_TRACE_EVENT_FIELD(uint32, QueueId)
 	UE_TRACE_EVENT_FIELD(uint32, FrameNumber)
 UE_TRACE_EVENT_END()
 
 #if WITH_RHI_BREADCRUMBS
-UE_TRACE_EVENT_BEGIN(GpuProfiler, EventBeginBreadcrumb, NoSync)
+UE_TRACE_EVENT_BEGIN(GpuProfiler, EventBeginBreadcrumb)
 	UE_TRACE_EVENT_FIELD(uint32, QueueId)
 	UE_TRACE_EVENT_FIELD(uint64, GPUTimestampTOP)
 	UE_TRACE_EVENT_FIELD(UE::Trace::WideString, Name)
@@ -298,25 +298,25 @@ UE_TRACE_EVENT_BEGIN(GpuProfiler, EventBeginBreadcrumb, NoSync)
 	//UE_TRACE_EVENT_FIELD(uint8[], LongNameArgs)
 UE_TRACE_EVENT_END()
 
-UE_TRACE_EVENT_BEGIN(GpuProfiler, EventEndBreadcrumb, NoSync)
+UE_TRACE_EVENT_BEGIN(GpuProfiler, EventEndBreadcrumb)
 	UE_TRACE_EVENT_FIELD(uint32, QueueId)
 	UE_TRACE_EVENT_FIELD(uint64, GPUTimestampBOP)
 	UE_TRACE_EVENT_FIELD(UE::Trace::WideString, Name) // for debugging
 UE_TRACE_EVENT_END()
 #endif // WITH_RHI_BREADCRUMBS
 
-UE_TRACE_EVENT_BEGIN(GpuProfiler, EventBeginWork, NoSync)
+UE_TRACE_EVENT_BEGIN(GpuProfiler, EventBeginWork)
 	UE_TRACE_EVENT_FIELD(uint32, QueueId)
 	UE_TRACE_EVENT_FIELD(uint64, GPUTimestampTOP)
 	UE_TRACE_EVENT_FIELD(uint64, CPUTimestamp)
 UE_TRACE_EVENT_END()
 
-UE_TRACE_EVENT_BEGIN(GpuProfiler, EventEndWork, NoSync)
+UE_TRACE_EVENT_BEGIN(GpuProfiler, EventEndWork)
 	UE_TRACE_EVENT_FIELD(uint32, QueueId)
 	UE_TRACE_EVENT_FIELD(uint64, GPUTimestampBOP)
 UE_TRACE_EVENT_END()
 
-UE_TRACE_EVENT_BEGIN(GpuProfiler, EventStats, NoSync)
+UE_TRACE_EVENT_BEGIN(GpuProfiler, EventStats)
 	UE_TRACE_EVENT_FIELD(uint32, QueueId)
 	UE_TRACE_EVENT_FIELD(uint32, NumDraws)
 	UE_TRACE_EVENT_FIELD(uint32, NumPrimitives)
@@ -325,27 +325,22 @@ UE_TRACE_EVENT_END()
 struct FGPUProfilerSink_Trace final : public FEventSink
 {
 	bool bInitialized = false;
+	TArray<FQueue> CachedQueues;
 
-	void InitializeQueues(TConstArrayView<FQueue> Queues) override
+	void Initialize()
 	{
-		if (!UE_TRACE_CHANNELEXPR_IS_ENABLED(GpuChannel))
-		{
-			//TODO: Should we still initialize even if GPU channel is off?
-			// (ex.: to keep a copy of all queues and to emit Init and QueueSpec later when GPU channel is on)
-			return;
-		}
+		ensure(!bInitialized);
 
-		if (bInitialized)
-		{
-			//TODO: Is it valid for InitializeQueues to be called multiple times?
-			//return;
-		}
-
-		bInitialized = true;
+		constexpr uint8 Version = 2;
 
 		UE_TRACE_LOG(GpuProfiler, Init, GpuChannel)
-			<< Init.Version(2);
+			<< Init.Version(Version);
+
+		bInitialized = true;
+	}
 
+	void TraceQueues(TConstArrayView<FQueue> Queues)
+	{
 		for (FQueue const& Queue : Queues)
 		{
 			FStringView TypeString(Queue.GetTypeString());
@@ -355,19 +350,40 @@ struct FGPUProfilerSink_Trace final : public FEventSink
 		}
 	}
 
-	void ProcessEvents(FQueue Queue, FEventStream const& EventStream) override
+	void InitializeQueues(TConstArrayView<FQueue> Queues) override
 	{
 		if (!UE_TRACE_CHANNELEXPR_IS_ENABLED(GpuChannel))
 		{
+			CachedQueues.Append(Queues);
 			return;
 		}
 
 		if (!bInitialized)
 		{
-			//TODO: Can we query now Queues and initialize now?
+			Initialize();
+		}
+
+		TraceQueues(Queues);
+	}
+
+	void ProcessEvents(FQueue Queue, FEventStream const& EventStream) override
+	{
+		if (!UE_TRACE_CHANNELEXPR_IS_ENABLED(GpuChannel))
+		{
 			return;
 		}
 
+		if (!bInitialized)
+		{
+			Initialize();
+		}
+
+		if (!CachedQueues.IsEmpty())
+		{
+			TraceQueues(CachedQueues);
+			CachedQueues.Empty();
+		}
+
 		for (FEvent const* Event : EventStream)
 		{
 			switch (Event->GetType())
