IRHIComputeContext::SetTrackedAccess should set the state on the referenced texture instead on a FRHITextureReference directly.
Making the FRHIViewableResource::SetTrackedAccess virtual allow the redirection.

diff --git a/Engine/Source/Runtime/RHI/Private/RHITextureReference.cpp b/Engine/Source/Runtime/RHI/Private/RHITextureReference.cpp
--- a/Engine/Source/Runtime/RHI/Private/RHITextureReference.cpp
+++ b/Engine/Source/Runtime/RHI/Private/RHITextureReference.cpp
@@ -82,3 +82,9 @@ const FRHITextureDesc& FRHITextureReference::GetDesc() const
 	check(ReferencedTexture);
 	return ReferencedTexture->GetDesc();
 }
+
+void FRHITextureReference::SetTrackedAccessFromContext(ERHIAccess Access)
+{
+	check(ReferencedTexture);
+	ReferencedTexture->SetTrackedAccessFromContext(Access);
+}
diff --git a/Engine/Source/Runtime/RHI/Public/RHIContext.h b/Engine/Source/Runtime/RHI/Public/RHIContext.h
--- a/Engine/Source/Runtime/RHI/Public/RHIContext.h
+++ b/Engine/Source/Runtime/RHI/Public/RHIContext.h
@@ -481,7 +481,7 @@ public:
 	{
 		check(Info.Resource != nullptr);
 		check(Info.Access != ERHIAccess::Unknown);
-		Info.Resource->TrackedAccess = Info.Access;
+		Info.Resource->SetTrackedAccessFromContext(Info.Access);
 	}
 
 	inline ERHIAccess GetTrackedAccess(const FRHIViewableResource* Resource) const
diff --git a/Engine/Source/Runtime/RHI/Public/RHIResources.h b/Engine/Source/Runtime/RHI/Public/RHIResources.h
--- a/Engine/Source/Runtime/RHI/Public/RHIResources.h
+++ b/Engine/Source/Runtime/RHI/Public/RHIResources.h
@@ -1344,6 +1344,11 @@ protected:
 		, TrackedAccess(InAccess, SubResourceCount)
 	{}
 
+	virtual void SetTrackedAccessFromContext(ERHIAccess Access)
+	{
+		TrackedAccess = Access;
+	}
+
 	void TakeOwnership(FRHIViewableResource& Other)
 	{
 		TrackedAccess = Other.TrackedAccess;
diff --git a/Engine/Source/Runtime/RHI/Public/RHITextureReference.h b/Engine/Source/Runtime/RHI/Public/RHITextureReference.h
--- a/Engine/Source/Runtime/RHI/Public/RHITextureReference.h
+++ b/Engine/Source/Runtime/RHI/Public/RHITextureReference.h
@@ -41,6 +41,9 @@ public:
 protected:
 	friend class FDynamicRHI;
 
+	// Only called from Platform Context
+	RHI_API virtual void SetTrackedAccessFromContext(ERHIAccess Access) override;
+
 	// Called only from FDynamicRHI::RHIUpdateTextureReference
 	void SetReferencedTexture(FRHITexture* InTexture)
 	{
