Add RHIValidation tracking for resource state tracking of resource used inside the SBT

diff --git a/Engine/Source/Runtime/Apple/MetalRHI/Private/MetalStateCache.cpp b/Engine/Source/Runtime/Apple/MetalRHI/Private/MetalStateCache.cpp
--- a/Engine/Source/Runtime/Apple/MetalRHI/Private/MetalStateCache.cpp
+++ b/Engine/Source/Runtime/Apple/MetalRHI/Private/MetalStateCache.cpp
@@ -14,6 +14,7 @@
 #include "RHIShaderParametersShared.h"
 #include "MetalResourceCollection.h"
 #include "MetalDynamicRHI.h"
+#include "RHIUniformBufferUtilities.h"
 
 #if PLATFORM_MAC
 	#ifndef UINT128_MAX
@@ -1961,7 +1962,7 @@ void FMetalStateCache::SetResourcesFromTables(ShaderType Shader, CrossCompiler::
 #endif
 		};
 
-		UE::RHICore::SetResourcesFromTables(
+		UE::RHI::Private::SetUniformBufferResourcesFromTables(
 			  FUniformResourceBinder { *this, Frequency }
 			, *Shader
 			, DirtyUniformBuffers[Frequency]
diff --git a/Engine/Source/Runtime/D3D12RHI/Private/D3D12Commands.cpp b/Engine/Source/Runtime/D3D12RHI/Private/D3D12Commands.cpp
--- a/Engine/Source/Runtime/D3D12RHI/Private/D3D12Commands.cpp
+++ b/Engine/Source/Runtime/D3D12RHI/Private/D3D12Commands.cpp
@@ -9,6 +9,7 @@ D3D12Commands.cpp: D3D RHI commands implementation.
 #include "ProfilingDebugging/RealtimeGPUProfiler.h"
 #include "D3D12ResourceCollection.h"
 #include "D3D12TextureReference.h"
+#include "RHIUniformBufferUtilities.h"
 
 static int32 GD3D12AllowDiscardResources = 1;
 static FAutoConsoleVariableRef CVarD3D12AllowDiscardResources(
@@ -1475,7 +1476,7 @@ void FD3D12CommandContext::SetResourcesFromTables(const ShaderType* RESTRICT Sha
 	static constexpr EShaderFrequency Frequency = static_cast<EShaderFrequency>(ShaderType::StaticFrequency);
 
 	FD3D12ResourceBinder Binder(*this, Frequency, Shader);
-	UE::RHICore::SetResourcesFromTables(
+	UE::RHI::Private::SetUniformBufferResourcesFromTables(
 		  Binder
 		, *Shader
 		, DirtyUniformBuffers[Frequency]
diff --git a/Engine/Source/Runtime/D3D12RHI/Private/D3D12RayTracing.cpp b/Engine/Source/Runtime/D3D12RHI/Private/D3D12RayTracing.cpp
--- a/Engine/Source/Runtime/D3D12RHI/Private/D3D12RayTracing.cpp
+++ b/Engine/Source/Runtime/D3D12RHI/Private/D3D12RayTracing.cpp
@@ -23,6 +23,7 @@
 #include "D3D12ExplicitDescriptorCache.h"
 #include "D3D12ResourceCollection.h"
 #include "RHIShaderBindingLayout.h"
+#include "RHIUniformBufferUtilities.h"
 
 extern int32 GD3D12ExplicitViewDescriptorHeapSize;
 extern int32 GD3D12ExplicitViewDescriptorHeapOverflowReported;
@@ -1933,7 +1934,7 @@ public:
 
 	UE_NONCOPYABLE(FD3D12RayTracingPipelineState)
 
-	FD3D12RayTracingPipelineState(FD3D12Device* Device, const FRayTracingPipelineStateInitializer& Initializer)
+	FD3D12RayTracingPipelineState(FD3D12Device* Device, const FRayTracingPipelineStateInitializer& Initializer) : FRHIRayTracingPipelineState(Initializer)
 	{
 		SCOPE_CYCLE_COUNTER(STAT_RTPSO_CreatePipeline);
 		TRACE_CPUPROFILER_EVENT_SCOPE(RTPSO_CreatePipeline);
@@ -4628,7 +4629,7 @@ static bool SetRayTracingShaderResources(
 
 	{
 		uint32 DirtyUniformBuffers = ~(0u);
-		UE::RHICore::SetResourcesFromTables(
+		UE::RHI::Private::SetUniformBufferResourcesFromTables(
 			  Bindings
 			, *Shader
 			, DirtyUniformBuffers
diff --git a/Engine/Source/Runtime/D3D12RHI/Private/D3D12WorkGraph.cpp b/Engine/Source/Runtime/D3D12RHI/Private/D3D12WorkGraph.cpp
--- a/Engine/Source/Runtime/D3D12RHI/Private/D3D12WorkGraph.cpp
+++ b/Engine/Source/Runtime/D3D12RHI/Private/D3D12WorkGraph.cpp
@@ -11,6 +11,7 @@
 #include "D3D12TextureReference.h"
 #include "PipelineStateCache.h"
 #include "ShaderBundles.h"
+#include "RHIUniformBufferUtilities.h"
 
 static bool GShaderBundleSkipDispatch = false;
 static FAutoConsoleVariableRef CVarShaderBundleSkipDispatch(
@@ -401,7 +402,7 @@ static void RecordBindings(
 
 
 	uint32 FakeDirtyUniformBuffers = ~(0u);
-	UE::RHICore::SetResourcesFromTables(BundleBinder, *WorkGraphShaderRHI, FakeDirtyUniformBuffers, BundleUniformBuffers
+	UE::RHI::Private::SetUniformBufferResourcesFromTables(BundleBinder, *WorkGraphShaderRHI, FakeDirtyUniformBuffers, BundleUniformBuffers
 #if ENABLE_RHI_VALIDATION
 		, Context.Tracker
 #endif
diff --git a/Engine/Source/Runtime/OpenGLDrv/Private/OpenGLCommands.cpp b/Engine/Source/Runtime/OpenGLDrv/Private/OpenGLCommands.cpp
--- a/Engine/Source/Runtime/OpenGLDrv/Private/OpenGLCommands.cpp
+++ b/Engine/Source/Runtime/OpenGLDrv/Private/OpenGLCommands.cpp
@@ -19,6 +19,7 @@
 #include "RHICoreShader.h"
 #include "RHIShaderParametersShared.h"
 #include "DataDrivenShaderPlatformInfo.h"
+#include "RHIUniformBufferUtilities.h"
 
 /*
 #define DECLARE_ISBOUNDSHADER(ShaderType) template <typename TShaderType> inline void ValidateBoundShader(TRefCountPtr<FOpenGLBoundShaderState> InBoundShaderState, FRHIGraphicsShader* GfxShader, TShaderType* ShaderTypeRHI) \
@@ -2097,7 +2098,7 @@ FORCEINLINE void FOpenGLDynamicRHI::SetResourcesFromTables(ShaderType* Shader)
 
 	static constexpr EShaderFrequency Frequency = ShaderType::Frequency;
 
-	UE::RHICore::SetResourcesFromTables(
+	UE::RHI::Private::SetUniformBufferResourcesFromTables(
 		FOpenGLResourceBinder(*this, Frequency, true)
 		, *Shader
 		, PendingState.DirtyUniformBuffers[Frequency]
diff --git a/Engine/Source/Runtime/RHI/Internal/RHIUniformBufferUtilities.h b/Engine/Source/Runtime/RHI/Internal/RHIUniformBufferUtilities.h
new file mode 100644
--- /dev/null
+++ b/Engine/Source/Runtime/RHI/Internal/RHIUniformBufferUtilities.h
@@ -0,0 +1,248 @@
+// Copyright Epic Games, Inc. All Rights Reserved.
+
+#pragma once
+
+#include "CoreTypes.h"
+#include "Misc/AssertionMacros.h"
+#include "Misc/App.h"
+#include "RHIDefinitions.h"
+#include "RHIResources.h"
+
+namespace UE::RHI::Private
+{
+
+template<typename TResourceType> struct TResourceTypeStr {};
+template<> struct TResourceTypeStr<FRHISamplerState> { static constexpr TCHAR String[] = TEXT("Sampler State"); };
+template<> struct TResourceTypeStr<FRHITexture> { static constexpr TCHAR String[] = TEXT("Texture"); };
+template<> struct TResourceTypeStr<FRHIShaderResourceView> { static constexpr TCHAR String[] = TEXT("Shader Resource View"); };
+template<> struct TResourceTypeStr<FRHIUnorderedAccessView> { static constexpr TCHAR String[] = TEXT("Unordered Access View"); };
+template<> struct TResourceTypeStr<FRHIResourceCollection> { static constexpr TCHAR String[] = TEXT("Resource Collection"); };
+
+template<typename TResourceType, typename TCallback>
+inline void EnumerateUniformBufferResources(const FRHIUniformBuffer* RESTRICT Buffer, int32 BufferIndex, const uint32* RESTRICT ResourceMap, TCallback && Callback)
+{
+	const TRefCountPtr<FRHIResource>* RESTRICT Resources = Buffer->GetResourceTable().GetData();
+
+	uint32 BufferOffset = ResourceMap[BufferIndex];
+	if (BufferOffset > 0)
+	{
+		const uint32* RESTRICT ResourceInfos = &ResourceMap[BufferOffset];
+		uint32 ResourceInfo = *ResourceInfos++;
+		do
+		{
+			checkSlow(FRHIResourceTableEntry::GetUniformBufferIndex(ResourceInfo) == BufferIndex);
+
+			const uint16 ResourceIndex = FRHIResourceTableEntry::GetResourceIndex(ResourceInfo);
+			const uint8 BindIndex = FRHIResourceTableEntry::GetBindIndex    (ResourceInfo);
+
+			TResourceType* Resource = static_cast<TResourceType*>(Resources[ResourceIndex].GetReference());
+			checkf(Resource
+			       , TEXT("Null %s (resource %d bind %d) on UB Layout %s")
+			       , TResourceTypeStr<TResourceType>::String
+			       , ResourceIndex
+			       , BindIndex
+			       , *Buffer->GetLayout().GetDebugName()
+			       );
+
+			Callback(Resource, BindIndex);
+
+			ResourceInfo = *ResourceInfos++;
+		} while (FRHIResourceTableEntry::GetUniformBufferIndex(ResourceInfo) == BufferIndex);
+	}
+}
+
+template<typename TBinder, typename TUniformBufferArrayType, typename TBitMaskType, bool bFullyBindless = false>
+void SetUniformBufferResourcesFromTables(TBinder && Binder, FRHIShader const& Shader, TBitMaskType& DirtyUniformBuffers, TUniformBufferArrayType const& BoundUniformBuffers
+                                         #if ENABLE_RHI_VALIDATION
+                                         , RHIValidation::FTracker* Tracker
+                                         #endif
+                                         )
+{
+	float CurrentTimeForTextureTimes = FApp::GetCurrentTime();
+	FShaderResourceTable const& SRT = Shader.GetShaderResourceTable();
+
+	// Mask the dirty bits by those buffers from which the shader has bound resources.
+	uint32 DirtyBits = SRT.ResourceTableBits & DirtyUniformBuffers;
+
+	#if PLATFORM_SUPPORTS_BINDLESS_RENDERING && !ENABLE_RHI_VALIDATION
+	if constexpr (bFullyBindless)
+	{
+		while (DirtyBits)
+		{
+			// Scan for the lowest set bit, compute its index, clear it in the set of dirty bits.
+			const uint32 LowestBitMask = (DirtyBits) & (-(int32)DirtyBits);
+			const int32 BufferIndex = FMath::CountTrailingZeros(LowestBitMask); // todo: This has a branch on zero, we know it could never be zero...
+			DirtyBits ^= LowestBitMask;
+
+			EnumerateUniformBufferResources<FRHITexture>(BoundUniformBuffers[BufferIndex], BufferIndex, SRT.TextureMap.GetData(),
+				[&](FRHITexture* Texture, uint8 Index)
+				{
+					Texture->SetLastRenderTime(CurrentTimeForTextureTimes);
+				});
+		}
+
+		DirtyUniformBuffers = TBitMaskType(0);
+
+		return;
+	}
+	#endif
+
+	while (DirtyBits)
+	{
+		// Scan for the lowest set bit, compute its index, clear it in the set of dirty bits.
+		const uint32 LowestBitMask = (DirtyBits) & (-(int32)DirtyBits);
+		const int32 BufferIndex = FMath::CountTrailingZeros(LowestBitMask); // todo: This has a branch on zero, we know it could never be zero...
+		DirtyBits ^= LowestBitMask;
+
+		check(BufferIndex < SRT.ResourceTableLayoutHashes.Num());
+
+		FRHIUniformBuffer* Buffer = BoundUniformBuffers[BufferIndex];
+
+		#if DO_CHECK
+
+		if (!Buffer)
+		{
+			UE_LOG(LogRHI, Fatal, TEXT("Shader expected a uniform buffer at slot %u but got null instead (Shader='%s' UB='%s'). Rendering code needs to set a valid uniform buffer for this slot.")
+			       , BufferIndex
+			       , Shader.GetShaderName()
+			       , *Shader.GetUniformBufferName(BufferIndex)
+			       );
+		}
+		else if (Buffer->GetLayout().GetHash() != SRT.ResourceTableLayoutHashes[BufferIndex])
+		{
+			FRHIUniformBufferLayout const& BufferLayout = Buffer->GetLayout();
+
+			FString ResourcesString;
+			for (FRHIUniformBufferResource const& Resource : BufferLayout.Resources)
+			{
+				ResourcesString += FString::Printf(TEXT("%s%d")
+				                                   , ResourcesString.Len() ? TEXT(" ") : TEXT("")
+				                                   , Resource.MemberType
+				                                   );
+			}
+
+			// This might mean you are accessing a data you haven't bound e.g. GBuffer
+			UE_LOG(LogRHI, Fatal,
+			       TEXT("Uniform buffer bound to slot %u is not what the shader expected:\n")
+			       TEXT("\tBound                : Uniform Buffer[%s] with Hash[0x%08x]\n")
+			       TEXT("\tExpected             : Uniform Buffer[%s] with Hash[0x%08x]\n")
+			       TEXT("\tShader Name          : %s\n")
+			       TEXT("\tLayout CB Size       : %d\n")
+			       TEXT("\tLayout Num Resources : %d\n")
+			       TEXT("\tResource Types       : %s\n")
+			       , BufferIndex
+			       , *BufferLayout.GetDebugName(), BufferLayout.GetHash()
+			       , *Shader.GetUniformBufferName(BufferIndex), SRT.ResourceTableLayoutHashes[BufferIndex]
+			       , Shader.GetShaderName()
+			       , BufferLayout.ConstantBufferSize
+			       , BufferLayout.Resources.Num()
+			       , *ResourcesString
+			       );
+		}
+
+		#endif // DO_CHECK
+
+		// Textures
+		EnumerateUniformBufferResources<FRHITexture>(Buffer, BufferIndex, SRT.TextureMap.GetData(),
+			[&](FRHITexture* Texture, uint8 Index)
+			{
+				#if ENABLE_RHI_VALIDATION
+				if (Tracker)
+				{
+					ERHIAccess Access = IsComputeShaderFrequency(Shader.GetFrequency())
+					? ERHIAccess::SRVCompute
+					: Shader.GetFrequency() == SF_Pixel
+					? ERHIAccess::SRVGraphicsPixel
+					: ERHIAccess::SRVGraphicsNonPixel;
+
+					// Textures bound here only have their "common" plane accessible. Stencil etc is ignored.
+					// (i.e. only access the color plane of a color texture, or depth plane of a depth texture)
+					Tracker->Assert(Texture->GetViewIdentity(0, 0, 0, 0, uint32(RHIValidation::EResourcePlane::Common), 1), Access);
+				}
+				#endif
+				Texture->SetLastRenderTime(CurrentTimeForTextureTimes);
+				Binder.SetTexture(Texture, Index);
+			});
+
+		// SRVs
+		EnumerateUniformBufferResources<FRHIShaderResourceView>(Buffer, BufferIndex, SRT.ShaderResourceViewMap.GetData(),
+			[&](FRHIShaderResourceView* SRV, uint8 Index)
+			{
+				#if ENABLE_RHI_VALIDATION
+				if (Tracker)
+				{
+					ERHIAccess Access = IsComputeShaderFrequency(Shader.GetFrequency())
+					? ERHIAccess::SRVCompute
+					: Shader.GetFrequency() == SF_Pixel
+					? ERHIAccess::SRVGraphicsPixel
+					: ERHIAccess::SRVGraphicsNonPixel;
+
+					Tracker->Assert(SRV->GetViewIdentity(), Access);
+				}
+				if (GRHIValidationEnabled)
+				{
+					RHIValidation::ValidateShaderResourceView(&Shader, Index, SRV);
+				}
+				#endif
+				Binder.SetSRV(SRV, Index);
+			});
+
+		#if PLATFORM_SUPPORTS_BINDLESS_RENDERING
+		EnumerateUniformBufferResources<FRHIResourceCollection>(Buffer, BufferIndex, SRT.ResourceCollectionMap.GetData(),
+			[&](FRHIResourceCollection* ResourceCollection, uint8 Index)
+			{
+				#if ENABLE_RHI_VALIDATION
+				// todo: christopher.waters - ResourceCollection validation
+				#endif
+				Binder.SetResourceCollection(ResourceCollection, Index);
+			});
+		#endif // PLATFORM_SUPPORTS_BINDLESS_RENDERING
+
+		// Samplers
+		EnumerateUniformBufferResources<FRHISamplerState>(Buffer, BufferIndex, SRT.SamplerMap.GetData(),
+			[&](FRHISamplerState* Sampler, uint8 Index)
+			{
+				Binder.SetSampler(Sampler, Index);
+			});
+
+		// UAVs
+		EnumerateUniformBufferResources<FRHIUnorderedAccessView>(Buffer, BufferIndex, SRT.UnorderedAccessViewMap.GetData(),
+			[&](FRHIUnorderedAccessView* UAV, uint8 Index)
+			{
+				#if ENABLE_RHI_VALIDATION
+				if (Tracker)
+				{
+					ERHIAccess Access = IsComputeShaderFrequency(Shader.GetFrequency())
+					? ERHIAccess::UAVCompute
+					: ERHIAccess::UAVGraphics;
+
+					Tracker->AssertUAV(UAV, Access, Index);
+				}
+				#endif
+				Binder.SetUAV(UAV, Index);
+			});
+	}
+
+	DirtyUniformBuffers = TBitMaskType(0);
+}
+
+
+template<typename TBinder, typename TUniformBufferArrayType, typename TBitMaskType>
+void SetFullyBindlessUniformBufferResourcesFromTables(TBinder && Binder, const FRHIShader& Shader, TBitMaskType& DirtyUniformBuffers, const TUniformBufferArrayType& BoundUniformBuffers
+                                                      #if ENABLE_RHI_VALIDATION
+                                                      , RHIValidation::FTracker* Tracker
+                                                      #endif
+                                                      )
+{
+	SetUniformBufferResourcesFromTables < TBinder, TUniformBufferArrayType, TBitMaskType, true > (
+		MoveTemp(Binder)
+		, Shader
+		, DirtyUniformBuffers
+		, BoundUniformBuffers
+		#if ENABLE_RHI_VALIDATION
+		, Tracker
+		#endif
+	);
+}
+
+} //! UE::RHI::Private
\ No newline at end of file
diff --git a/Engine/Source/Runtime/RHI/Private/RHIValidation.cpp b/Engine/Source/Runtime/RHI/Private/RHIValidation.cpp
--- a/Engine/Source/Runtime/RHI/Private/RHIValidation.cpp
+++ b/Engine/Source/Runtime/RHI/Private/RHIValidation.cpp
@@ -13,7 +13,9 @@
 #include "Misc/OutputDeviceRedirector.h"
 #include "RHIContext.h"
 #include "RHIStrings.h"
+#include "RHIUniformBufferUtilities.h"
 #include "Algo/BinarySearch.h"
+#include "Async/ParallelFor.h"
 
 #if ENABLE_RHI_VALIDATION
 
@@ -329,8 +331,216 @@ namespace RHIValidation
 			}
 		}
 	}
+			
+	FRayTracingPipelineState::FRayTracingPipelineState(const FRayTracingPipelineStateInitializer& Initializer)
+	{		
+		HitGroupShaders = Initializer.GetHitGroupTable();
+		MissShaders = Initializer.GetMissTable();
+		CallableShaders = Initializer.GetCallableTable();
+	}
+
+	FRHIRayTracingShader* FRayTracingPipelineState::GetShader(ERayTracingBindingType BindingType, uint32 Index) const
+	{
+		switch (BindingType)
+		{
+			case ERayTracingBindingType::HitGroup:
+				return HitGroupShaders[Index];
+			case ERayTracingBindingType::CallableShader:
+				return CallableShaders[Index];
+			case ERayTracingBindingType::MissShader:
+				return MissShaders[Index];
+		}
+
+		return nullptr;
+	}
+
+	FShaderBindingTable::FShaderBindingTable(const FRayTracingShaderBindingTableInitializer& InInitializer)
+	{
+		ShaderBindingMode = InInitializer.ShaderBindingMode;
+		HitGroupIndexingMode = InInitializer.HitGroupIndexingMode;
+	}
+
+	void FShaderBindingTable::Clear()
+	{
+		WorkerData[0].SRVs.Empty();
+		WorkerData[0].UAVs.Empty();
+
+		bIsDirty = true;
+	}
+	
+	static void CollectShaderBindingTableResources(FRayTracingPipelineState* RayTracingPipelineState, FShaderBindingTable* ShaderTable, const FRayTracingLocalShaderBindings& LocalShaderBinding, ERayTracingBindingType BindingType, uint32 WorkerIndex)
+	{
+		FRHIRayTracingShader* Shader = RayTracingPipelineState->GetShader(BindingType, LocalShaderBinding.ShaderIndexInPipeline);
+		ensure(Shader);
+
+		struct FResourceBinder
+		{
+			FResourceBinder(FShaderBindingTable* InShaderTable, FRHIRayTracingShader* InShader, uint32 InWorkerIndex, uint32 InRecordIndex) :
+			ShaderBindingTable(InShaderTable), RHIShader(InShader), WorkerIndex(InWorkerIndex), RecordIndex(InRecordIndex)
+			{					
+			}
+
+			void SetUAV(FRHIUnorderedAccessView* UAV, uint8 Index)
+			{
+				if (GRHIValidationEnabled)
+				{
+					RHIValidation::ValidateUnorderedAccessView(RHIShader, Index, UAV);
+				}
+				ShaderBindingTable->AddUAV(UAV, Index, WorkerIndex);
+			}
+
+			void SetSRV(FRHIShaderResourceView* SRV, uint8 Index)
+			{
+				if (GRHIValidationEnabled)
+				{
+					RHIValidation::ValidateShaderResourceView(RHIShader, Index, SRV);
+				}
+				ShaderBindingTable->AddSRV(SRV->GetViewIdentity(), WorkerIndex);
+			}
+
+			void SetTexture(FRHITexture* Texture, uint8 Index)
+			{
+				if (GRHIValidationEnabled)
+				{
+					RHIValidation::ValidateShaderResourceView(RHIShader, Index, Texture);
+				}				
+				ShaderBindingTable->AddSRV(Texture->GetWholeResourceIdentitySRV(), WorkerIndex);
+			}
+
+			void SetResourceCollection(FRHIResourceCollection* ResourceCollection, uint8 Index)
+			{
+				for (const FRHIResourceCollectionMember& Member : ResourceCollection->Members)
+				{
+					switch (Member.Type)
+					{
+						case FRHIResourceCollectionMember::EType::Texture:
+							if (FRHITexture* Texture = static_cast<FRHITexture*>(Member.Resource))
+							{
+								ShaderBindingTable->AddSRV(Texture->GetWholeResourceIdentitySRV(), WorkerIndex);
+							}
+							break;
+						case FRHIResourceCollectionMember::EType::TextureReference:
+							if (FRHITextureReference* Texture = static_cast<FRHITextureReference*>(Member.Resource))
+							{
+								ShaderBindingTable->AddSRV(Texture->GetWholeResourceIdentitySRV(), WorkerIndex);
+							}
+							break;
+						case FRHIResourceCollectionMember::EType::ShaderResourceView:
+							if (FRHIShaderResourceView* SRV = static_cast<FRHIShaderResourceView*>(Member.Resource))
+							{
+								ShaderBindingTable->AddSRV(SRV->GetViewIdentity(), WorkerIndex);
+							}
+							break;
+						default:
+							break;
+					}
+				}
+			}
+
+			void SetSampler(FRHISamplerState* RHISampler, uint8 Index)
+			{				
+				// nothing to validate
+			}
+
+			FShaderBindingTable* ShaderBindingTable;
+			FRHIRayTracingShader* RHIShader;
+			uint32 WorkerIndex;
+			uint32 RecordIndex;
+		};
+		FResourceBinder Binder(ShaderTable, Shader, WorkerIndex, LocalShaderBinding.RecordIndex);
+		
+		// Don't need to perform the state validation now because they can still change before the actual ray dispatch
+		RHIValidation::FTracker* ValidationTracker = nullptr; 
+
+		// Use rhi core function to find all the used RHI resources in the uniform buffers
+		uint32 DirtyUniformBuffers = ~(0u);
+		UE::RHI::Private::SetUniformBufferResourcesFromTables(Binder, *Shader, DirtyUniformBuffers, LocalShaderBinding.UniformBuffers, ValidationTracker);
+	}		
+
+	void FShaderBindingTable::SetBindingsOnShaderBindingTable(FRayTracingPipelineState* RayTracingPipelineState, uint32 NumBindings, const FRayTracingLocalShaderBindings* Bindings, ERayTracingBindingType BindingType)
+	{		
+		TRACE_CPUPROFILER_EVENT_SCOPE(RHIValidation-SetBindingsOnShaderBindingTable);
+		
+		FGraphEventArray TaskList;
+
+		const uint32 NumWorkerThreads = FTaskGraphInterface::Get().GetNumWorkerThreads();
+		const uint32 MaxTasks = FApp::ShouldUseThreadingForPerformance() ? FMath::Min<uint32>(NumWorkerThreads, MaxBindingWorkers) : 1;
+
+		struct FTaskContext
+		{
+			uint32 WorkerIndex = 0;
+		};
+
+		TArray<FTaskContext, TInlineAllocator<MaxBindingWorkers>> TaskContexts;
+		for (uint32 WorkerIndex = 0; WorkerIndex < MaxTasks; ++WorkerIndex)
+		{
+			TaskContexts.Add(FTaskContext{WorkerIndex});
+		}
+
+		auto BindingTask = [this, RayTracingPipelineState, Bindings, BindingType](const FTaskContext& Context, int32 CurrentIndex)
+		{
+			const FRayTracingLocalShaderBindings& Binding = Bindings[CurrentIndex];
+
+			// only collect shader binding data if RTPSO & hit group indexing mode
+			if (HitGroupIndexingMode == ERayTracingHitGroupIndexingMode::Allow && Binding.Geometry)
+			{
+				if (EnumHasAnyFlags(ShaderBindingMode, ERayTracingShaderBindingMode::RTPSO))
+				{
+					CollectShaderBindingTableResources(RayTracingPipelineState, this, Binding, BindingType, Context.WorkerIndex);
+				}
+			}
+		};
+
+		const int32 ItemsPerTask = 1024;
+		ParallelForWithExistingTaskContext(TEXT("SetRayTracingBindings"), MakeArrayView(TaskContexts), NumBindings, ItemsPerTask, BindingTask);
+				
+		// Mark dirty
+		bIsDirty = true;
+	}
+
+	void FShaderBindingTable::Commit()
+	{		
+		ensureMsgf(bIsDirty, TEXT("bIsDirty should be checked before calling Commit()"));
+		
+		// Merge all data from worker threads into the main set
+		for (uint32 WorkerIndex = 1; WorkerIndex < MaxBindingWorkers; ++WorkerIndex)
+		{
+			for (const FResourceIdentity& ResourceIdentity : WorkerData[WorkerIndex].SRVs)
+			{
+				AddSRV(ResourceIdentity, 0);
+			}
+
+			for (const FUAVBinding& UAVBinding : WorkerData[WorkerIndex].UAVs)
+			{
+				AddUAV(UAVBinding.UAV, UAVBinding.Slot, 0);
+			}
+
+			WorkerData[WorkerIndex].SRVs.Empty();
+			WorkerData[WorkerIndex].UAVs.Empty();
+		}
+
+		bIsDirty = false;
+	}
+
+	void FShaderBindingTable::ValidateStateForDispatch(RHIValidation::FTracker* Tracker) const
+	{
+		ensureMsgf(!bIsDirty, TEXT("RayTracing bindings have not been committed. You must call CommitRayTracingBindings first."));
+
+		// Validate all used SRVs
+		for (const FResourceIdentity& SRV : WorkerData[0].SRVs)
+		{
+			Tracker->Assert(SRV, ERHIAccess::SRVCompute);
+		}
+
+		// Validate all used UAVs
+		for (const FUAVBinding& UAVBinding : WorkerData[0].UAVs)
+		{
+			Tracker->AssertUAV(UAVBinding.UAV, ERHIAccess::UAVCompute, UAVBinding.Slot);
+		}
+	}
 }
 
+
 TSet<uint32> FValidationRHI::SeenFailureHashes;
 FCriticalSection FValidationRHI::SeenFailureHashesMutex;
 
@@ -3100,7 +3310,7 @@ void ValidateShaderParameters(FRHIShader* RHIShader, RHIValidation::FTracker* Tr
 						}
 						break;
 					case FRHIResourceCollectionMember::EType::ShaderResourceView:
-						if (FRHIShaderResourceView* SRV = static_cast<FRHIShaderResourceView*>(Parameter.Resource))
+						if (FRHIShaderResourceView* SRV = static_cast<FRHIShaderResourceView*>(Member.Resource))
 						{
 							Tracker->Assert(SRV->GetViewIdentity(), InRequiredAccess);
 						}
diff --git a/Engine/Source/Runtime/RHI/Public/RHIResources.h b/Engine/Source/Runtime/RHI/Public/RHIResources.h
--- a/Engine/Source/Runtime/RHI/Public/RHIResources.h
+++ b/Engine/Source/Runtime/RHI/Public/RHIResources.h
@@ -34,6 +34,7 @@ class FRHIShaderBindingLayout;
 class FResourceBulkDataInterface;
 class FResourceArrayInterface;
 struct FResourceArrayUploadInterface;
+class FRayTracingPipelineStateInitializer;
 
 struct FClearValueBinding;
 struct FRHIResourceInfo;
@@ -1133,9 +1134,16 @@ public:
 };
 
 class FRHIRayTracingPipelineState : public FRHIResource
+#if ENABLE_RHI_VALIDATION
+	, public RHIValidation::FRayTracingPipelineState
+#endif
 {
 public:
-	FRHIRayTracingPipelineState() : FRHIResource(RRT_RayTracingPipelineState) {}
+	FRHIRayTracingPipelineState(const FRayTracingPipelineStateInitializer& InInitializer) : FRHIResource(RRT_RayTracingPipelineState) 
+#if ENABLE_RHI_VALIDATION
+		, RHIValidation::FRayTracingPipelineState(InInitializer)
+#endif
+	{}
 };
 
 //
@@ -3341,7 +3349,7 @@ enum class ERayTracingAccelerationStructureFlags
 };
 ENUM_CLASS_FLAGS(ERayTracingAccelerationStructureFlags);
 
-enum class ERayTracingShaderBindingMode
+enum class ERayTracingShaderBindingMode : uint8
 {
 	Disabled	= 0,				//< No binding data at all
 	Inline		= 1 << 0,			//< Binding data for inline raytracing
@@ -3349,7 +3357,7 @@ enum class ERayTracingShaderBindingMode
 };
 ENUM_CLASS_FLAGS(ERayTracingShaderBindingMode);
 
-enum class ERayTracingHitGroupIndexingMode
+enum class ERayTracingHitGroupIndexingMode : uint8
 {
 	Allow,
 	Disallow,
@@ -3527,12 +3535,15 @@ public:
 class FRHIShaderBindingTable
 	: public FRHIResource
 #if ENABLE_RHI_VALIDATION
-	, public RHIValidation::FRayTracingShaderBindingTable
+	, public RHIValidation::FShaderBindingTable
 #endif
 {
 public:
 	FRHIShaderBindingTable(const FRayTracingShaderBindingTableInitializer& InInitializer)
-		: FRHIResource(RRT_RayTracingShaderBindingTable)
+		: FRHIResource(RRT_RayTracingShaderBindingTable)	
+#if ENABLE_RHI_VALIDATION
+		, RHIValidation::FShaderBindingTable(InInitializer)
+#endif
 		, Initializer(InInitializer)
 	{}
 
diff --git a/Engine/Source/Runtime/RHI/Public/RHIValidationCommon.h b/Engine/Source/Runtime/RHI/Public/RHIValidationCommon.h
--- a/Engine/Source/Runtime/RHI/Public/RHIValidationCommon.h
+++ b/Engine/Source/Runtime/RHI/Public/RHIValidationCommon.h
@@ -12,6 +12,7 @@
 #include "RHIStrings.h"
 #include "RHIAccess.h"
 #include "RHIBreadcrumbs.h"
+#include "Templates/TypeHash.h"
 
 #if ENABLE_RHI_VALIDATION
 extern RHI_API bool GRHIValidationEnabled;
@@ -37,6 +38,14 @@ struct FRHITextureCreateDesc;
 struct FRHITransitionInfo;
 struct FRHIViewDesc;
 
+class FRayTracingPipelineStateInitializer;
+class FRHIRayTracingPipelineState;
+struct FRayTracingShaderBindingTableInitializer;
+struct FRayTracingLocalShaderBindings;
+
+enum class ERayTracingBindingType : uint8;
+enum class ERayTracingShaderBindingMode : uint8;
+enum class ERayTracingHitGroupIndexingMode : uint8;
 
 namespace RHIValidation
 {
@@ -231,6 +240,16 @@ namespace RHIValidation
 		inline bool IsWholeResource(FResource& Resource) const;
 	};
 
+	inline uint32 GetTypeHash(const FSubresourceRange& Range)
+	{
+		uint32 Hash = HashCombineFast(Range.MipIndex, Range.NumMips);
+		Hash = HashCombineFast(Hash, Range.ArraySlice);
+		Hash = HashCombineFast(Hash, Range.NumArraySlices);
+		Hash = HashCombineFast(Hash, Range.PlaneIndex);
+		Hash = HashCombineFast(Hash, Range.NumPlanes);
+		return Hash;
+	}
+
 	struct FResourceIdentity
 	{
 		FResource* Resource;
@@ -249,6 +268,13 @@ namespace RHIValidation
 			return !(*this == RHS);
 		}
 	};
+	
+	inline uint32 GetTypeHash(const FResourceIdentity& ResourceIdentiy)
+	{
+		uint32 Hash = PointerHash(ResourceIdentiy.Resource);
+		Hash = HashCombineFast(Hash, GetTypeHash(ResourceIdentiy.SubresourceRange));
+		return Hash;
+	}
 
 	struct FViewIdentity : public FResourceIdentity
 	{
@@ -440,21 +466,7 @@ namespace RHIValidation
 	{
 	public:
 	};
-
-	class FRayTracingShaderBindingTable
-	{
-	public:
-		void SetCommitted(bool bInCommitted)
-		{
-			bCommitted = bInCommitted;
-		}
-
-		bool IsCommitted() const { return bCommitted; }
-
-	private:
-		bool bCommitted = false;
-	};
-
+	   
 	class FTextureResource
 	{
 	private:
@@ -504,6 +516,83 @@ namespace RHIValidation
 		}
 	};
 
+	class FRayTracingPipelineState
+	{
+	public:
+
+		RHI_API FRayTracingPipelineState(const FRayTracingPipelineStateInitializer& Initializer);
+		RHI_API FRHIRayTracingShader* GetShader(ERayTracingBindingType BindingType, uint32 Index) const;
+
+	private:
+		
+		// Cache the RHIShaders per binding type so they can be retrieved during SetBindingsOnShaderBindingTable to find all the used resources for a certain RHIShader
+		TArray<FRHIRayTracingShader*> MissShaders;
+		TArray<FRHIRayTracingShader*> HitGroupShaders;
+		TArray<FRHIRayTracingShader*> CallableShaders;
+	};
+		
+	struct FUAVBinding
+	{
+		FRHIUnorderedAccessView* UAV = nullptr;
+		uint32 Slot = 0;
+
+		inline bool operator == (FUAVBinding const& RHS) const
+		{
+			return UAV == RHS.UAV
+			&& Slot == RHS.Slot;
+		}
+
+		inline bool operator != (FUAVBinding const& RHS) const
+		{
+			return !(*this == RHS);
+		}
+	};
+		
+	inline uint32 GetTypeHash(const FUAVBinding& UAVBinding)
+	{
+		uint32 Hash = PointerHash(UAVBinding.UAV);
+		Hash = HashCombineFast(Hash, UAVBinding.Slot);
+		return Hash;
+	}
+
+	class FShaderBindingTable
+	{
+	public:
+
+		RHI_API FShaderBindingTable(const FRayTracingShaderBindingTableInitializer& InInitializer);
+
+		RHI_API void Clear();
+		RHI_API void SetBindingsOnShaderBindingTable(FRayTracingPipelineState* RayTracingPipelineState, uint32 NumBindings, const FRayTracingLocalShaderBindings* Bindings, ERayTracingBindingType BindingType);
+		RHI_API void Commit();
+
+		RHI_API void ValidateStateForDispatch(RHIValidation::FTracker* Tracker) const;
+		
+		void AddSRV(const FResourceIdentity& ResourceIdentity, uint32 WorkerIndex)
+		{
+			WorkerData[WorkerIndex].SRVs.Add(ResourceIdentity);
+		}
+
+		void AddUAV(FRHIUnorderedAccessView* UAV, uint32 InSlot, uint32 WorkerIndex)
+		{
+			WorkerData[WorkerIndex].UAVs.Add({ UAV, InSlot });
+		}
+
+	private:
+						
+		ERayTracingShaderBindingMode ShaderBindingMode;
+		ERayTracingHitGroupIndexingMode HitGroupIndexingMode;
+		bool bIsDirty = true;
+		
+		struct FWorkerThreadData
+		{
+			TSet<FResourceIdentity> SRVs;
+			TSet<FUAVBinding> UAVs;
+		};
+
+		static constexpr uint32 MaxBindingWorkers = 5; // RHI thread + 4 parallel workers.
+		FWorkerThreadData WorkerData[MaxBindingWorkers];
+	};
+
 	struct FFence
 	{
 		bool bSignaled = false;
diff --git a/Engine/Source/Runtime/RHI/Public/RHIValidationContext.h b/Engine/Source/Runtime/RHI/Public/RHIValidationContext.h
--- a/Engine/Source/Runtime/RHI/Public/RHIValidationContext.h
+++ b/Engine/Source/Runtime/RHI/Public/RHIValidationContext.h
@@ -1185,13 +1185,13 @@ public:
 
 	void RHIClearShaderBindingTable(FRHIShaderBindingTable* SBT)
 	{
-		SBT->SetCommitted(false);
+		SBT->Clear();
 		RHIContext->RHIClearShaderBindingTable(SBT);
 	}
 
 	void RHICommitShaderBindingTable(FRHIShaderBindingTable* SBT)
 	{
-		SBT->SetCommitted(true);
+		SBT->Commit();
 		RHIContext->RHICommitShaderBindingTable(SBT);
 	}
 
@@ -1243,7 +1243,7 @@ public:
 		FRHIShaderBindingTable* SBT, const FRayTracingShaderBindings& GlobalResourceBindings,
 		uint32 Width, uint32 Height) override final
 	{
-		ensureMsgf(SBT->IsCommitted(), TEXT("RayTracing bindings have not been committed. You must call CommitRayTracingBindings first."));
+		SBT->ValidateStateForDispatch(Tracker);
 		RHIContext->RHIRayTraceDispatch(RayTracingPipelineState, RayGenShader, SBT, GlobalResourceBindings, Width, Height);
 	}
 
@@ -1255,13 +1255,13 @@ public:
 		Tracker->Assert(ArgumentBuffer->GetWholeResourceIdentity(), ERHIAccess::IndirectArgs);
 		Tracker->Assert(ArgumentBuffer->GetWholeResourceIdentity(), ERHIAccess::SRVCompute);
 
-		ensureMsgf(SBT->IsCommitted(), TEXT("RayTracing bindings have not been committed. You must call CommitRayTracingBindings first."));
+		SBT->ValidateStateForDispatch(Tracker);
 		RHIContext->RHIRayTraceDispatchIndirect(RayTracingPipelineState, RayGenShader, SBT, GlobalResourceBindings, ArgumentBuffer, ArgumentOffset);
 	}
 
 	virtual void RHISetBindingsOnShaderBindingTable(FRHIShaderBindingTable* SBT, FRHIRayTracingPipelineState* Pipeline, uint32 NumBindings, const FRayTracingLocalShaderBindings* Bindings, ERayTracingBindingType BindingType) override final
 	{
-		SBT->SetCommitted(false);
+		SBT->SetBindingsOnShaderBindingTable(Pipeline, NumBindings, Bindings, BindingType);
 		RHIContext->RHISetBindingsOnShaderBindingTable(SBT, Pipeline, NumBindings, Bindings, BindingType);
 	}
 
diff --git a/Engine/Source/Runtime/RHICore/Public/RHICoreShader.h b/Engine/Source/Runtime/RHICore/Public/RHICoreShader.h
--- a/Engine/Source/Runtime/RHICore/Public/RHICoreShader.h
+++ b/Engine/Source/Runtime/RHICore/Public/RHICoreShader.h
@@ -104,239 +104,6 @@ void ApplyStaticUniformBuffers(
 		});
 }
 
-template <typename TResourceType> struct TResourceTypeStr {};
-template<> struct TResourceTypeStr<FRHISamplerState       > { static constexpr TCHAR String[] = TEXT("Sampler State"); };
-template<> struct TResourceTypeStr<FRHITexture            > { static constexpr TCHAR String[] = TEXT("Texture"); };
-template<> struct TResourceTypeStr<FRHIShaderResourceView > { static constexpr TCHAR String[] = TEXT("Shader Resource View"); };
-template<> struct TResourceTypeStr<FRHIUnorderedAccessView> { static constexpr TCHAR String[] = TEXT("Unordered Access View"); };
-template<> struct TResourceTypeStr<FRHIResourceCollection>  { static constexpr TCHAR String[] = TEXT("Resource Collection"); };
-
-template <typename TResourceType, typename TCallback>
-inline void EnumerateUniformBufferResources(const FRHIUniformBuffer* RESTRICT Buffer, int32 BufferIndex, const uint32* RESTRICT ResourceMap, TCallback&& Callback)
-{
-	const TRefCountPtr<FRHIResource>* RESTRICT Resources = Buffer->GetResourceTable().GetData();
-
-	uint32 BufferOffset = ResourceMap[BufferIndex];
-	if (BufferOffset > 0)
-	{
-		const uint32* RESTRICT ResourceInfos = &ResourceMap[BufferOffset];
-		uint32 ResourceInfo = *ResourceInfos++;
-		do
-		{
-			checkSlow(FRHIResourceTableEntry::GetUniformBufferIndex(ResourceInfo) == BufferIndex);
-
-			const uint16 ResourceIndex = FRHIResourceTableEntry::GetResourceIndex(ResourceInfo);
-			const uint8  BindIndex     = FRHIResourceTableEntry::GetBindIndex    (ResourceInfo);
-
-			TResourceType* Resource = static_cast<TResourceType*>(Resources[ResourceIndex].GetReference());
-			checkf(Resource
-				, TEXT("Null %s (resource %d bind %d) on UB Layout %s")
-				, TResourceTypeStr<TResourceType>::String
-				, ResourceIndex
-				, BindIndex
-				, *Buffer->GetLayout().GetDebugName()
-			);
-
-			Callback(Resource, BindIndex);
-
-			ResourceInfo = *ResourceInfos++;
-		} while (FRHIResourceTableEntry::GetUniformBufferIndex(ResourceInfo) == BufferIndex);
-	}
-}
-
-template <typename TBinder, typename TUniformBufferArrayType, typename TBitMaskType, bool bFullyBindless = false>
-void SetResourcesFromTables(TBinder&& Binder, FRHIShader const& Shader, TBitMaskType& DirtyUniformBuffers, TUniformBufferArrayType const& BoundUniformBuffers
-#if ENABLE_RHI_VALIDATION
-	, RHIValidation::FTracker* Tracker
-#endif
-)
-{
-	float CurrentTimeForTextureTimes = FApp::GetCurrentTime();
-	FShaderResourceTable const& SRT = Shader.GetShaderResourceTable();
-
-	// Mask the dirty bits by those buffers from which the shader has bound resources.
-	uint32 DirtyBits = SRT.ResourceTableBits & DirtyUniformBuffers;
-
-#if PLATFORM_SUPPORTS_BINDLESS_RENDERING && !ENABLE_RHI_VALIDATION
-	if constexpr (bFullyBindless)
-	{
-		while (DirtyBits)
-		{
-			// Scan for the lowest set bit, compute its index, clear it in the set of dirty bits.
-			const uint32 LowestBitMask = (DirtyBits) & (-(int32)DirtyBits);
-			const int32 BufferIndex = FMath::CountTrailingZeros(LowestBitMask); // todo: This has a branch on zero, we know it could never be zero...
-			DirtyBits ^= LowestBitMask;
-
-			EnumerateUniformBufferResources<FRHITexture>(BoundUniformBuffers[BufferIndex], BufferIndex, SRT.TextureMap.GetData(),
-				[&](FRHITexture* Texture, uint8 Index)
-				{
-					Texture->SetLastRenderTime(CurrentTimeForTextureTimes);
-				});
-		}
-
-		DirtyUniformBuffers = TBitMaskType(0);
-
-		return;
-	}
-#endif
-
-	while (DirtyBits)
-	{
-		// Scan for the lowest set bit, compute its index, clear it in the set of dirty bits.
-		const uint32 LowestBitMask = (DirtyBits) & (-(int32)DirtyBits);
-		const int32 BufferIndex = FMath::CountTrailingZeros(LowestBitMask); // todo: This has a branch on zero, we know it could never be zero...
-		DirtyBits ^= LowestBitMask;
-
-		check(BufferIndex < SRT.ResourceTableLayoutHashes.Num());
-
-		FRHIUniformBuffer* Buffer = BoundUniformBuffers[BufferIndex];
-
-#if DO_CHECK
-
-		if (!Buffer)
-		{
-			UE_LOG(LogRHICore, Fatal, TEXT("Shader expected a uniform buffer at slot %u but got null instead (Shader='%s' UB='%s'). Rendering code needs to set a valid uniform buffer for this slot.")
-				, BufferIndex
-				, Shader.GetShaderName()
-				, *Shader.GetUniformBufferName(BufferIndex)
-			);
-		}
-		else if (Buffer->GetLayout().GetHash() != SRT.ResourceTableLayoutHashes[BufferIndex])
-		{
-			FRHIUniformBufferLayout const& BufferLayout = Buffer->GetLayout();
-
-			FString ResourcesString;
-			for (FRHIUniformBufferResource const& Resource : BufferLayout.Resources)
-			{
-				ResourcesString += FString::Printf(TEXT("%s%d")
-					, ResourcesString.Len() ? TEXT(" ") : TEXT("")
-					, Resource.MemberType
-				);
-			}
-
-			// This might mean you are accessing a data you haven't bound e.g. GBuffer
-			UE_LOG(LogRHICore, Fatal,
-				TEXT("Uniform buffer bound to slot %u is not what the shader expected:\n")
-				TEXT("\tBound                : Uniform Buffer[%s] with Hash[0x%08x]\n")
-				TEXT("\tExpected             : Uniform Buffer[%s] with Hash[0x%08x]\n")
-				TEXT("\tShader Name          : %s\n")
-				TEXT("\tLayout CB Size       : %d\n")
-				TEXT("\tLayout Num Resources : %d\n")
-				TEXT("\tResource Types       : %s\n")
-				, BufferIndex
-				, *BufferLayout.GetDebugName(), BufferLayout.GetHash()
-				, *Shader.GetUniformBufferName(BufferIndex), SRT.ResourceTableLayoutHashes[BufferIndex]
-				, Shader.GetShaderName()
-				, BufferLayout.ConstantBufferSize
-				, BufferLayout.Resources.Num()
-				, *ResourcesString
-			);
-		}
-
-#endif // DO_CHECK
-
-		// Textures
-		EnumerateUniformBufferResources<FRHITexture>(Buffer, BufferIndex, SRT.TextureMap.GetData(),
-			[&](FRHITexture* Texture, uint8 Index)
-			{
-#if ENABLE_RHI_VALIDATION
-				if (Tracker)
-				{
-					ERHIAccess Access = IsComputeShaderFrequency(Shader.GetFrequency())
-						? ERHIAccess::SRVCompute
-						: Shader.GetFrequency() == SF_Pixel
-							? ERHIAccess::SRVGraphicsPixel
-							: ERHIAccess::SRVGraphicsNonPixel;
-
-					// Textures bound here only have their "common" plane accessible. Stencil etc is ignored.
-					// (i.e. only access the color plane of a color texture, or depth plane of a depth texture)
-					Tracker->Assert(Texture->GetViewIdentity(0, 0, 0, 0, uint32(RHIValidation::EResourcePlane::Common), 1), Access);
-				}
-#endif
-				Texture->SetLastRenderTime(CurrentTimeForTextureTimes);
-				Binder.SetTexture(Texture, Index);
-			});
-
-		// SRVs
-		EnumerateUniformBufferResources<FRHIShaderResourceView>(Buffer, BufferIndex, SRT.ShaderResourceViewMap.GetData(),
-			[&](FRHIShaderResourceView* SRV, uint8 Index)
-			{
-#if ENABLE_RHI_VALIDATION
-				if (Tracker)
-				{
-					ERHIAccess Access = IsComputeShaderFrequency(Shader.GetFrequency())
-						? ERHIAccess::SRVCompute
-						: Shader.GetFrequency() == SF_Pixel
-							? ERHIAccess::SRVGraphicsPixel
-							: ERHIAccess::SRVGraphicsNonPixel;
-
-					Tracker->Assert(SRV->GetViewIdentity(), Access);
-				}
-				if (GRHIValidationEnabled)
-				{
-					RHIValidation::ValidateShaderResourceView(&Shader, Index, SRV);
-				}
-#endif
-				Binder.SetSRV(SRV, Index);
-			});
-
-#if PLATFORM_SUPPORTS_BINDLESS_RENDERING
-		EnumerateUniformBufferResources<FRHIResourceCollection>(Buffer, BufferIndex, SRT.ResourceCollectionMap.GetData(),
-			[&](FRHIResourceCollection* ResourceCollection, uint8 Index)
-			{
-#if ENABLE_RHI_VALIDATION
-				// todo: christopher.waters - ResourceCollection validation
-#endif
-				Binder.SetResourceCollection(ResourceCollection, Index);
-			});
-#endif // PLATFORM_SUPPORTS_BINDLESS_RENDERING
-
-		// Samplers
-		EnumerateUniformBufferResources<FRHISamplerState>(Buffer, BufferIndex, SRT.SamplerMap.GetData(),
-			[&](FRHISamplerState* Sampler, uint8 Index)
-			{
-				Binder.SetSampler(Sampler, Index);
-			});
-
-		// UAVs
-		EnumerateUniformBufferResources<FRHIUnorderedAccessView>(Buffer, BufferIndex, SRT.UnorderedAccessViewMap.GetData(),
-			[&](FRHIUnorderedAccessView* UAV, uint8 Index)
-			{
-#if ENABLE_RHI_VALIDATION
-				if (Tracker)
-				{
-					ERHIAccess Access = IsComputeShaderFrequency(Shader.GetFrequency())
-						? ERHIAccess::UAVCompute
-						: ERHIAccess::UAVGraphics;
-
-					Tracker->AssertUAV(UAV, Access, Index);
-				}
-#endif
-				Binder.SetUAV(UAV, Index);
-			});
-	}
-
-	DirtyUniformBuffers = TBitMaskType(0);
-}
-
-
-template <typename TBinder, typename TUniformBufferArrayType, typename TBitMaskType>
-void SetFullyBindlessResourcesFromTables(TBinder&& Binder, const FRHIShader& Shader, TBitMaskType& DirtyUniformBuffers, const TUniformBufferArrayType& BoundUniformBuffers
-#if ENABLE_RHI_VALIDATION
-	, RHIValidation::FTracker* Tracker
-#endif
-)
-{
-	SetResourcesFromTables<TBinder, TUniformBufferArrayType, TBitMaskType, true>(
-		MoveTemp(Binder)
-		, Shader
-		, DirtyUniformBuffers
-		, BoundUniformBuffers
-#if ENABLE_RHI_VALIDATION
-		, Tracker
-#endif
-	);
-}
 
 } //! RHICore
 } //! UE
diff --git a/Engine/Source/Runtime/VulkanRHI/Private/VulkanCommands.cpp b/Engine/Source/Runtime/VulkanRHI/Private/VulkanCommands.cpp
--- a/Engine/Source/Runtime/VulkanRHI/Private/VulkanCommands.cpp
+++ b/Engine/Source/Runtime/VulkanRHI/Private/VulkanCommands.cpp
@@ -15,6 +15,7 @@
 #include "RHIUtilities.h"
 #include "VulkanResourceCollection.h"
 #include "RHICoreShader.h"
+#include "RHIUniformBufferUtilities.h"
 
 static TAutoConsoleVariable<int32> GCVarSubmitOnDispatch(
 	TEXT("r.Vulkan.SubmitOnDispatch"),
@@ -215,7 +216,7 @@ void FVulkanCommandListContext::SetResourcesFromTables(const ShaderType* Shader)
 	if (Frequency == SF_Compute)
 	{
 		FVulkanResourceBinder Binder(*this, Frequency, PendingComputeState);
-		UE::RHICore::SetResourcesFromTables(
+		UE::RHI::Private::SetUniformBufferResourcesFromTables(
 			Binder
 			, *Shader
 			, DirtyUniformBuffers[Frequency]
@@ -228,7 +229,7 @@ void FVulkanCommandListContext::SetResourcesFromTables(const ShaderType* Shader)
 	else
 	{
 		FVulkanResourceBinder Binder(*this, Frequency, PendingGfxState);
-		UE::RHICore::SetResourcesFromTables(
+		UE::RHI::Private::SetUniformBufferResourcesFromTables(
 			Binder
 			, *Shader
 			, DirtyUniformBuffers[Frequency]
diff --git a/Engine/Source/Runtime/VulkanRHI/Private/VulkanRayTracing.cpp b/Engine/Source/Runtime/VulkanRHI/Private/VulkanRayTracing.cpp
--- a/Engine/Source/Runtime/VulkanRHI/Private/VulkanRayTracing.cpp
+++ b/Engine/Source/Runtime/VulkanRHI/Private/VulkanRayTracing.cpp
@@ -1430,7 +1430,7 @@ void FVulkanDevice::CleanUpRayTracing()
 
 
 FVulkanRayTracingPipelineState::FVulkanRayTracingPipelineState(FVulkanDevice* const InDevice, const FRayTracingPipelineStateInitializer& Initializer)
-	: FDeviceChild(InDevice)
+	: FRHIRayTracingPipelineState(Initializer), FDeviceChild(InDevice)
 {
 	checkf(InDevice->SupportsBindless(), TEXT("Vulkan ray tracing pipelines are only supported in bindless."));
 
diff --git a/Engine/Source/Runtime/Windows/D3D11RHI/Private/D3D11Commands.cpp b/Engine/Source/Runtime/Windows/D3D11RHI/Private/D3D11Commands.cpp
--- a/Engine/Source/Runtime/Windows/D3D11RHI/Private/D3D11Commands.cpp
+++ b/Engine/Source/Runtime/Windows/D3D11RHI/Private/D3D11Commands.cpp
@@ -15,6 +15,7 @@
 #include "SceneUtils.h"
 #include "EngineGlobals.h"
 #include "RHIShaderParametersShared.h"
+#include "RHIUniformBufferUtilities.h"
 
 // For Depth Bounds Test interface
 #include "Windows/AllowWindowsPlatformTypes.h"
@@ -1012,7 +1013,7 @@ void FD3D11DynamicRHI::SetResourcesFromTables(const ShaderType* RESTRICT Shader)
 	checkSlow(Shader);
 	static constexpr EShaderFrequency Frequency = static_cast<EShaderFrequency>(ShaderType::StaticFrequency);
 
-	UE::RHICore::SetResourcesFromTables(
+	UE::RHI::Private::SetUniformBufferResourcesFromTables(
 		  FD3D11ResourceBinder<Frequency> { *this }
 		, *Shader
 		, DirtyUniformBuffers[Frequency]
