Improved static_assert message for TCHAR string literals in RHI breadcrumbs
 - Clarify that RHI_BREADCRUMB_FORCE_STRING_LITERAL can be used to silence the static_assert if required.

diff --git a/Engine/Source/Runtime/RHI/Public/RHIBreadcrumbs.h b/Engine/Source/Runtime/RHI/Public/RHIBreadcrumbs.h
--- a/Engine/Source/Runtime/RHI/Public/RHIBreadcrumbs.h
+++ b/Engine/Source/Runtime/RHI/Public/RHIBreadcrumbs.h
@@ -634,7 +634,11 @@
 	template <typename T>
 	struct TRHIBreadcrumbValue
 	{
-		static_assert(!(std::is_same_v<T, TCHAR*> || std::is_same_v<T, TCHAR const*>), "Do not use raw TCHAR pointers with RHI breadcrumbs. Pass the FString, FName, or string literal instead.");
+		static_assert(!(std::is_same_v<T, TCHAR*> || std::is_same_v<T, TCHAR const*>),
+			"Do not use raw TCHAR pointers with RHI breadcrumbs. Pass the FString, FName, or string literal instead. If you are certain your TCHAR pointer is a string literal "
+			"(e.g. from a function returning a literal) and you cannot pass that literal directly to an RHI breadcrumb, use the RHI_BREADCRUMB_FORCE_STRING_LITERAL macro to silence "
+			"this static assert. Incorrect use of RHI_BREADCRUMB_FORCE_STRING_LITERAL will lead to use-after-free, as only the raw pointer is retained by the breadcrumb."
+		);
 
 		T Value;
 		TRHIBreadcrumbValue(T const& Value)
@@ -920,9 +924,12 @@
 	#define RHI_BREADCRUMB_EVENT_CONDITIONAL_STR_DEPRECATED(...)
 
 #endif
+
 //
-// Used to override the static_assert caused when choosing between two string literals with a ternary operator, like so:
-// SCOPED_DRAW_EVENTF(RHICmdList, EventName, TEXT("Name=%s"), RHI_BREADCRUMB_FORCE_STRING_LITERAL(bCondition ? TEXT("True") : TEXT("False"))
+// Used to override the static_assert check for string literals in RHI breadcrumbs. This is required when using
+// literals returned by functions, or choosing between two string literals with a ternary operator, like so:
+// 
+//    SCOPED_DRAW_EVENTF(RHICmdList, EventName, TEXT("Name=%s"), RHI_BREADCRUMB_FORCE_STRING_LITERAL(bCondition ? TEXT("True") : TEXT("False"))
 // 
 // !! DO NOT USE THIS MACRO FOR NON-STRING LITERALS !!
 //
