Issue a render thread fence when updating vertex stream slots. Fixes a hypothetical race where slots are updated that could technically overlap with prior draw passes.

diff --git a/Engine/Source/Runtime/RHI/Public/RHICommandList.h b/Engine/Source/Runtime/RHI/Public/RHICommandList.h
--- a/Engine/Source/Runtime/RHI/Public/RHICommandList.h
+++ b/Engine/Source/Runtime/RHI/Public/RHICommandList.h
@@ -829,10 +829,18 @@ public:
 	FORCEINLINE void UpdateStreamSourceSlot(FRHIStreamSourceSlot* StreamSourceSlotRHI, FRHIBuffer* BufferRHI)
 	{
 		check(StreamSourceSlotRHI);
-		EnqueueLambda([this, StreamSourceSlotRHI, BufferRHI] (FRHICommandListBase&)
+		if (Bypass())
 		{
 			StreamSourceSlotRHI->Buffer = BufferRHI;
-		});
+		}
+		else
+		{
+			EnqueueLambda([this, StreamSourceSlotRHI, BufferRHI] (FRHICommandListBase&)
+			{
+				StreamSourceSlotRHI->Buffer = BufferRHI;
+			});
+			RHIThreadFence(true);
+		}
 	}
 
 	FORCEINLINE void UpdateTexture2D(FRHITexture* Texture, uint32 MipIndex, const struct FUpdateTextureRegion2D& UpdateRegion, uint32 SourcePitch, const uint8* SourceData)
