Fix some code in ITargetDevice.h and RHIDefinitions.h that causes SwitchUnhandledEnumerator warnings when included by modules that have those warnings enabled.
These changes may affect behavior as they add some switch cases whose absence seems to be wrong.

diff --git a/Engine/Source/Developer/TargetPlatform/Public/Interfaces/ITargetDevice.h b/Engine/Source/Developer/TargetPlatform/Public/Interfaces/ITargetDevice.h
--- a/Engine/Source/Developer/TargetPlatform/Public/Interfaces/ITargetDevice.h
+++ b/Engine/Source/Developer/TargetPlatform/Public/Interfaces/ITargetDevice.h
@@ -110,6 +110,10 @@ namespace TargetDeviceTypes
 		case ETargetDeviceTypes::Tablet:
 			return FString("Tablet");
 
+		case ETargetDeviceTypes::HMD:
+			return FString("HMD");
+
+		case ETargetDeviceTypes::Indeterminate:
 		default:
 			return FString("Indeterminate");
 		}
@@ -143,6 +147,7 @@ namespace TargetDeviceConnectionTypes
 		case ETargetDeviceConnectionTypes::Proprietary:
 			return FString("Proprietary");
 
+		case ETargetDeviceConnectionTypes::Unknown:
 		default:
 			return FString("Unknown");
 		}
diff --git a/Engine/Source/Runtime/RHI/Public/RHIDefinitions.h b/Engine/Source/Runtime/RHI/Public/RHIDefinitions.h
--- a/Engine/Source/Runtime/RHI/Public/RHIDefinitions.h
+++ b/Engine/Source/Runtime/RHI/Public/RHIDefinitions.h
@@ -547,7 +547,12 @@ inline bool IsResourceBindingTypeSRV(EShaderCodeResourceBindingType Type)
 	case EShaderCodeResourceBindingType::RWStructuredBuffer:
 	case EShaderCodeResourceBindingType::RWBuffer:
 	case EShaderCodeResourceBindingType::RasterizerOrderedTexture2D:
+	case EShaderCodeResourceBindingType::SamplerState:
+	case EShaderCodeResourceBindingType::TextureMetadata:
+	case EShaderCodeResourceBindingType::RWTextureMetadata:
 		return false;
+	case EShaderCodeResourceBindingType::MAX:
+	case EShaderCodeResourceBindingType::Invalid:
 	default:
 		ensureMsgf(0, TEXT("Missing or invalid SRV or UAV Type"));
 	}
@@ -1456,6 +1461,8 @@ inline EGpuVendorId RHIConvertToGpuVendorId(uint32 VendorId)
 		return EGpuVendorId::NotQueried;
 
 	case EGpuVendorId::Amd:
+	case EGpuVendorId::Kazan:
+	case EGpuVendorId::Codeplay:
 	case EGpuVendorId::Mesa:
 	case EGpuVendorId::ImgTec:
 	case EGpuVendorId::Nvidia:
@@ -1463,20 +1470,40 @@ inline EGpuVendorId RHIConvertToGpuVendorId(uint32 VendorId)
 	case EGpuVendorId::Broadcom:
 	case EGpuVendorId::Qualcomm:
 	case EGpuVendorId::Intel:
+	case EGpuVendorId::Apple:
+	case EGpuVendorId::Vivante:
+	case EGpuVendorId::VeriSilicon:
 	case EGpuVendorId::SamsungAMD:
 	case EGpuVendorId::Microsoft:
 		return (EGpuVendorId)VendorId;
 
+	case EGpuVendorId::Unknown:
 	default:
-		break;
+		return EGpuVendorId::Unknown;
 	}
-
-	return EGpuVendorId::Unknown;
 }
 
 inline bool IsGeometryPipelineShaderFrequency(EShaderFrequency Frequency)
 {
-	return Frequency == SF_Mesh || Frequency == SF_Amplification;
+	switch (Frequency)
+	{
+	case SF_Mesh:
+	case SF_Amplification:
+		return true;
+	case SF_Vertex:
+	case SF_Pixel:
+	case SF_Geometry:
+	case SF_Compute:
+	case SF_RayGen:
+	case SF_RayMiss:
+	case SF_RayHitGroup:
+	case SF_RayCallable:
+	case SF_WorkGraphRoot:
+	case SF_WorkGraphComputeNode:
+	case SF_NumFrequencies:
+	default:
+		return false;
+	}
 }
 
 inline bool IsRayTracingShaderFrequency(EShaderFrequency Frequency)
@@ -1488,6 +1515,15 @@ inline bool IsRayTracingShaderFrequency(EShaderFrequency Frequency)
 	case SF_RayHitGroup:
 	case SF_RayCallable:
 		return true;
+	case SF_Vertex:
+	case SF_Mesh:
+	case SF_Amplification:
+	case SF_Pixel:
+	case SF_Geometry:
+	case SF_Compute:
+	case SF_WorkGraphRoot:
+	case SF_WorkGraphComputeNode:
+	case SF_NumFrequencies:
 	default:
 		return false;
 	}
@@ -1500,6 +1536,17 @@ inline bool IsWorkGraphShaderFrequency(EShaderFrequency Frequency)
 	case SF_WorkGraphRoot:
 	case SF_WorkGraphComputeNode:
 		return true;
+	case SF_Vertex:
+	case SF_Mesh:
+	case SF_Amplification:
+	case SF_Pixel:
+	case SF_Geometry:
+	case SF_Compute:
+	case SF_RayGen:
+	case SF_RayMiss:
+	case SF_RayHitGroup:
+	case SF_RayCallable:
+	case SF_NumFrequencies:
 	default:
 		return false;
 	}
