Remove some redundant forward declares of old RHI functions.

diff --git a/Engine/Source/Runtime/Engine/Public/RawIndexBuffer.h b/Engine/Source/Runtime/Engine/Public/RawIndexBuffer.h
--- a/Engine/Source/Runtime/Engine/Public/RawIndexBuffer.h
+++ b/Engine/Source/Runtime/Engine/Public/RawIndexBuffer.h
@@ -502,7 +502,7 @@ public:
 	{
 		if (CachedNumIndices)
 		{
-			// Need to cache number of indices from the source array *before* RHICreateIndexBuffer is called
+			// Need to cache number of indices from the source array *before* RHICreateBuffer is called
 			// because it will empty the source array.
 			CachedNumIndices = Indices.Num();
 
diff --git a/Engine/Source/Runtime/RHI/Public/RHICommandList.h b/Engine/Source/Runtime/RHI/Public/RHICommandList.h
--- a/Engine/Source/Runtime/RHI/Public/RHICommandList.h
+++ b/Engine/Source/Runtime/RHI/Public/RHICommandList.h
@@ -4752,12 +4752,6 @@ public:
 	~FScopedRHIThreadStaller();
 };
 
-
-// Forward declare RHI creation function so they can still be called from the deprecated immediate command list resource creation functions
-FBufferRHIRef RHICreateIndexBuffer(uint32 Stride, uint32 Size, uint32 InUsage, FRHIResourceCreateInfo& CreateInfo);
-FBufferRHIRef RHICreateVertexBuffer(uint32 Size, uint32 InUsage, FRHIResourceCreateInfo& CreateInfo);
-FBufferRHIRef RHICreateStructuredBuffer(uint32 Stride, uint32 Size, uint32 InUsage, FRHIResourceCreateInfo& CreateInfo);
-
 extern RHI_API ERHIAccess RHIGetDefaultResourceState(ETextureCreateFlags InUsage, bool bInHasInitialData);
 extern RHI_API ERHIAccess RHIGetDefaultResourceState(EBufferUsageFlags InUsage, bool bInHasInitialData);
 
diff --git a/Engine/Source/Runtime/RHI/Public/RHIUtilities.h b/Engine/Source/Runtime/RHI/Public/RHIUtilities.h
--- a/Engine/Source/Runtime/RHI/Public/RHIUtilities.h
+++ b/Engine/Source/Runtime/RHI/Public/RHIUtilities.h
@@ -217,7 +217,7 @@ struct FRWBuffer
 		Release();
 	}
 
-	// @param AdditionalUsage passed down to RHICreateVertexBuffer(), get combined with "BUF_UnorderedAccess | BUF_ShaderResource" e.g. BUF_Static
+	// @param AdditionalUsage passed down to RHICreateBuffer(), get combined with "BUF_UnorderedAccess | BUF_ShaderResource" e.g. BUF_Static
 	void Initialize(FRHICommandListBase& RHICmdList, const TCHAR* InDebugName, uint32 BytesPerElement, uint32 NumElements, EPixelFormat Format, ERHIAccess InResourceState, EBufferUsageFlags AdditionalUsage = BUF_None, FResourceArrayUploadInterface* InResourceArray = nullptr)
 	{
 		// Provide a debug name if using Fast VRAM so the allocators diagnostics will work
