Serialize RHIShaderBindingLayout used for RTPSO compilation in the shader pipeline cache entries for PSO precompilation

diff --git a/Engine/Source/Runtime/D3D12RHI/Private/D3D12RayTracing.cpp b/Engine/Source/Runtime/D3D12RHI/Private/D3D12RayTracing.cpp
--- a/Engine/Source/Runtime/D3D12RHI/Private/D3D12RayTracing.cpp
+++ b/Engine/Source/Runtime/D3D12RHI/Private/D3D12RayTracing.cpp
@@ -1940,7 +1940,7 @@ public:
 		TRACE_CPUPROFILER_EVENT_SCOPE(RTPSO_CreatePipeline);
 
 		checkf(Initializer.GetRayGenTable().Num() > 0 || Initializer.bPartial, TEXT("Ray tracing pipelines must have at leat one ray generation shader."));
-		checkf(Initializer.GetHitGroupTable().Num() > 0, TEXT("Ray tracing pipelines must have at leat one hit shader."));
+		checkf(Initializer.bPartial || Initializer.GetHitGroupTable().Num() > 0, TEXT("Ray tracing pipelines must have at leat one hit shader."));
 
 		uint64 TotalCreationTime = 0;
 		uint64 CompileTime = 0;
diff --git a/Engine/Source/Runtime/RHI/Private/PipelineFileCache.cpp b/Engine/Source/Runtime/RHI/Private/PipelineFileCache.cpp
--- a/Engine/Source/Runtime/RHI/Private/PipelineFileCache.cpp
+++ b/Engine/Source/Runtime/RHI/Private/PipelineFileCache.cpp
@@ -84,12 +84,13 @@ enum class EPipelineCacheFileFormatVersions : uint32
 	BeforeStableCacheVersioning = 26,
 	RemovingLineAA = 27,
 	AddingDepthBounds = 28,
+	AddRTPSOShaderBindingLayout = 29,
 };
 
 const uint64 FPipelineCacheFileFormatMagic = 0x5049504543414348; // PIPECACH
 const uint64 FPipelineCacheTOCFileFormatMagic = 0x544F435354415232; // TOCSTAR2
 const uint64 FPipelineCacheEOFFileFormatMagic = 0x454F462D4D41524B; // EOF-MARK
-const RHI_API uint32 FPipelineCacheFileFormatCurrentVersion = (uint32)EPipelineCacheFileFormatVersions::AddingDepthBounds;
+const RHI_API uint32 FPipelineCacheFileFormatCurrentVersion = (uint32)EPipelineCacheFileFormatVersions::AddRTPSOShaderBindingLayout;
 const int32  FPipelineCacheGraphicsDescPartsNum = 67; // parser will expect this number of parts in a description string
 
 /**
@@ -1413,6 +1414,12 @@ bool FPipelineCacheFileFormatPSO::Verify() const
 			Ar << Frequency;
 			Info.RayTracingDesc.Frequency = EShaderFrequency(Frequency);
 			
+			if (Ar.GameNetVer() >= (uint32)EPipelineCacheFileFormatVersions::AddRTPSOShaderBindingLayout)
+			{
+				// Serialize RHIShaderBindingLayout
+				Ar << Info.RayTracingDesc.ShaderBindingLayout;		
+			}
+			
 			break;
 		}
 		default:
@@ -4357,21 +4364,52 @@ bool FPipelineFileCacheManager::MergePipelineFileCaches(FString const& PathA, FS
 FPipelineCacheFileFormatPSO::FPipelineFileCacheRayTracingDesc::FPipelineFileCacheRayTracingDesc(const FRayTracingPipelineStateInitializer& Initializer, const FRHIRayTracingShader* ShaderRHI)
 : ShaderHash(ShaderRHI->GetHash())
 , Frequency(ShaderRHI->GetFrequency())
+, ShaderBindingLayout(Initializer.ShaderBindingLayout ? *Initializer.ShaderBindingLayout : FRHIShaderBindingLayout())
 {
 }
 
 FString FPipelineCacheFileFormatPSO::FPipelineFileCacheRayTracingDesc::HeaderLine() const
 {
-	return FString(TEXT("RayTracingShader,DeprecatedMaxPayloadSizeInBytes,Frequency"));
+	FString Result(TEXT("RayTracingShader,DeprecatedMaxPayloadSizeInBytes,Frequency,ShaderBindingLayout-Hash,ShaderBindingLayout-Flags,ShaderBindingLayout-NumUBEntries"));
+	for (int32 Index = 0; Index < FRHIShaderBindingLayout::MaxUniformBufferEntries; Index++)
+	{
+		Result += FString::Printf(TEXT(",%s%d,%s%d,%s%d,%s%d,%s%d,%s%d")
+			, TEXT("ShaderBindingLayout-UBEntry-Name"), Index
+			, TEXT("ShaderBindingLayout-UBEntry-RegisterSpace"), Index
+			, TEXT("ShaderBindingLayout-UBEntry-CBVResourceIndex"), Index
+			, TEXT("ShaderBindingLayout-UBEntry-BaseSRVResourceIndex"), Index
+			, TEXT("ShaderBindingLayout-UBEntry-BaseUAVResourceIndex"), Index
+			, TEXT("ShaderBindingLayout-UBEntry-BaseSamplerResourceIndex"), Index
+		);
+	}
+
+	return Result;
 }
 
 FString FPipelineCacheFileFormatPSO::FPipelineFileCacheRayTracingDesc::ToString() const
 {
-	return FString::Printf(TEXT("%s,%d,%d")
+	FString Result = FString::Printf(TEXT("%s,%d,%d,%d,%d,%d")
 		, *ShaderHash.ToString()
 		, DeprecatedMaxPayloadSizeInBytes
 		, uint32(Frequency)
+		, ShaderBindingLayout.GetHash()
+		, ShaderBindingLayout.GetFlags()
+		, ShaderBindingLayout.GetNumUniformBufferEntries()
 	);
+
+	for (uint32 Index = 0; Index < FRHIShaderBindingLayout::MaxUniformBufferEntries; Index++)
+	{
+		FRHIUniformBufferShaderBindingLayout Entry = Index < ShaderBindingLayout.GetNumUniformBufferEntries() ? ShaderBindingLayout.GetUniformBufferEntry(Index) : FRHIUniformBufferShaderBindingLayout();
+		Result += FString::Printf(TEXT(",%s,%d,%d,%d,%d,%d")
+			, Entry.LayoutName.Len() > 0 ? *Entry.LayoutName : TEXT("None")
+			, Entry.RegisterSpace
+			, Entry.CBVResourceIndex
+			, Entry.BaseSRVResourceIndex
+			, Entry.BaseUAVResourceIndex
+			, Entry.BaseSamplerResourceIndex
+		);
+	}
+	return Result;
 }
 
 void FPipelineCacheFileFormatPSO::FPipelineFileCacheRayTracingDesc::AddToReadableString(TReadableStringBuilder& OutBuilder) const
@@ -4401,16 +4439,51 @@ void FPipelineCacheFileFormatPSO::FPipelineFileCacheRayTracingDesc::FromString(c
 	TArray<FString> Parts;
 	Src.TrimStartAndEnd().ParseIntoArray(Parts, TEXT(","));
 
+	// make sure we have required number of parts
+	if (Parts.Num() != (6 + 6 * FRHIShaderBindingLayout::MaxUniformBufferEntries))
+	{
+		return;
+	}
+
 	ShaderHash.FromString(Parts[0]);
 
 	// Not used, but kept for back-compatibility
 	LexFromString(DeprecatedMaxPayloadSizeInBytes, Parts[1]);
 
+	uint32 Temp = 0;
+	LexFromString(Temp, Parts[2]);
+	Frequency = EShaderFrequency(Temp);
+
+	uint32 ShaderBiningLayoutHash;
+	LexFromString(ShaderBiningLayoutHash, Parts[3]);
+
+	LexFromString(Temp, Parts[4]);
+	EShaderBindingLayoutFlags ShaderBindingLayoutFlags = (EShaderBindingLayoutFlags) Temp;
+
+	uint32 ShaderBindingLayoutNumUniformBuffers;
+	LexFromString(ShaderBindingLayoutNumUniformBuffers, Parts[5]);
+
+	uint32 CurrentPartIndex = 6;
+	TArray<FRHIUniformBufferShaderBindingLayout> ShaderBindingLayoutUBEntries;
+	ShaderBindingLayoutUBEntries.SetNum(ShaderBindingLayoutNumUniformBuffers);
+	for (uint32 UBEntryIndex = 0; UBEntryIndex < ShaderBindingLayoutNumUniformBuffers; ++UBEntryIndex)
 	{
-		uint32 Temp = 0;
-		LexFromString(Temp, Parts[2]);
-		Frequency = EShaderFrequency(Temp);
+		FRHIUniformBufferShaderBindingLayout& UBEntry = ShaderBindingLayoutUBEntries[UBEntryIndex];
+		UBEntry.LayoutName = Parts[CurrentPartIndex++];
+		LexFromString(Temp, Parts[CurrentPartIndex++]);
+		UBEntry.RegisterSpace = Temp;
+		LexFromString(Temp, Parts[CurrentPartIndex++]);
+		UBEntry.CBVResourceIndex = Temp;
+		LexFromString(Temp, Parts[CurrentPartIndex++]);
+		UBEntry.BaseSRVResourceIndex = Temp;
+		LexFromString(Temp, Parts[CurrentPartIndex++]);
+		UBEntry.BaseUAVResourceIndex = Temp;
+		LexFromString(Temp, Parts[CurrentPartIndex++]);
+		UBEntry.BaseSamplerResourceIndex = Temp;
 	}	
+
+	ShaderBindingLayout = FRHIShaderBindingLayout(ShaderBindingLayoutFlags, ShaderBindingLayoutUBEntries);
+	check(ShaderBindingLayout.GetHash() == ShaderBiningLayoutHash);
 }
 
 bool FPipelineCacheFileFormatPSO::Init(FPipelineCacheFileFormatPSO& PSO, FPipelineCacheFileFormatPSO::FPipelineFileCacheRayTracingDesc const& Desc)
diff --git a/Engine/Source/Runtime/RHI/Public/PipelineFileCache.h b/Engine/Source/Runtime/RHI/Public/PipelineFileCache.h
--- a/Engine/Source/Runtime/RHI/Public/PipelineFileCache.h
+++ b/Engine/Source/Runtime/RHI/Public/PipelineFileCache.h
@@ -6,6 +6,7 @@
 #include "Containers/List.h"
 #include "Containers/StringView.h"
 #include "RHI.h"
+#include "RHIShaderBindingLayout.h"
 
 DECLARE_STATS_GROUP(TEXT("ShaderPipelineCache"),STATGROUP_PipelineStateCache, STATCAT_Advanced);
 
@@ -175,6 +176,7 @@ struct FPipelineCacheFileFormatPSO
 		FSHAHash ShaderHash;
 		uint32 DeprecatedMaxPayloadSizeInBytes = 0;
 		EShaderFrequency Frequency = SF_RayGen;
+		FRHIShaderBindingLayout ShaderBindingLayout;
 		
 		FPipelineFileCacheRayTracingDesc() = default;
 		FPipelineFileCacheRayTracingDesc(const FRayTracingPipelineStateInitializer& Initializer, const FRHIRayTracingShader* ShaderRHI);
@@ -186,14 +188,14 @@ struct FPipelineCacheFileFormatPSO
 
 		friend uint32 GetTypeHash(const FPipelineFileCacheRayTracingDesc& Desc)
 		{
-			return GetTypeHash(Desc.ShaderHash) ^
-				GetTypeHash(Desc.Frequency);
+			return HashCombineFast(GetTypeHash(Desc.ShaderHash), HashCombineFast(GetTypeHash(Desc.Frequency), GetTypeHash(Desc.ShaderBindingLayout)));
 		}
 
 		bool operator == (const FPipelineFileCacheRayTracingDesc& Other) const
 		{
 			return ShaderHash == Other.ShaderHash &&
-				Frequency == Other.Frequency;
+				Frequency == Other.Frequency &&
+				ShaderBindingLayout == Other.ShaderBindingLayout;
 		}
 	};
 	enum class DescriptorType : uint32
diff --git a/Engine/Source/Runtime/RenderCore/Private/ShaderPipelineCache.cpp b/Engine/Source/Runtime/RenderCore/Private/ShaderPipelineCache.cpp
--- a/Engine/Source/Runtime/RenderCore/Private/ShaderPipelineCache.cpp
+++ b/Engine/Source/Runtime/RenderCore/Private/ShaderPipelineCache.cpp
@@ -1138,31 +1138,32 @@ bool FShaderPipelineCacheTask::Precompile(FRHICommandListImmediate& RHICmdList,
 				if (RayTracingShader.IsValid())
 				{
 					FRayTracingPipelineStateInitializer Initializer;
-						Initializer.bPartial = true; // Indicates that this RTPSO is used only as input for later RTPSO linking step (not for rendering)						
-						Initializer.MaxPayloadSizeInBytes = RayTracingShader->RayTracingPayloadSize;
+					Initializer.bPartial = true; // Indicates that this RTPSO is used only as input for later RTPSO linking step (not for rendering)						
+					Initializer.MaxPayloadSizeInBytes = RayTracingShader->RayTracingPayloadSize;
+					Initializer.ShaderBindingLayout = PSO.RayTracingDesc.ShaderBindingLayout.GetNumUniformBufferEntries() > 0 ? &PSO.RayTracingDesc.ShaderBindingLayout : nullptr;
 
-						FRHIRayTracingShader* ShaderTable[] =
-						{
-							RayTracingShader.GetReference()
-						};
+					FRHIRayTracingShader* ShaderTable[] =
+					{
+						RayTracingShader.GetReference()
+					};
 
-						switch (PSO.RayTracingDesc.Frequency)
-						{
-						case SF_RayGen:
-							Initializer.SetRayGenShaderTable(ShaderTable);
-							break;
-						case SF_RayMiss:
-							Initializer.SetMissShaderTable(ShaderTable);
-							break;
-						case SF_RayHitGroup:
-							Initializer.SetHitGroupTable(ShaderTable);
-							break;
-						case SF_RayCallable:
-							Initializer.SetCallableTable(ShaderTable);
-							break;
-						default:
-							checkNoEntry();
-						}
+					switch (PSO.RayTracingDesc.Frequency)
+					{
+					case SF_RayGen:
+						Initializer.SetRayGenShaderTable(ShaderTable);
+						break;
+					case SF_RayMiss:
+						Initializer.SetMissShaderTable(ShaderTable);
+						break;
+					case SF_RayHitGroup:
+						Initializer.SetHitGroupTable(ShaderTable);
+						break;
+					case SF_RayCallable:
+						Initializer.SetCallableTable(ShaderTable);
+						break;
+					default:
+						checkNoEntry();
+					}
 
 					FRayTracingPipelineState* RayTracingPipeline = PipelineStateCache::GetAndOrCreateRayTracingPipelineState(RHICmdList, Initializer);
 					bOk = RayTracingPipeline != nullptr;
