Improve new ProfileGPU command output
 - Formatted output into a unicode table with auto-sized columns
 - Added exclusive stats.
 - Implemented r.ProfileGPU.Root wildcard search. Nodes not matching the wildcard are discarded. Nodes that match include their entire subtree
 - Implemented r.ProfileGPU.ThresholdPercent. Nodes whose inclusive time is less than the threshold are hidden.

diff --git a/Engine/Source/Runtime/RHI/Private/GPUProfiler.cpp b/Engine/Source/Runtime/RHI/Private/GPUProfiler.cpp
--- a/Engine/Source/Runtime/RHI/Private/GPUProfiler.cpp
+++ b/Engine/Source/Runtime/RHI/Private/GPUProfiler.cpp
@@ -61,27 +61,33 @@ static TAutoConsoleVariable<FString> GProfileGPUPatternCVar(
 	TEXT("e.g. AmbientOcclusionSetup, AmbientOcclusion*, Ambient???lusion*, *"),
 	ECVF_Default);
 
-static TAutoConsoleVariable<FString> GProfileGPURootCVar(
+static TAutoConsoleVariable<FString> GCVarProfileGPU_Root(
 	TEXT("r.ProfileGPU.Root"),
 	TEXT("*"),
 	TEXT("Allows to filter the tree when using ProfileGPU, the pattern match is case sensitive."),
 	ECVF_Default);
 
-static TAutoConsoleVariable<float> GProfileThresholdPercent(
+static TAutoConsoleVariable<float> GCVarProfileGPU_ThresholdPercent(
 	TEXT("r.ProfileGPU.ThresholdPercent"),
 	0.0f,
 	TEXT("Percent of the total execution duration the event needs to be larger than to be printed."),
 	ECVF_Default);
 
+static TAutoConsoleVariable<bool> GCVarProfileGPU_UnicodeOutput(
+	TEXT("r.ProfileGPU.UnicodeOutput"),
+	true,
+	TEXT("When enabled, the output results will be formatted in a unicode table."),
+	ECVF_Default);
+
 static TAutoConsoleVariable<int32> GProfileShowEventHistogram(
 	TEXT("r.ProfileGPU.ShowEventHistogram"),
 	0,
 	TEXT("Whether the event histogram should be shown."),
 	ECVF_Default);
 
-static TAutoConsoleVariable<int32> GProfileGPUShowEvents(
+static TAutoConsoleVariable<bool> GCVarProfileGPU_ShowLeafEvents(
 	TEXT("r.ProfileGPU.ShowLeafEvents"),
-	1,
+	true,
 	TEXT("Allows profileGPU to display event-only leaf nodes with no draws associated."),
 	ECVF_Default);
 
@@ -225,7 +231,7 @@ struct FGPUProfileStatSummary
 		, bGatherSummaryStats(false)
 		, bDumpEventLeafNodes(false)
 	{
-		bDumpEventLeafNodes = GProfileGPUShowEvents.GetValueOnRenderThread() != 0;
+		bDumpEventLeafNodes = GCVarProfileGPU_ShowLeafEvents.GetValueOnRenderThread() != 0;
 		bGatherSummaryStats = GProfilePrintAssetSummary.GetValueOnRenderThread() != 0;
 	}
 
@@ -321,7 +327,7 @@ static void DumpStatsEventNode(FGPUProfilerEventNode* Node, float RootResult, in
 		Summary.TotalNumDraws += Node->NumDraws;
 		// Percent that this node was of the total frame time
 		const float Percent = Node->TimingResult * 100.0f / (RootResult * 1000.0f);
-		const float PercentThreshold = GProfileThresholdPercent.GetValueOnRenderThread();
+		const float PercentThreshold = GCVarProfileGPU_ThresholdPercent.GetValueOnRenderThread();
 		const int32 EffectiveDepth = FMath::Max(Depth - 1, 0);
 		const bool bDisplayEvent = (bParentMatchedFilter || WildcardFilter.IsMatch(Node->Name)) && (Percent > PercentThreshold || Summary.bDumpEventLeafNodes);
 
@@ -485,14 +491,14 @@ void FGPUProfilerEventNodeFrame::DumpEventTree()
 
 		FString ConfigString;
 
-		if (GProfileGPURootCVar.GetValueOnRenderThread() != TEXT("*"))
+		if (GCVarProfileGPU_Root.GetValueOnRenderThread() != TEXT("*"))
 		{
-			ConfigString += FString::Printf(TEXT("Root filter: %s "), *GProfileGPURootCVar.GetValueOnRenderThread());
+			ConfigString += FString::Printf(TEXT("Root filter: %s "), *GCVarProfileGPU_Root.GetValueOnRenderThread());
 		}
 
-		if (GProfileThresholdPercent.GetValueOnRenderThread() > 0.0f)
+		if (GCVarProfileGPU_ThresholdPercent.GetValueOnRenderThread() > 0.0f)
 		{
-			ConfigString += FString::Printf(TEXT("Threshold: %.2f%% "), GProfileThresholdPercent.GetValueOnRenderThread());
+			ConfigString += FString::Printf(TEXT("Threshold: %.2f%% "), GCVarProfileGPU_ThresholdPercent.GetValueOnRenderThread());
 		}
 
 		if (ConfigString.Len() > 0)
@@ -520,8 +526,7 @@ void FGPUProfilerEventNodeFrame::DumpEventTree()
 			GatherStatsEventNode(EventTree[BaseNodeIndex], 0, EventHistogram);
 		}
 
-		static IConsoleVariable* CVar2 = IConsoleManager::Get().FindConsoleVariable(TEXT("r.ProfileGPU.Root"));
-		FString RootWildcardString = CVar2->GetString(); 
+		FString RootWildcardString = GCVarProfileGPU_Root.GetValueOnRenderThread();
 		FWildcardString RootWildcard(RootWildcardString);
 
 		FGPUProfileStatSummary Summary;
@@ -551,11 +556,9 @@ void FGPUProfilerEventNodeFrame::DumpEventTree()
 			// Log stats about the node histogram
 			UE_LOG(LogRHI, Log, TEXT("Node histogram %u buckets"), EventHistogram.Num());
 
-			static IConsoleVariable* CVar = IConsoleManager::Get().FindConsoleVariable(TEXT("r.ProfileGPU.Pattern"));
-
 			// bad: reading on render thread but we don't support ECVF_RenderThreadSafe on strings yet
 			// It's very unlikely to cause a problem as the cvar is only changes by the user.
-			FString WildcardString = CVar->GetString(); 
+			FString WildcardString = GProfileGPUPatternCVar.GetValueOnRenderThread();
 
 			FGPUProfilerEventNodeStats Sum;
 
@@ -1379,15 +1382,60 @@ namespace UE::RHI::GPUProfiler
 		struct FNode
 		{
 			FString Name;
+
 			FNode* Parent = nullptr;
 			FNode* Next = nullptr;
-			uint32 Level = 0;
-			TArray<uint64> Timestamps;
 
-			uint32 NumDraws = 0;
-			uint32 NumPrimitives = 0;
+			TArray<FNode*> Children;
 
-			uint64 BusyCycles = 0;
+			uint64 LastCycles = 0;
+
+			struct FStats
+			{
+				uint32 NumDraws      = 0;
+				uint32 NumDispatches = 0;
+				uint32 NumPrimitives = 0;
+				uint32 NumVertices   = 0;
+				uint64 BusyCycles    = 0;
+
+				double GetBusyMilliseconds() const
+				{
+					return FPlatformTime::ToMilliseconds64(BusyCycles);
+				}
+
+				bool HasWork() const
+				{
+					return NumDraws > 0 || NumDispatches > 0;
+				}
+
+				FStats& operator += (FStats const& Stats)
+				{
+					NumDraws      += Stats.NumDraws;
+					NumDispatches += Stats.NumDispatches;
+					NumPrimitives += Stats.NumPrimitives;
+					NumVertices   += Stats.NumVertices;
+					BusyCycles    += Stats.BusyCycles;
+					return *this;
+				}
+
+				FStats& operator += (FEvent::FStats const& Stats)
+				{
+					NumDraws      += Stats.NumDraws;
+					NumPrimitives += Stats.NumPrimitives;
+
+					// @todo add NumDispatches / NumVertices to FEvent::FStats
+					//NumDispatches += Stats.NumDispatches;
+					//NumVertices   += Stats.NumVertices;
+
+					return *this;
+				}
+			};
+
+			// Exclusive stats for this node
+			FStats Exclusive;
+
+			// Sum of stats including all children
+			FStats Inclusive;
 
 			FNode(FString&& Name)
 				: Name(MoveTemp(Name))
@@ -1397,52 +1445,62 @@ namespace UE::RHI::GPUProfiler
 		struct FQueueState
 		{
 			FQueue const Queue;
-			TArray<TUniquePtr<FNode>> Nodes;
-			FNode* Current = nullptr;
-			FNode* Prev = nullptr;
-			FNode* First = nullptr;
 
-			enum class EState
+			enum class EMode
 			{
 				Idle,
 				WaitingFrame,
 				Active
-			} State = EState::Idle;
+			};
 
-			std::atomic<bool> bTriggerProfile { false };
+			struct
+			{
+				TArray<TUniquePtr<FNode>> Nodes;
+				FNode* Current = nullptr;
+				FNode* Prev    = nullptr;
+				FNode* First   = nullptr;
 
-			FQueueState(FQueue Queue)
-				: Queue(Queue)
-			{}
+				uint64 LastBeginCycles = 0;
+				uint64 LastEndCycles   = 0;
 
-			void PushNode(FString&& Name)
-			{
-				FNode* Parent = Current;
-				Current = Nodes.Emplace_GetRef(MakeUnique<FNode>(MoveTemp(Name))).Get();
-				Current->Parent = Parent;
+				EMode Mode = EMode::Idle;
 
-				if (!First)
+				void PushNode(FString&& Name)
 				{
-					First = Current;
-				}
+					FNode* Parent = Current;
+					Current = Nodes.Emplace_GetRef(MakeUnique<FNode>(MoveTemp(Name))).Get();
+					Current->Parent = Parent;
 
-				if (Parent)
-				{
-					Current->Level = Parent->Level + 1;
+					if (!First)
+					{
+						First = Current;
+					}
+
+					if (Parent)
+					{
+						Parent->Children.Add(Current);
+					}
+
+					if (Prev)
+					{
+						Prev->Next = Current;
+					}
+					Prev = Current;
 				}
 
-				if (Prev)
+				void PopNode()
 				{
-					Prev->Next = Current;
+					check(Current && Current->Parent);
+					Current = Current->Parent;
 				}
-				Prev = Current;
-			}
 
-			void PopNode()
-			{
-				check(Current && Current->Parent);
-				Current = Current->Parent;
-			}
+			} State {};
+
+			std::atomic<bool> bTriggerProfile { false };
+
+			FQueueState(FQueue Queue)
+				: Queue(Queue)
+			{}
 
 			void ProcessEvents(FEventStream const& EventStream);
 
@@ -1481,12 +1539,12 @@ namespace UE::RHI::GPUProfiler
 		auto End = std::end(EventStream);
 
 	Restart:
-		if (State == EState::Idle && bTriggerProfile.exchange(false))
+		if (State.Mode == EMode::Idle && bTriggerProfile.exchange(false))
 		{
-			State = EState::WaitingFrame;
+			State.Mode = EMode::WaitingFrame;
 		}
 
-		if (State == EState::WaitingFrame)
+		if (State.Mode == EMode::WaitingFrame)
 		{
 			// Discard all received events until we reach a FrameBoundary event
 			for (; Iterator != End; ++Iterator)
@@ -1494,12 +1552,12 @@ namespace UE::RHI::GPUProfiler
 				if ((*Iterator)->GetType() == FEvent::EType::FrameBoundary)
 				{
 					// Start profiling until we receive another FrameBoundary event
-					State = EState::Active;
+					State.Mode = EMode::Active;
 
 					FEvent::FFrameBoundary const& FrameBoundary = (*Iterator)->Value.Get<FEvent::FFrameBoundary>();
 
 					// Build the node tree 
-					PushNode(TEXT("<root>"));
+					State.PushNode(TEXT("<root>"));
 
 				#if WITH_RHI_BREADCRUMBS
 					auto Recurse = [&](auto& Recurse, FRHIBreadcrumbNode* Breadcrumb) -> void
@@ -1512,7 +1570,7 @@ namespace UE::RHI::GPUProfiler
 						Recurse(Recurse, Breadcrumb->GetParent());
 
 						FRHIBreadcrumb::FBuffer Buffer;
-						PushNode(Breadcrumb->Name.GetTCHAR(Buffer));
+						State.PushNode(Breadcrumb->Name.GetTCHAR(Buffer));
 					};
 					Recurse(Recurse, FrameBoundary.Breadcrumb);
 				#endif // WITH_RHI_BREADCRUMBS
@@ -1523,8 +1581,34 @@ namespace UE::RHI::GPUProfiler
 			}
 		}
 
-		if (State == EState::Active)
+		if (State.Mode == EMode::Active)
 		{
+			auto BeginWork = [&](uint64 Timestamp)
+			{
+				for (FNode* Node = State.Current; Node; Node = Node->Parent)
+				{
+					Node->LastCycles = FMath::Max(Node->LastCycles, Timestamp);
+				}
+			};
+
+			auto EndWork = [&](uint64 Timestamp)
+			{
+				for (FNode* Node = State.Current; Node; Node = Node->Parent)
+				{
+					uint64 Delta = Timestamp > Node->LastCycles
+						? Timestamp - Node->LastCycles
+						: 0;
+
+					Node->LastCycles = FMath::Max(Node->LastCycles, Timestamp);
+					Node->Inclusive.BusyCycles += Delta;
+
+					if (Node == State.Current)
+					{
+						Node->Exclusive.BusyCycles += Delta;
+					}
+				}
+			};
+
 			for (; Iterator != End; ++Iterator)
 			{
 				FEvent const* Event = *Iterator;
@@ -1532,23 +1616,11 @@ namespace UE::RHI::GPUProfiler
 				switch (Event->GetType())
 				{
 				case FEvent::EType::BeginWork:
-					{
-						uint64 Timestamp = Event->Value.Get<FEvent::FBeginWork>().GPUTimestampTOP;
-						for (FNode* Node = Current; Node; Node = Node->Parent)
-						{
-							Node->Timestamps.Add(Timestamp);
-						}
-					}
+					BeginWork(Event->Value.Get<FEvent::FBeginWork>().GPUTimestampTOP);
 					break;
 
 				case FEvent::EType::EndWork:
-					{
-						uint64 Timestamp = Event->Value.Get<FEvent::FEndWork>().GPUTimestampBOP;
-						for (FNode* Node = Current; Node; Node = Node->Parent)
-						{
-							Node->Timestamps.Add(Timestamp);
-						}
-					}
+					EndWork(Event->Value.Get<FEvent::FEndWork>().GPUTimestampBOP);
 					break;
 
 			#if WITH_RHI_BREADCRUMBS
@@ -1556,20 +1628,24 @@ namespace UE::RHI::GPUProfiler
 					{
 						FEvent::FBeginBreadcrumb const& BeginBreadcrumb = Event->Value.Get<FEvent::FBeginBreadcrumb>();
 
+						EndWork(BeginBreadcrumb.GPUTimestampTOP);
+
 						// Push a new node
 						FRHIBreadcrumb::FBuffer Buffer;
-						PushNode(BeginBreadcrumb.Breadcrumb->Name.GetTCHAR(Buffer));
+						State.PushNode(BeginBreadcrumb.Breadcrumb->Name.GetTCHAR(Buffer));
 
-						Current->Timestamps.Add(BeginBreadcrumb.GPUTimestampTOP);
+						BeginWork(BeginBreadcrumb.GPUTimestampTOP);
 					}
 					break;
 
 				case FEvent::EType::EndBreadcrumb:
 					{
 						FEvent::FEndBreadcrumb const& EndBreadcrumb = Event->Value.Get<FEvent::FEndBreadcrumb>();
-						Current->Timestamps.Add(EndBreadcrumb.GPUTimestampBOP);
+						EndWork(EndBreadcrumb.GPUTimestampBOP);
 
-						PopNode();
+						State.PopNode();
+
+						BeginWork(EndBreadcrumb.GPUTimestampBOP);
 					}
 					break;
 			#endif // WITH_RHI_BREADCRUMBS
@@ -1577,8 +1653,14 @@ namespace UE::RHI::GPUProfiler
 				case FEvent::EType::Stats:
 					{
 						FEvent::FStats const& Stats = Event->Value.Get<FEvent::FStats>();
-						Current->NumDraws += Stats.NumDraws;
-						Current->NumPrimitives += Stats.NumPrimitives;
+
+						check(State.Current);
+						State.Current->Exclusive += Stats;
+
+						for (FNode* Node = State.Current; Node; Node = Node->Parent)
+						{
+							Node->Inclusive += Stats;
+						}
 					}
 					break;
 
@@ -1588,12 +1670,7 @@ namespace UE::RHI::GPUProfiler
 						LogTree(FrameBoundary.FrameNumber);
 
 						// Reset tracking
-						Nodes.Reset();
-						Current = nullptr;
-						Prev = nullptr;
-						First = nullptr;
-
-						State = EState::Idle;
+						State = {};
 						goto Restart;
 					}
 					break;
@@ -1647,49 +1724,568 @@ namespace UE::RHI::GPUProfiler
 
 	void FGPUProfilerSink_ProfileGPU::FQueueState::LogTree(uint32 FrameNumber)
 	{
-		for (FNode* Node = First; Node; Node = Node->Next)
+		EGPUProfileSortMode SortMode = (EGPUProfileSortMode)FMath::Clamp(GProfileGPUSort.GetValueOnAnyThread(), 0, ((int32)EGPUProfileSortMode::EMax - 1));
+		FWildcardString RootWildcard(GCVarProfileGPU_Root.GetValueOnAnyThread());
+		const bool bShowEmptyNodes = GCVarProfileGPU_ShowLeafEvents.GetValueOnAnyThread();
+		const double PercentThreshold = FMath::Clamp(GCVarProfileGPU_ThresholdPercent.GetValueOnAnyThread(), 0.0f, 100.0f);
+
+		if (SortMode != EGPUProfileSortMode::EChronological)
 		{
-			check(Node->Timestamps.Num() % 2 == 0);
-			Node->BusyCycles = 0;
+			for (FNode* Node = State.First; Node; Node = Node->Next)
+			{
+				Node->Children.Sort([SortMode](FNode const& A, FNode const& B)
+				{
+					switch (SortMode)
+					{
+					default:
+					case EGPUProfileSortMode::ETimeElapsed: return B.Inclusive.BusyCycles    < A.Inclusive.BusyCycles;
+					case EGPUProfileSortMode::ENumPrims   : return B.Inclusive.NumPrimitives < A.Inclusive.NumPrimitives;
+					case EGPUProfileSortMode::ENumVerts   : return B.Inclusive.NumVertices   < A.Inclusive.NumVertices;
+					}
+				});
+			}
+		}
+
+		struct FTable
+		{
+			bool const bUnicodeOutput;
+
+			FTable()
+				: bUnicodeOutput(GCVarProfileGPU_UnicodeOutput.GetValueOnAnyThread())
+			{}
+
+			enum class EColumn : uint32
+			{
+				Exclusive_NumDraws,
+				Exclusive_NumDispatches,
+				Exclusive_NumPrimitives,
+				Exclusive_NumVertices,
+				Exclusive_Percent,
+				Exclusive_Bar,
+				Exclusive_Time,
+
+				Inclusive_NumDraws,
+				Inclusive_NumDispatches,
+				Inclusive_NumPrimitives,
+				Inclusive_NumVertices,
+				Inclusive_Percent,
+				Inclusive_Bar,
+				Inclusive_Time,
+
+				Events,
+
+				Num
+			};
+
+			uint32 GetColumnMinimumWidth(EColumn Column) const
+			{
+				switch (Column)
+				{
+				case EColumn::Exclusive_Bar:
+				case EColumn::Inclusive_Bar:
+					return 8;
+
+				case EColumn::Events:
+					return 6;
+				}
+
+				return 0;
+			}
+
+			TCHAR const* GetColumnHeader(EColumn Column) const
+			{
+				switch (Column)
+				{
+				case EColumn::Exclusive_NumDraws:
+				case EColumn::Inclusive_NumDraws:
+					return TEXT("Draws");
+
+				case EColumn::Exclusive_NumDispatches:
+				case EColumn::Inclusive_NumDispatches:
+					return TEXT("Dsptch");
+				
+				case EColumn::Exclusive_NumPrimitives:
+				case EColumn::Inclusive_NumPrimitives:
+					return TEXT("Prim");
+
+				case EColumn::Exclusive_NumVertices:
+				case EColumn::Inclusive_NumVertices:
+					return TEXT("Vert");
+
+				case EColumn::Exclusive_Percent:
+				case EColumn::Inclusive_Percent:
+					return TEXT("Percent");
+				
+				case EColumn::Exclusive_Time:
+				case EColumn::Inclusive_Time:
+					return TEXT("Time");
+				}
+
+				return TEXT("");
+			}
+
+			uint32 GetColumnGroup(EColumn Column) const
+			{
+			switch (Column)
+				{
+				case EColumn::Exclusive_NumDraws:
+				case EColumn::Exclusive_NumDispatches:
+				case EColumn::Exclusive_NumPrimitives:
+				case EColumn::Exclusive_NumVertices:
+				case EColumn::Exclusive_Percent:
+				case EColumn::Exclusive_Bar:
+				case EColumn::Exclusive_Time:
+					return 0;
+
+				case EColumn::Inclusive_NumDraws:
+				case EColumn::Inclusive_NumDispatches:
+				case EColumn::Inclusive_NumPrimitives:
+				case EColumn::Inclusive_NumVertices:
+				case EColumn::Inclusive_Percent:
+				case EColumn::Inclusive_Bar:
+				case EColumn::Inclusive_Time:
+					return 1;
+
+				default:
+				case EColumn::Events:
+					return 2;
+				}
+			}
+
+			TCHAR const* GetGroupName(uint32 GroupIndex) const
+			{
+				switch (GroupIndex)
+				{
+				case 0: return TEXT("Exclusive");
+				case 1: return TEXT("Inclusive");
+				case 2: return TEXT("Events");
+				}
+				return TEXT("");
+			}
+
+			uint32 NumRows = 0;
+			TStaticArray<TArray<FString>, uint32(EColumn::Num)> Columns { InPlace };
+			TArray<bool> RowBreaks;
+
+			FString& Col(EColumn Column)
+			{
+				return Columns[uint32(Column)].Emplace_GetRef();
+			}
+
+			bool HasRows() const
+			{
+				return NumRows > 0;
+			}
+
+			void AddRow(FNode* Root, FNode::FStats const& Inclusive, FNode::FStats const& Exclusive, FString const& Name, uint32 Level)
+			{
+				double ExclusivePercent = double(Exclusive.BusyCycles) / Root->Inclusive.BusyCycles;
+				double InclusivePercent = double(Inclusive.BusyCycles) / Root->Inclusive.BusyCycles;
+
+				Col(EColumn::Exclusive_NumDraws     ) = FString::Printf(TEXT("%d"), Exclusive.NumDraws);
+				Col(EColumn::Exclusive_NumDispatches) = FString::Printf(TEXT("%d"), Exclusive.NumDispatches);
+				Col(EColumn::Exclusive_NumPrimitives) = FString::Printf(TEXT("%d"), Exclusive.NumPrimitives);
+				Col(EColumn::Exclusive_NumVertices  ) = FString::Printf(TEXT("%d"), Exclusive.NumVertices);
+				Col(EColumn::Exclusive_Percent      ) = FString::Printf(TEXT("%.1f%%"), ExclusivePercent * 100.0);
+				Col(EColumn::Exclusive_Time         ) = FString::Printf(TEXT("%.3f ms"), FPlatformTime::ToMilliseconds64(Exclusive.BusyCycles));
+
+				Col(EColumn::Inclusive_NumDraws     ) = FString::Printf(TEXT("%d"), Inclusive.NumDraws);
+				Col(EColumn::Inclusive_NumDispatches) = FString::Printf(TEXT("%d"), Inclusive.NumDispatches);
+				Col(EColumn::Inclusive_NumPrimitives) = FString::Printf(TEXT("%d"), Inclusive.NumPrimitives);
+				Col(EColumn::Inclusive_NumVertices  ) = FString::Printf(TEXT("%d"), Inclusive.NumVertices);
+				Col(EColumn::Inclusive_Percent      ) = FString::Printf(TEXT("%.1f%%"), InclusivePercent * 100.0);
+				Col(EColumn::Inclusive_Time         ) = FString::Printf(TEXT("%.3f ms"), FPlatformTime::ToMilliseconds64(Inclusive.BusyCycles));
+				
+				if (bUnicodeOutput)
+				{
+					TUnicodeHorizontalBar<8> ExclusiveBar = ExclusivePercent;
+					TUnicodeHorizontalBar<8> InclusiveBar = InclusivePercent;
+					Col(EColumn::Exclusive_Bar) = ExclusiveBar.Text;
+					Col(EColumn::Inclusive_Bar) = InclusiveBar.Text;
+				}
+
+				static constexpr uint32 SpacesPerIndent = 3;
+				Col(EColumn::Events) = FString::Printf(TEXT("%*s"), Name.Len() + (Level * SpacesPerIndent), *Name);
+
+				// Insert a horizontal rule before each root level row.
+				RowBreaks.Add(Level == 0);
 
-			uint64 LastBeginCycles = 0;
-			uint64 LastEndCycles = 0;
+				NumRows++;
+			}
 
-			for (int32 Index = 0; Index < Node->Timestamps.Num(); ++Index)
+			struct FChars
 			{
-				if ((Index & 1) == 0)
+				TCHAR const* Left;
+				TCHAR const* GroupSeparator;
+				TCHAR const* LastGroupSeparator;
+				TCHAR const* Right;
+				TCHAR const* CellSeparator;
+			};
+
+			struct FFormat
+			{
+				TCHAR const* LineMajor;
+				TCHAR const* LineMinor;
+
+				FChars const TopRow;
+				FChars const GroupNameRow;
+				FChars const GroupBorderRow;
+				FChars const ValueRow;
+				FChars const DividorRow;
+				FChars const BottomRow;
+			};
+
+			FString ToString() const
+			{
+				if (bUnicodeOutput)
 				{
-					// Begin
-					LastBeginCycles = FMath::Max(LastEndCycles, Node->Timestamps[Index]);
+					static constexpr FFormat Unicode = 
+					{
+						.LineMajor  = TEXT("━"),
+						.LineMinor  = TEXT("─"),
+
+						//                     Left     GrpSep     LastGrp     Right     CellSep
+						.TopRow        { TEXT("    ┏"), TEXT("┳"), TEXT("┳"), TEXT("┓"), TEXT(" ") },
+						.GroupNameRow  { TEXT("    ┃"), TEXT("┃"), TEXT("┃"), TEXT("┃"), TEXT(" ") },
+						.GroupBorderRow{ TEXT("    ┠"), TEXT("╂"), TEXT("┨"), TEXT("┃"), TEXT("┬") },
+						.ValueRow      { TEXT("    ┃"), TEXT("┃"), TEXT("┃"), TEXT("┃"), TEXT("│") },
+						.DividorRow    { TEXT("    ┠"), TEXT("╂"), TEXT("╂"), TEXT("┨"), TEXT("┼") },
+						.BottomRow     { TEXT("    ┗"), TEXT("┻"), TEXT("┻"), TEXT("┛"), TEXT("┷") },
+					};
+
+					return ToStringInner(Unicode);
 				}
 				else
 				{
-					// End
-					uint64 End = Node->Timestamps[Index];
-					Node->BusyCycles += End - LastBeginCycles;
+					static constexpr FFormat Ascii = 
+					{
+						.LineMajor  = TEXT("-"),
+						.LineMinor  = TEXT("-"),
+
+						//                     Left     GrpSep     LastGrp     Right     CellSep
+						.TopRow        { TEXT("    +"), TEXT("+"), TEXT("+"), TEXT("+"), TEXT(" ") },
+						.GroupNameRow  { TEXT("    |"), TEXT("|"), TEXT("|"), TEXT("|"), TEXT(" ") },
+						.GroupBorderRow{ TEXT("    +"), TEXT("+"), TEXT("+"), TEXT("|"), TEXT("+") },
+						.ValueRow      { TEXT("    |"), TEXT("|"), TEXT("|"), TEXT("|"), TEXT("|") },
+						.DividorRow    { TEXT("    +"), TEXT("+"), TEXT("+"), TEXT("+"), TEXT("+") },
+						.BottomRow     { TEXT("    +"), TEXT("+"), TEXT("+"), TEXT("+"), TEXT("+") },
+					};
 
-					LastEndCycles = End;
+					return ToStringInner(Ascii);
 				}
 			}
-		}
 
-		FString LogMessage;
-		for (FNode* Node = First; Node; Node = Node->Next)
+			FString ToStringInner(FFormat const& Format) const
+			{
+				struct FGroup  { uint32 Index, Width; };
+				struct FColumn { uint32 Index, Width; };
+
+				static constexpr uint32 NumGroups = 3;
+				static constexpr uint32 CellPadding = 1;
+
+				// Auto-size column widths to their contents
+				TStaticArray<int32, uint32(EColumn::Num)> ColumnWidths{ InPlace, 0 };
+				for (uint32 ColumnIndex = 0; ColumnIndex < uint32(EColumn::Num); ++ColumnIndex)
+				{
+					if (Columns[ColumnIndex].Num() == 0)
+						continue;
+
+					check(Columns[ColumnIndex].Num() == NumRows);
+
+					int32& Width = ColumnWidths[ColumnIndex];
+
+					// Auto-size column width
+					Width = GetColumnMinimumWidth(EColumn(ColumnIndex));
+					Width = FMath::Max(Width, FCString::Strlen(GetColumnHeader(EColumn(ColumnIndex))));
+
+					for (FString const& Cell : Columns[ColumnIndex])
+					{
+						Width = FMath::Max(Width, Cell.Len());
+					}
+				}
+
+				FString Result;
+
+				auto EmitGroupRow = [&](FChars const& Chars, TUniqueFunction<void(FGroup)> GroupCallback)
+				{
+					uint32 const CellSeparatorLength = FCString::Strlen(Chars.CellSeparator);
+
+					Result += Chars.Left;
+
+					uint32 GroupWidth = 0;
+					uint32 GroupIndex = 0;
+
+					for (uint32 ColumnIndex = 0; ColumnIndex < uint32(EColumn::Num); ++ColumnIndex)
+					{
+						if (Columns[ColumnIndex].Num() == 0)
+							continue;
+
+						GroupWidth += ColumnWidths[ColumnIndex] + CellPadding * 2;
+						GroupIndex = GetColumnGroup(EColumn(ColumnIndex));
+
+						if (GroupIndex != GetColumnGroup(EColumn(ColumnIndex + 1)))
+						{ 
+							// Group Change
+							GroupCallback({ GroupIndex, GroupWidth });
+
+							// Add the group separator character
+							Result += GroupIndex < NumGroups - 2
+								? Chars.GroupSeparator
+								: Chars.LastGroupSeparator;
+
+							GroupWidth = 0;
+						}
+						else if (ColumnIndex < uint32(EColumn::Num) - 1)
+						{
+							// Same group. Count the (missing) cell division
+							GroupWidth += CellSeparatorLength;
+						}
+					}
+
+					// Emit final group
+					GroupCallback({ GroupIndex, GroupWidth });
+
+					// Close the row
+					Result += Chars.Right;
+					Result += TEXT("\n");
+				};
+
+				auto EmitValueRow = [&](FChars const& Chars, TUniqueFunction<void(FColumn)> CellCallback)
+				{
+					Result += Chars.Left;
+
+					for (uint32 ColumnIndex = 0; ColumnIndex < uint32(EColumn::Num); ++ColumnIndex)
+					{
+						if (Columns[ColumnIndex].Num() == 0)
+							continue;
+
+						CellCallback({ ColumnIndex, ColumnWidths[ColumnIndex] + (CellPadding * 2) });
+
+						if (ColumnIndex < uint32(EColumn::Num) - 1)
+						{
+							uint32 GroupIndex = GetColumnGroup(EColumn(ColumnIndex));
+							if (GroupIndex != GetColumnGroup(EColumn(ColumnIndex + 1)))
+							{
+								// Group change, add the group separator
+								Result += GroupIndex < NumGroups - 2
+									? Chars.GroupSeparator
+									: Chars.LastGroupSeparator;
+							}
+							else
+							{
+								// Same group, add the cell separator
+								Result += Chars.CellSeparator;
+							}
+						}
+					}
+
+					// Close the row
+					Result += Chars.Right;
+					Result += TEXT("\n");
+				};
+
+				auto AlignCenter = [&](TCHAR const* Str, uint32 Width)
+				{
+					int32 PaddingLeft = FMath::Max(0, int32(Width) - FCString::Strlen(Str));
+					int32 PaddingRight = (PaddingLeft / 2) + (PaddingLeft & 1);
+					PaddingLeft /= 2;
+
+					Result += FString::Printf(TEXT("%*s%s%*s"), PaddingLeft, TEXT(""), Str, PaddingRight, TEXT(""));
+				};
+
+				// Top Border
+				EmitGroupRow(Format.TopRow, [&](FGroup Group)
+				{
+					while (Group.Width--)
+					{
+						Result += Format.LineMajor;
+					}
+				});
+
+				// Exclusive / Inclusive Group Row
+				EmitGroupRow(Format.GroupNameRow, [&](FGroup Group)
+				{
+					TCHAR const* Str = Group.Index != GetColumnGroup(EColumn::Events)
+						? GetGroupName(Group.Index)
+						: TEXT("");
+
+					AlignCenter(Str, Group.Width);
+				});
+
+				// Events Group Row
+				EmitValueRow(Format.GroupBorderRow, [&](FColumn Column)
+				{
+					if (Column.Index == uint32(EColumn::Events))
+					{
+						AlignCenter(GetGroupName(GetColumnGroup(EColumn::Events)), Column.Width);
+					}
+					else
+					{
+						while (Column.Width--)
+						{
+							Result += Format.LineMinor;
+						}
+					}
+				});
+
+				// Header Row
+				EmitValueRow(Format.ValueRow, [&](FColumn Column)
+				{
+					AlignCenter(GetColumnHeader(EColumn(Column.Index)), Column.Width);
+				});
+
+				// Header Border Row
+				EmitValueRow(Format.DividorRow, [&](FColumn Column)
+				{
+					while (Column.Width--)
+					{
+						Result += Format.LineMinor;
+					}
+				});
+
+				// Value rows
+				for (uint32 RowIndex = 0; RowIndex < NumRows; ++RowIndex)
+				{
+					if (RowIndex > 0 && RowBreaks[RowIndex])
+					{
+						// Add a horizontal rule
+						EmitValueRow(Format.DividorRow, [&](FColumn Column)
+						{
+							while (Column.Width--)
+							{
+								Result += Format.LineMinor;
+							}
+						});
+					}
+
+					EmitValueRow(Format.ValueRow, [&](FColumn Column)
+					{
+						int32 Width = Column.Width - (CellPadding * 2);
+						if (EColumn(Column.Index) == EColumn::Events)
+						{
+							Width = -Width; // Align left
+						}
+
+						FString const& Cell = Columns[Column.Index][RowIndex];
+						Result += FString::Printf(TEXT("%*s%*s%*s")
+							, CellPadding, TEXT("")
+							, Width, *Cell
+							, CellPadding, TEXT(""));
+					});
+				}
+
+				// Bottom Border
+				EmitValueRow(Format.BottomRow, [&](FColumn Column)
+				{
+					while (Column.Width--)
+					{
+						Result += Format.LineMajor;
+					}
+				});
+
+				return Result;
+			}
+		} Table;
+
+		auto Recurse = [&](auto& Recurse, FNode* Root, FNode* Current, bool bParentMatchedFilter, int32 Level) -> bool
+		{
+			// Percent that this node was of the total frame time
+			const double Percent = Root ? (Current->Inclusive.GetBusyMilliseconds() / Root->Inclusive.GetBusyMilliseconds()) * 100.0 : 100.0;
+
+			// Filter nodes according to cvar settings
+			const bool bAboveThreshold = Percent >= PercentThreshold;
+			const bool bNameMatches = bParentMatchedFilter || RootWildcard.IsMatch(Current->Name);
+			const bool bHasWork = bShowEmptyNodes || Current->Inclusive.HasWork();
+
+			const bool bDisplayEvent = bNameMatches && bHasWork && bAboveThreshold;
+			
+			if (bDisplayEvent)
+			{
+				if (Root == nullptr)
+				{
+					Root = Current;
+				}
+
+				Table.AddRow(
+					Root,
+					Current->Inclusive,
+					Current->Exclusive,
+					Current->Name,
+					Level
+				);
+			}
+
+			FNode::FStats OtherChildrenInclusive;
+			FNode::FStats OtherChildrenExclusive;
+			uint32 NumHiddenChildren = 0;
+
+			for (FNode* Child : Current->Children)
+			{
+				bool bChildShown = Recurse(Recurse, Root, Child, bDisplayEvent, bDisplayEvent ? Level + 1 : Level);
+				if (!bChildShown)
+				{
+					OtherChildrenInclusive += Child->Inclusive;
+					OtherChildrenExclusive += Child->Exclusive;
+
+					NumHiddenChildren++;
+				}
+			}
+
+			if (bDisplayEvent && NumHiddenChildren > 0)
+			{
+				// Don't show the "other children" node if their total inclusive time is still below the percent threshold
+				if ((double(OtherChildrenInclusive.BusyCycles) / Root->Inclusive.BusyCycles) >= PercentThreshold)
+				{
+					Table.AddRow(
+						Root,
+						OtherChildrenInclusive,
+						OtherChildrenExclusive,
+						FString::Printf(TEXT("%d Other %s"), NumHiddenChildren, NumHiddenChildren >= 2 ? TEXT("Children") : TEXT("Child")),
+						Level + 1
+					);
+				}
+			}
+			
+			return bDisplayEvent;
+		};
+
+		Recurse(Recurse, nullptr, State.First, false, 0);
+
+		if (Table.HasRows())
 		{
-			double Milliseconds = FPlatformTime::ToMilliseconds64(Node->BusyCycles);
-
-			TUnicodeHorizontalBar<8> Bar = Milliseconds / FPlatformTime::ToMilliseconds64(First->BusyCycles);
-			LogMessage += FString::Printf(TEXT("%9.3f ms |%s| %6d | %6d | %*s\n")
-				, Milliseconds
-				, Bar.Text
-				, Node->NumDraws
-				, Node->NumPrimitives
-				, Node->Name.Len() + (Node->Level * 4)
-				, *Node->Name
+			FString Final = FString::Printf(
+				TEXT("\n")
+				TEXT("GPU Profile for Frame %d - GPU %d - %s %d\n")
+				TEXT("\n")
+				TEXT("    - %-30s: %.2fms\n")
+				TEXT("    - %-30s: \"%s\"\n")
+				TEXT("    - %-30s: %.2f%%\n")
+				TEXT("    - %-30s: %s\n")
+				TEXT("\n")
+				TEXT("%s")
+				, FrameNumber
+				, Queue.GPU
+				, Queue.GetTypeString()
+				, Queue.Index
+				, TEXT("Frame Time")
+				, State.First->Inclusive.GetBusyMilliseconds()
+				, *IConsoleManager::Get().FindConsoleObjectName(GCVarProfileGPU_Root.AsVariable())
+				, *RootWildcard
+				, *IConsoleManager::Get().FindConsoleObjectName(GCVarProfileGPU_ThresholdPercent.AsVariable())
+				, PercentThreshold
+				, *IConsoleManager::Get().FindConsoleObjectName(GCVarProfileGPU_ShowLeafEvents.AsVariable())
+				, bShowEmptyNodes ? TEXT("true") : TEXT("false")
+				, *Table.ToString()
 			);
-		}
 
-		UE_LOG(LogRHI, Display, TEXT("GPU Profile for Frame %d, Queue [%s, GPU: %d, Idx: %d]:\n%s\n\n"), FrameNumber, Queue.GetTypeString(), Queue.GPU, Queue.Index, *LogMessage);
+			TArray<FString> Lines;
+			Final.ParseIntoArrayLines(Lines, false);
+
+			for (FString const& Line : Lines)
+			{
+				UE_LOG(LogRHI, Display, TEXT("%s"), *Line);
+			}
+		}
 	}
 
 	static FAutoConsoleCommand ProfileGPUNew(
