Merge percent and bar columns in profilegpu output

diff --git a/Engine/Source/Runtime/RHI/Private/GPUProfiler.cpp b/Engine/Source/Runtime/RHI/Private/GPUProfiler.cpp
--- a/Engine/Source/Runtime/RHI/Private/GPUProfiler.cpp
+++ b/Engine/Source/Runtime/RHI/Private/GPUProfiler.cpp
@@ -1761,7 +1761,6 @@ namespace UE::RHI::GPUProfiler
 				Exclusive_NumPrimitives,
 				Exclusive_NumVertices,
 				Exclusive_Percent,
-				Exclusive_Bar,
 				Exclusive_Time,
 
 				Inclusive_NumDraws,
@@ -1769,7 +1768,6 @@ namespace UE::RHI::GPUProfiler
 				Inclusive_NumPrimitives,
 				Inclusive_NumVertices,
 				Inclusive_Percent,
-				Inclusive_Bar,
 				Inclusive_Time,
 
 				Events,
@@ -1781,10 +1779,6 @@ namespace UE::RHI::GPUProfiler
 			{
 				switch (Column)
 				{
-				case EColumn::Exclusive_Bar:
-				case EColumn::Inclusive_Bar:
-					return 8;
-
 				case EColumn::Events:
 					return 6;
 				}
@@ -1833,7 +1827,6 @@ namespace UE::RHI::GPUProfiler
 				case EColumn::Exclusive_NumPrimitives:
 				case EColumn::Exclusive_NumVertices:
 				case EColumn::Exclusive_Percent:
-				case EColumn::Exclusive_Bar:
 				case EColumn::Exclusive_Time:
 					return 0;
 
@@ -1842,7 +1835,6 @@ namespace UE::RHI::GPUProfiler
 				case EColumn::Inclusive_NumPrimitives:
 				case EColumn::Inclusive_NumVertices:
 				case EColumn::Inclusive_Percent:
-				case EColumn::Inclusive_Bar:
 				case EColumn::Inclusive_Time:
 					return 1;
 
@@ -1882,28 +1874,26 @@ namespace UE::RHI::GPUProfiler
 				double ExclusivePercent = double(Exclusive.BusyCycles) / Root->Inclusive.BusyCycles;
 				double InclusivePercent = double(Inclusive.BusyCycles) / Root->Inclusive.BusyCycles;
 
+				static constexpr uint32 BarWidth = 8;
+				TUnicodeHorizontalBar<BarWidth> ExclusiveBar = ExclusivePercent;
+				TUnicodeHorizontalBar<BarWidth> InclusiveBar = InclusivePercent;
+
+				static constexpr TCHAR const BarSeparator[] = TEXT(" ┊ ");
+
 				Col(EColumn::Exclusive_NumDraws     ) = FString::Printf(TEXT("%d"), Exclusive.NumDraws);
 				Col(EColumn::Exclusive_NumDispatches) = FString::Printf(TEXT("%d"), Exclusive.NumDispatches);
 				Col(EColumn::Exclusive_NumPrimitives) = FString::Printf(TEXT("%d"), Exclusive.NumPrimitives);
 				Col(EColumn::Exclusive_NumVertices  ) = FString::Printf(TEXT("%d"), Exclusive.NumVertices);
-				Col(EColumn::Exclusive_Percent      ) = FString::Printf(TEXT("%.1f%%"), ExclusivePercent * 100.0);
+				Col(EColumn::Exclusive_Percent      ) = FString::Printf(TEXT("%.1f%%%s%s"), ExclusivePercent * 100.0, bUnicodeOutput ? BarSeparator : TEXT(""), bUnicodeOutput ? ExclusiveBar.Text : TEXT(""));
 				Col(EColumn::Exclusive_Time         ) = FString::Printf(TEXT("%.3f ms"), FPlatformTime::ToMilliseconds64(Exclusive.BusyCycles));
 
 				Col(EColumn::Inclusive_NumDraws     ) = FString::Printf(TEXT("%d"), Inclusive.NumDraws);
 				Col(EColumn::Inclusive_NumDispatches) = FString::Printf(TEXT("%d"), Inclusive.NumDispatches);
 				Col(EColumn::Inclusive_NumPrimitives) = FString::Printf(TEXT("%d"), Inclusive.NumPrimitives);
 				Col(EColumn::Inclusive_NumVertices  ) = FString::Printf(TEXT("%d"), Inclusive.NumVertices);
-				Col(EColumn::Inclusive_Percent      ) = FString::Printf(TEXT("%.1f%%"), InclusivePercent * 100.0);
+				Col(EColumn::Inclusive_Percent      ) = FString::Printf(TEXT("%.1f%%%s%s"), InclusivePercent * 100.0, bUnicodeOutput ? BarSeparator : TEXT(""), bUnicodeOutput ? InclusiveBar.Text : TEXT(""));
 				Col(EColumn::Inclusive_Time         ) = FString::Printf(TEXT("%.3f ms"), FPlatformTime::ToMilliseconds64(Inclusive.BusyCycles));
 				
-				if (bUnicodeOutput)
-				{
-					TUnicodeHorizontalBar<8> ExclusiveBar = ExclusivePercent;
-					TUnicodeHorizontalBar<8> InclusiveBar = InclusivePercent;
-					Col(EColumn::Exclusive_Bar) = ExclusiveBar.Text;
-					Col(EColumn::Inclusive_Bar) = InclusiveBar.Text;
-				}
-
 				static constexpr uint32 SpacesPerIndent = 3;
 				Col(EColumn::Events) = FString::Printf(TEXT("%*s"), Name.Len() + (Level * SpacesPerIndent), *Name);
 
