Remove old references to r.RHICmdBypass cvar in config

diff --git a/Engine/Config/BaseEngine.ini b/Engine/Config/BaseEngine.ini
--- a/Engine/Config/BaseEngine.ini
+++ b/Engine/Config/BaseEngine.ini
@@ -2356,9 +2356,6 @@ con.DebugLateCheat = True
 ; Do not vsync in the editor. Throws off gpu profiling.
 r.VSync=0
 
-; Parallel rendering has not been tested in the editor so just disable it.
-r.RHICmdBypass=0
-
 [SystemSettingsSplitScreen2]
 ; System settings overrides for split screen
 ; Use medium detail mode in splitscreen, this allows LD's to mark meshes as High detail and they won't render in SS
diff --git a/Engine/Config/Linux/LinuxEngine.ini b/Engine/Config/Linux/LinuxEngine.ini
--- a/Engine/Config/Linux/LinuxEngine.ini
+++ b/Engine/Config/Linux/LinuxEngine.ini
@@ -11,9 +11,6 @@ fx.NiagaraAllowRuntimeScalabilityChanges=1
 QualityLevelMapping="high"
 r.Shadow.DetectVertexShaderLayerAtRuntime=1
 
-[SystemSettingsEditor]
-;r.RHICmdBypass=1
-
 [TextureStreaming]
 ; PoolSizeVRAMPercentage is how much percentage of GPU Dedicated VRAM should be used as a TexturePool cache for streaming textures (0 - unlimited streaming)
 PoolSizeVRAMPercentage=70
diff --git a/Engine/Config/LinuxArm64/LinuxArm64Engine.ini b/Engine/Config/LinuxArm64/LinuxArm64Engine.ini
--- a/Engine/Config/LinuxArm64/LinuxArm64Engine.ini
+++ b/Engine/Config/LinuxArm64/LinuxArm64Engine.ini
@@ -12,9 +12,6 @@ QualityLevelMapping="low"
 fx.Niagara.QualityLevel.Min=0
 fx.Niagara.QualityLevel.Max=0
 
-[SystemSettingsEditor]
-;r.RHICmdBypass=1
-
 [TextureStreaming]
 ; PoolSizeVRAMPercentage is how much percentage of GPU Dedicated VRAM should be used as a TexturePool cache for streaming textures (0 - unlimited streaming)
 PoolSizeVRAMPercentage=70
diff --git a/Engine/Source/Runtime/RHI/Private/RHICommandList.cpp b/Engine/Source/Runtime/RHI/Private/RHICommandList.cpp
--- a/Engine/Source/Runtime/RHI/Private/RHICommandList.cpp
+++ b/Engine/Source/Runtime/RHI/Private/RHICommandList.cpp
@@ -1466,13 +1466,11 @@ void FRHICommandListExecutor::LatchBypass()
 	{
 		if (FParse::Param(FCommandLine::Get(), TEXT("forcerhibypass")) && CVarRHICmdBypass.GetValueOnRenderThread() == 0)
 		{
-			IConsoleVariable* BypassVar = IConsoleManager::Get().FindConsoleVariable(TEXT("r.RHICmdBypass"));
-			BypassVar->Set(1, ECVF_SetByCommandline);
+			CVarRHICmdBypass->Set(1, ECVF_SetByCommandline);
 		}
 		else if (FParse::Param(FCommandLine::Get(), TEXT("parallelrendering")) && CVarRHICmdBypass.GetValueOnRenderThread() >= 1)
 		{
-			IConsoleVariable* BypassVar = IConsoleManager::Get().FindConsoleVariable(TEXT("r.RHICmdBypass"));
-			BypassVar->Set(0, ECVF_SetByCommandline);
+			CVarRHICmdBypass->Set(0, ECVF_SetByCommandline);
 		}
 	});
 
