Removing RHI resource log spam in debug builds.

diff --git a/Engine/Source/Runtime/RHI/Public/RHICommandList.h b/Engine/Source/Runtime/RHI/Public/RHICommandList.h
--- a/Engine/Source/Runtime/RHI/Public/RHICommandList.h
+++ b/Engine/Source/Runtime/RHI/Public/RHICommandList.h
@@ -871,16 +871,9 @@ public:
 			FRHITextureCreateDesc NewCreateDesc(CreateDesc);
 			NewCreateDesc.SetInitialState(RHIGetDefaultResourceState(CreateDesc.Flags, CreateDesc.BulkData != nullptr));
 
-#if RHI_TRANSITIONS_HL_DEBUG
-				UE_LOG(LogRHI, Display, TEXT("Resource: %s Creation State %s"), CreateDesc.DebugName, *GetRHIAccessName(NewCreateDesc.InitialState));
-#endif
-
 			return GDynamicRHI->RHICreateTexture(*this, NewCreateDesc);
 		}
 
-#if RHI_TRANSITIONS_HL_DEBUG
-		UE_LOG(LogRHI, Display, TEXT("Resource: %s Creation State %s"), CreateDesc.DebugName, *GetRHIAccessName(CreateDesc.InitialState));
-#endif
 		return GDynamicRHI->RHICreateTexture(*this, CreateDesc);
 	}
 
