Default laptops on performance mode (D3D11 ES31).

#tests tested locally and by QA on a wide range of hardware
#rnx

diff --git a/Engine/Source/Runtime/RHI/Private/Windows/WindowsDynamicRHI.cpp b/Engine/Source/Runtime/RHI/Private/Windows/WindowsDynamicRHI.cpp
--- a/Engine/Source/Runtime/RHI/Private/Windows/WindowsDynamicRHI.cpp
+++ b/Engine/Source/Runtime/RHI/Private/Windows/WindowsDynamicRHI.cpp
@@ -367,21 +367,32 @@ FParsedWindowsDynamicRHIConfig ParseWindowsDynamicRHIConfig()
 	return Config;
 }
 
-// Default to Performance Mode on low-end machines
-static bool DefaultFeatureLevelES31()
+// Whether to default to Performance Mode on low-end machines (D3D11 RHI, ES31 feature level).
+// Looks at a few configuration options that control thresholds to enable this as the default.
+static bool ShouldDefaultToPerformanceMode()
 {
-	static TOptional<bool> ForceES31;
-	if (ForceES31.IsSet())
+	static TOptional<bool> ForceD3D11_ES31;
+	if (ForceD3D11_ES31.IsSet())
 	{
-		return ForceES31.GetValue();
+		return ForceD3D11_ES31.GetValue();
+	}
+
+	// Force Performance mode for laptops. We don't have an explicit way to check for laptops
+	// so we use whether we're running on battery as a proxy.
+	bool bIsLaptop = FPlatformMisc::IsRunningOnBattery();
+	if (GConfig->GetBoolOrDefault(TEXT("PerformanceMode"), TEXT("bDefaultOnLaptops"), false, GEngineIni) && bIsLaptop)
+	{
+		UE_LOG(LogRHI, Log, TEXT("Running on laptop: default feature level to %s"), GetLogName(ERHIFeatureLevel::ES3_1));
+		ForceD3D11_ES31 = true;
+		return true;
 	}
 
 	// Force Performance mode for machines with too few cores including hyperthreads
 	int MinCoreCount = 0;
 	if (GConfig->GetInt(TEXT("PerformanceMode"), TEXT("MinCoreCount"), MinCoreCount, GEngineIni) && FPlatformMisc::NumberOfCoresIncludingHyperthreads() < MinCoreCount)
 	{
-		UE_LOG(LogRHI, Log, TEXT("MinCoreCount found (%d) and user core count (%d) forced feature level to %s"), MinCoreCount, FPlatformMisc::NumberOfCoresIncludingHyperthreads(), GetLogName(ERHIFeatureLevel::ES3_1));
-		ForceES31 = true;
+		UE_LOG(LogRHI, Log, TEXT("MinCoreCount found (%d) and user core count (%d): default feature level to %s"), MinCoreCount, FPlatformMisc::NumberOfCoresIncludingHyperthreads(), GetLogName(ERHIFeatureLevel::ES3_1));
+		ForceD3D11_ES31 = true;
 		return true;
 	}
 
@@ -399,9 +410,9 @@ static bool DefaultFeatureLevelES31()
 			{
 				if (FPlatformMemory::GetMemorySizeBucket() >= EPlatformMemorySizeBucket(EnumIndex))
 				{
-					UE_LOG(LogRHI, Log, TEXT("MinMemorySizeBucket found (%s) and user memory (%d) forced feature level to %s"), *MinMemorySizeBucketString, int32(FPlatformMemory::GetMemorySizeBucket()), GetLogName(ERHIFeatureLevel::ES3_1));
+					UE_LOG(LogRHI, Log, TEXT("MinMemorySizeBucket found (%s) and user memory (%d): default feature level to %s"), *MinMemorySizeBucketString, int32(FPlatformMemory::GetMemorySizeBucket()), GetLogName(ERHIFeatureLevel::ES3_1));
 
-					ForceES31 = true;
+					ForceD3D11_ES31 = true;
 					return true;
 				}
 			}
@@ -412,9 +423,9 @@ static bool DefaultFeatureLevelES31()
 				const int MIN_GPU_MEMORY = 512 * 1024 * 1024;
 				if (FPlatformMemory::GetMemorySizeBucket() >= EPlatformMemorySizeBucket(EnumIndex) && BestGPUInfo.DedicatedVideoMemory < MIN_GPU_MEMORY)
 				{
-					UE_LOG(LogRHI, Log, TEXT("MinMemorySizeBucket found (%s) and user memory (%d) forced feature level to %s"), *MinMemorySizeBucketString, int32(FPlatformMemory::GetMemorySizeBucket()), GetLogName(ERHIFeatureLevel::ES3_1));
+					UE_LOG(LogRHI, Log, TEXT("MinMemorySizeBucket found (%s) and user memory (%d): default feature level to %s"), *MinMemorySizeBucketString, int32(FPlatformMemory::GetMemorySizeBucket()), GetLogName(ERHIFeatureLevel::ES3_1));
 
-					ForceES31 = true;
+					ForceD3D11_ES31 = true;
 					return true;
 				}
 			}
@@ -439,9 +450,9 @@ static bool DefaultFeatureLevelES31()
 
 		if (RHIName.Compare("D3D11_ES31", ESearchCase::IgnoreCase) == 0 && GPUBrand.Compare(DeviceName, ESearchCase::IgnoreCase) == 0)
 		{
-			UE_LOG(LogRHI, Log, TEXT("Found RHIName %s with DeviceName %s which has forced feature level to %s"), *RHIName, *DeviceName, GetLogName(ERHIFeatureLevel::ES3_1));
+			UE_LOG(LogRHI, Log, TEXT("Found RHIName %s with DeviceName %s: default feature level to %s"), *RHIName, *DeviceName, GetLogName(ERHIFeatureLevel::ES3_1));
 
-			ForceES31 = true;
+			ForceD3D11_ES31 = true;
 			return true;
 		}
 
@@ -457,14 +468,14 @@ static bool DefaultFeatureLevelES31()
 			BestGPUInfo.VendorId == VendorIdInt && BestGPUInfo.DeviceId == DeviceIdInt &&
 			RHIName.Compare("D3D11_ES31", ESearchCase::IgnoreCase) == 0)
 		{
-			UE_LOG(LogRHI, Log, TEXT("Found RHIName %s with DeviceName %s, VendorId %s, and DeviceId %s which has forced feature level to %s"), *RHIName, *DeviceName, *VendorId, *DeviceId, GetLogName(ERHIFeatureLevel::ES3_1));
+			UE_LOG(LogRHI, Log, TEXT("Found RHIName %s with DeviceName %s, VendorId %s, and DeviceId %s: default feature level to %s"), *RHIName, *DeviceName, *VendorId, *DeviceId, GetLogName(ERHIFeatureLevel::ES3_1));
 
-			ForceES31 = true;
+			ForceD3D11_ES31 = true;
 			return true;
 		}
 	}
 
-	ForceES31 = false;
+	ForceD3D11_ES31 = false;
 	return false;
 }
 
@@ -508,7 +519,7 @@ static TOptional<ERHIFeatureLevel::Type> GetPreferredFeatureLevel(EWindowsRHI Ch
 			bool bDefaultES31 = false;
 			if (!bFoundPreference && !CVarIgnorePerformanceModeCheck.GetValueOnAnyThread())
 			{
-				bDefaultES31 = DefaultFeatureLevelES31();
+				bDefaultES31 = ShouldDefaultToPerformanceMode();
 			}
 
 			if (bPreferFeatureLevelES31 || bDefaultES31)
@@ -622,9 +633,22 @@ static TOptional<ERHIFeatureLevel::Type> ChooseDefaultFeatureLevel(EWindowsRHI I
 	return TOptional<ERHIFeatureLevel::Type>();
 }
 
-// Choose the default from DefaultGraphicsRHI or TargetedRHIs. DefaultGraphicsRHI has precedence.
+// Choose the default RHI using this search order:
+//  * Performance mode (D3D11) if the machine is low-end (see ShouldDefaultToPerformanceMode) and if we are not running the editor;
+//  * DefaultGraphicsRHI if specified;
+//  * The first supported RHI from TargetedRHIs, enumerated according to GRHISearchOrder;
+//  * D3D11, if all else fails.
 static EWindowsRHI ChooseDefaultRHI(const FParsedWindowsDynamicRHIConfig& Config)
 {
+	if (!GIsEditor && !CVarIgnorePerformanceModeCheck.GetValueOnAnyThread())
+	{
+		bool bDefaultD3D11PerfMode = ShouldDefaultToPerformanceMode();
+		if (bDefaultD3D11PerfMode)
+		{
+			return EWindowsRHI::D3D11;
+		}
+	}
+
 	// Default graphics RHI is the main project setting that governs the choice, so it takes the priority
 	if (TOptional<EWindowsRHI> ConfigDefault = Config.DefaultRHI)
 	{
