New GPU profiler improvements
- Wrapped the legacy profileGPU command in (RHI_NEW_GPU_PROFILER == 0) macros, and renamed the new command from "ProfileGPUNew" to "ProfileGPU". When the new profiler is compiled in, the same "profilegpu" command now works as before.
- Consolidated various cvars into GPUProfiler.cpp, and moved the cvars that are not implemented or supported by the new profiler into the (RHI_NEW_GPU_PROFILER == 0) section of the file

diff --git a/Engine/Source/Runtime/Engine/Classes/Engine/Engine.h b/Engine/Source/Runtime/Engine/Classes/Engine/Engine.h
--- a/Engine/Source/Runtime/Engine/Classes/Engine/Engine.h
+++ b/Engine/Source/Runtime/Engine/Classes/Engine/Engine.h
@@ -2336,7 +2336,7 @@ public:
 	ENGINE_API bool HandleDumpParticleMemCommand( const TCHAR* Cmd, FOutputDevice& Ar );
 #endif
 
-#if WITH_PROFILEGPU
+#if WITH_PROFILEGPU && (RHI_NEW_GPU_PROFILER == 0)
 	ENGINE_API bool HandleProfileGPUCommand( const TCHAR* Cmd, FOutputDevice& Ar );	
 #endif
 
diff --git a/Engine/Source/Runtime/Engine/Private/UnrealEngine.cpp b/Engine/Source/Runtime/Engine/Private/UnrealEngine.cpp
--- a/Engine/Source/Runtime/Engine/Private/UnrealEngine.cpp
+++ b/Engine/Source/Runtime/Engine/Private/UnrealEngine.cpp
@@ -5194,7 +5194,7 @@ bool UEngine::Exec( UWorld* InWorld, const TCHAR* Cmd, FOutputDevice& Ar )
 	}
 #endif
 
-#if WITH_PROFILEGPU
+#if WITH_PROFILEGPU && (RHI_NEW_GPU_PROFILER == 0)
 	else if( FParse::Command(&Cmd,TEXT("PROFILEGPU")) )
 	{
 		return HandleProfileGPUCommand( Cmd, Ar );
@@ -6132,7 +6132,7 @@ bool UEngine::HandleDumpShaderCompileStatsCommand(const TCHAR* Cmd, FOutputDevic
 
 bool UEngine::HandleProfileCommand( const TCHAR* Cmd, FOutputDevice& Ar )
 {
-#if WITH_PROFILEGPU
+#if WITH_PROFILEGPU && (RHI_NEW_GPU_PROFILER == 0)
 	if ( FParse::Command(&Cmd,TEXT("GPU")) )
 	{
 		if (!FApp::CanEverRender())
@@ -6162,7 +6162,7 @@ bool UEngine::HandleProfileCommand( const TCHAR* Cmd, FOutputDevice& Ar )
 
 #endif // !(UE_BUILD_SHIPPING || UE_BUILD_TEST)
 
-#if WITH_PROFILEGPU
+#if WITH_PROFILEGPU && (RHI_NEW_GPU_PROFILER == 0)
 bool UEngine::HandleProfileGPUCommand( const TCHAR* Cmd, FOutputDevice& Ar )
 {
 	if (!FApp::CanEverRender())
@@ -6198,7 +6198,7 @@ bool UEngine::HandleProfileGPUCommand( const TCHAR* Cmd, FOutputDevice& Ar )
 	return true;
 }
 
-#endif // WITH_PROFILEGPU
+#endif // WITH_PROFILEGPU && (RHI_NEW_GPU_PROFILER == 0)
 
 #if WITH_DUMPGPU
 
diff --git a/Engine/Source/Runtime/Launch/Private/LaunchEngineLoop.cpp b/Engine/Source/Runtime/Launch/Private/LaunchEngineLoop.cpp
--- a/Engine/Source/Runtime/Launch/Private/LaunchEngineLoop.cpp
+++ b/Engine/Source/Runtime/Launch/Private/LaunchEngineLoop.cpp
@@ -5505,7 +5505,7 @@ static inline void BeginFrameRenderThread(FRHICommandListImmediate& RHICmdList,
 
 		// If we are profiling, kick off a long GPU task to make the GPU always behind the CPU so that we
 		// won't get GPU idle time measured in profiling results
-	#if WITH_PROFILEGPU 
+	#if WITH_PROFILEGPU && (RHI_NEW_GPU_PROFILER == 0)
 		if (GTriggerGPUProfile && !GTriggerGPUHitchProfile)
 		{
 			IssueScalableLongGPUTask(RHICmdList);
@@ -5688,7 +5688,7 @@ void FEngineLoop::Tick()
 	{
 		SCOPE_CYCLE_COUNTER(STAT_FrameTime);
 
-		#if WITH_PROFILEGPU && !UE_BUILD_SHIPPING
+		#if WITH_PROFILEGPU && !UE_BUILD_SHIPPING && (RHI_NEW_GPU_PROFILER == 0)
 			// Issue the measurement of the execution time of a basic LongGPUTask unit on the very first frame
 			// The results will be retrived on the first call of IssueScalableLongGPUTask
 			if (GFrameCounter == 0 && IsFeatureLevelSupported(GMaxRHIShaderPlatform, ERHIFeatureLevel::SM5) && FApp::CanEverRender())
diff --git a/Engine/Source/Runtime/RHI/Private/GPUProfiler.cpp b/Engine/Source/Runtime/RHI/Private/GPUProfiler.cpp
--- a/Engine/Source/Runtime/RHI/Private/GPUProfiler.cpp
+++ b/Engine/Source/Runtime/RHI/Private/GPUProfiler.cpp
@@ -44,21 +44,26 @@
 #endif // HAS_GPU_STATS
 }
 
-static TAutoConsoleVariable<int> CVarGPUCsvStatsEnabled(
-	TEXT("r.GPUCsvStatsEnabled"),
-	0,
-	TEXT("Enables or disables GPU stat recording to CSVs"));
-
 #define LOCTEXT_NAMESPACE "GpuProfiler"
 
-static TAutoConsoleVariable<FString> GProfileGPUPatternCVar(
-	TEXT("r.ProfileGPU.Pattern"),
-	TEXT("*"),
-	TEXT("Allows to filter the entries when using ProfileGPU, the pattern match is case sensitive.\n")
-	TEXT("'*' can be used in the end to get all entries starting with the string.\n")
-	TEXT("    '*' without any leading characters disables the pattern matching and uses a time threshold instead (default).\n")
-	TEXT("'?' allows to ignore one character.\n")
-	TEXT("e.g. AmbientOcclusionSetup, AmbientOcclusion*, Ambient???lusion*, *"),
+enum class EGPUProfileSortMode
+{
+	Chronological,
+	TimeElapsed,
+	NumPrims,
+	NumVerts,
+
+	Max
+};
+
+static TAutoConsoleVariable<int32> GCVarProfileGPU_Sort(
+	TEXT("r.ProfileGPU.Sort"),
+	0,
+	TEXT("Sorts the TTY Dump independently at each level of the tree in various modes.\n")
+	TEXT("0 : Chronological\n")
+	TEXT("1 : By time elapsed\n")
+	TEXT("2 : By number of prims\n")
+	TEXT("3 : By number of verts\n"),
 	ECVF_Default);
 
 static TAutoConsoleVariable<FString> GCVarProfileGPU_Root(
@@ -79,18 +84,35 @@
 	TEXT("When enabled, the output results will be formatted in a unicode table."),
 	ECVF_Default);
 
-static TAutoConsoleVariable<int32> GProfileShowEventHistogram(
-	TEXT("r.ProfileGPU.ShowEventHistogram"),
-	0,
-	TEXT("Whether the event histogram should be shown."),
-	ECVF_Default);
-
 static TAutoConsoleVariable<bool> GCVarProfileGPU_ShowLeafEvents(
 	TEXT("r.ProfileGPU.ShowLeafEvents"),
 	true,
 	TEXT("Allows profileGPU to display event-only leaf nodes with no draws associated."),
 	ECVF_Default);
 
+static TAutoConsoleVariable<int> CVarGPUCsvStatsEnabled(
+	TEXT("r.GPUCsvStatsEnabled"),
+	0,
+	TEXT("Enables or disables GPU stat recording to CSVs"));
+
+#if (RHI_NEW_GPU_PROFILER == 0)
+
+static TAutoConsoleVariable<FString> GProfileGPUPatternCVar(
+	TEXT("r.ProfileGPU.Pattern"),
+	TEXT("*"),
+	TEXT("Allows to filter the entries when using ProfileGPU, the pattern match is case sensitive.\n")
+	TEXT("'*' can be used in the end to get all entries starting with the string.\n")
+	TEXT("    '*' without any leading characters disables the pattern matching and uses a time threshold instead (default).\n")
+	TEXT("'?' allows to ignore one character.\n")
+	TEXT("e.g. AmbientOcclusionSetup, AmbientOcclusion*, Ambient???lusion*, *"),
+	ECVF_Default);
+
+static TAutoConsoleVariable<int32> GProfileShowEventHistogram(
+	TEXT("r.ProfileGPU.ShowEventHistogram"),
+	0,
+	TEXT("Whether the event histogram should be shown."),
+	ECVF_Default);
+
 RHI_API TAutoConsoleVariable<int32> GProfileGPUTransitions(
 	TEXT("r.ProfileGPU.ShowTransitions"),
 	0,
@@ -112,6 +134,26 @@
 	TEXT("r.ProfileGPU.PrintAssetSummary must be true to enable this feature"),
 	ECVF_Default);
 
+static TAutoConsoleVariable<int32> GSaveScreenshotAfterProfilingGPUCVar(
+	TEXT("r.ProfileGPU.Screenshot"),
+	1,
+	TEXT("Whether a screenshot should be taken when profiling the GPU. 0:off, 1:on (default)"),
+	ECVF_RenderThreadSafe);
+
+static TAutoConsoleVariable<int32> GShowProfilerAfterProfilingGPUCVar(
+	TEXT("r.ProfileGPU.ShowUI"),
+	1,
+	TEXT("Whether the user interface profiler should be displayed after profiling the GPU.\n")
+	TEXT("The results will always go to the log/console\n")
+	TEXT("0:off, 1:on (default)"),
+	ECVF_RenderThreadSafe);
+
+static TAutoConsoleVariable<float> GGPUHitchThresholdCVar(
+	TEXT("RHI.GPUHitchThreshold"),
+	100.0f,
+	TEXT("Threshold for detecting hitches on the GPU (in milliseconds).")
+);
+
 static TAutoConsoleVariable<int32> CVarGPUCrashDataCollectionEnable(
 	TEXT("r.gpucrash.collectionenable"),
 	1,
@@ -124,26 +166,23 @@
 	TEXT("Limits the amount of marker scope depth we record for GPU crash debugging to the given scope depth."),
 	ECVF_RenderThreadSafe);
 
-enum class EGPUProfileSortMode
+namespace RHIConfig
 {
-	EChronological,
-	ETimeElapsed,
-	ENumPrims,
-	ENumVerts,
-	EMax
-};
+	bool ShouldSaveScreenshotAfterProfilingGPU()
+	{
+		return GSaveScreenshotAfterProfilingGPUCVar.GetValueOnAnyThread() != 0;
+	}
 
-static TAutoConsoleVariable<int32> GProfileGPUSort(
-	TEXT("r.ProfileGPU.Sort"),	
-	0,
-	TEXT("Sorts the TTY Dump independently at each level of the tree in various modes.\n")
-	TEXT("0 : Chronological\n")
-	TEXT("1 : By time elapsed\n")
-	TEXT("2 : By number of prims\n")
-	TEXT("3 : By number of verts\n"),
-	ECVF_Default);
+	bool ShouldShowProfilerAfterProfilingGPU()
+	{
+		return GShowProfilerAfterProfilingGPUCVar.GetValueOnAnyThread() != 0;
+	}
 
-#if (RHI_NEW_GPU_PROFILER == 0)
+	float GetGPUHitchThreshold()
+	{
+		return GGPUHitchThresholdCVar.GetValueOnAnyThread() * 0.001f;
+	}
+}
 
 /** Recursively generates a histogram of nodes and stores their timing in TimingResult. */
 static void GatherStatsEventNode(FGPUProfilerEventNode* Node, int32 Depth, TMap<FString, FGPUProfilerEventNodeStats>& EventHistogram)
@@ -386,19 +425,19 @@
 			{
 				switch (SortMode)
 				{
-					case EGPUProfileSortMode::ENumPrims:
+					case EGPUProfileSortMode::NumPrims:
 						return B->NumTotalPrimitives < A->NumTotalPrimitives;
-					case EGPUProfileSortMode::ENumVerts:
+					case EGPUProfileSortMode::NumVerts:
 						return B->NumTotalVertices < A->NumTotalVertices;
-					case EGPUProfileSortMode::ETimeElapsed:
+					case EGPUProfileSortMode::TimeElapsed:
 					default:
 						return B->TimingResult < A->TimingResult;
 				}
 			}
 		};
 
-		EGPUProfileSortMode SortMode = (EGPUProfileSortMode)FMath::Clamp(GProfileGPUSort.GetValueOnRenderThread(), 0, ((int32)EGPUProfileSortMode::EMax - 1));
-		if (SortMode != EGPUProfileSortMode::EChronological)
+		EGPUProfileSortMode SortMode = (EGPUProfileSortMode)FMath::Clamp(GCVarProfileGPU_Sort.GetValueOnRenderThread(), 0, ((int32)EGPUProfileSortMode::Max - 1));
+		if (SortMode != EGPUProfileSortMode::Chronological)
 		{
 			Node->Children.Sort(FCompareGPUProfileNode(SortMode));
 		}
@@ -716,6 +755,8 @@
 
 namespace UE::RHI::GPUProfiler
 {
+	TLockFreePointerListUnordered<void, PLATFORM_CACHE_LINE_SIZE> FEventStream::FChunk::MemoryPool;
+
 	static TArray<FEventSink*>& GetSinks()
 	{
 		static TArray<FEventSink*> Sinks;
@@ -1377,6 +1418,7 @@
 
 	} GGPUProfilerSink_StatSystem;
 
+#if WITH_PROFILEGPU
 	struct FGPUProfilerSink_ProfileGPU final : public FEventSink
 	{
 		struct FNode
@@ -1724,12 +1766,12 @@
 
 	void FGPUProfilerSink_ProfileGPU::FQueueState::LogTree(uint32 FrameNumber)
 	{
-		EGPUProfileSortMode SortMode = (EGPUProfileSortMode)FMath::Clamp(GProfileGPUSort.GetValueOnAnyThread(), 0, ((int32)EGPUProfileSortMode::EMax - 1));
+		EGPUProfileSortMode SortMode = (EGPUProfileSortMode)FMath::Clamp(GCVarProfileGPU_Sort.GetValueOnAnyThread(), 0, ((int32)EGPUProfileSortMode::Max - 1));
 		FWildcardString RootWildcard(GCVarProfileGPU_Root.GetValueOnAnyThread());
 		const bool bShowEmptyNodes = GCVarProfileGPU_ShowLeafEvents.GetValueOnAnyThread();
 		const double PercentThreshold = FMath::Clamp(GCVarProfileGPU_ThresholdPercent.GetValueOnAnyThread(), 0.0f, 100.0f);
 
-		if (SortMode != EGPUProfileSortMode::EChronological)
+		if (SortMode != EGPUProfileSortMode::Chronological)
 		{
 			for (FNode* Node = State.First; Node; Node = Node->Next)
 			{
@@ -1738,9 +1780,9 @@
 					switch (SortMode)
 					{
 					default:
-					case EGPUProfileSortMode::ETimeElapsed: return B.Inclusive.BusyCycles    < A.Inclusive.BusyCycles;
-					case EGPUProfileSortMode::ENumPrims   : return B.Inclusive.NumPrimitives < A.Inclusive.NumPrimitives;
-					case EGPUProfileSortMode::ENumVerts   : return B.Inclusive.NumVertices   < A.Inclusive.NumVertices;
+					case EGPUProfileSortMode::TimeElapsed: return B.Inclusive.BusyCycles    < A.Inclusive.BusyCycles;
+					case EGPUProfileSortMode::NumPrims   : return B.Inclusive.NumPrimitives < A.Inclusive.NumPrimitives;
+					case EGPUProfileSortMode::NumVerts   : return B.Inclusive.NumVertices   < A.Inclusive.NumVertices;
 					}
 				});
 			}
@@ -2278,15 +2320,14 @@
 		}
 	}
 
-	static FAutoConsoleCommand ProfileGPUNew(
-		TEXT("ProfileGPUNew"),
-		TEXT(""),
+	static FAutoConsoleCommand GCommand_ProfileGPU(
+		TEXT("ProfileGPU"),
+		TEXT("Captures statistics about a frame of GPU work and prints the results to the log."),
 		FConsoleCommandWithArgsDelegate::CreateLambda([](const TArray<FString>& Args)
 		{
 			GGPUProfilerSink_ProfileGPU.ProfileNextFrame();
 		}));
-
-	TLockFreePointerListUnordered<void, PLATFORM_CACHE_LINE_SIZE> FEventStream::FChunk::MemoryPool;
+#endif // WITH_PROFILEGPU
 }
 
 #endif // RHI_NEW_GPU_PROFILER
diff --git a/Engine/Source/Runtime/RHI/Private/RHI.cpp b/Engine/Source/Runtime/RHI/Private/RHI.cpp
--- a/Engine/Source/Runtime/RHI/Private/RHI.cpp
+++ b/Engine/Source/Runtime/RHI/Private/RHI.cpp
@@ -1088,23 +1088,6 @@ static_assert((int32)ERHIZBuffer::FarPlane == 0 || (int32)ERHIZBuffer::FarPlane
 /**
  * RHI configuration settings.
  */
-static TAutoConsoleVariable<int32> GSaveScreenshotAfterProfilingGPUCVar(
-	TEXT("r.ProfileGPU.Screenshot"),
-	1,
-	TEXT("Whether a screenshot should be taken when profiling the GPU. 0:off, 1:on (default)"),
-	ECVF_RenderThreadSafe);
-static TAutoConsoleVariable<int32> GShowProfilerAfterProfilingGPUCVar(
-	TEXT("r.ProfileGPU.ShowUI"),
-	1,
-	TEXT("Whether the user interface profiler should be displayed after profiling the GPU.\n")
-	TEXT("The results will always go to the log/console\n")
-	TEXT("0:off, 1:on (default)"),
-	ECVF_RenderThreadSafe);
-static TAutoConsoleVariable<float> GGPUHitchThresholdCVar(
-	TEXT("RHI.GPUHitchThreshold"),
-	100.0f,
-	TEXT("Threshold for detecting hitches on the GPU (in milliseconds).")
-	);
 static TAutoConsoleVariable<int32> GCVarRHIRenderPass(
 	TEXT("r.RHIRenderPasses"),
 	0,
@@ -1237,24 +1220,6 @@ ERHIBindlessConfiguration RHIGetRuntimeBindlessSamplersConfiguration(EShaderPlat
 	return DetermineBindlessConfiguration(Platform, TEXT("BindlessSamplers"), GRHIBindlessSamplerConfiguration);
 }
 
-namespace RHIConfig
-{
-	bool ShouldSaveScreenshotAfterProfilingGPU()
-	{
-		return GSaveScreenshotAfterProfilingGPUCVar.GetValueOnAnyThread() != 0;
-	}
-
-	bool ShouldShowProfilerAfterProfilingGPU()
-	{
-		return GShowProfilerAfterProfilingGPUCVar.GetValueOnAnyThread() != 0;
-	}
-
-	float GetGPUHitchThreshold()
-	{
-		return GGPUHitchThresholdCVar.GetValueOnAnyThread() * 0.001f;
-	}
-}
-
 // By default, read only states and UAV states are allowed to participate in state merging.
 ERHIAccess GRHIMergeableAccessMask = ERHIAccess::ReadOnlyMask | ERHIAccess::UAVMask;
 
diff --git a/Engine/Source/Runtime/RHI/Public/RHI.h b/Engine/Source/Runtime/RHI/Public/RHI.h
--- a/Engine/Source/Runtime/RHI/Public/RHI.h
+++ b/Engine/Source/Runtime/RHI/Public/RHI.h
@@ -30,12 +30,14 @@ class FRHICommandList;
  * RHI configuration settings.
  */
 
+#if (RHI_NEW_GPU_PROFILER == 0)
 namespace RHIConfig
 {
 	RHI_API bool ShouldSaveScreenshotAfterProfilingGPU();
 	RHI_API bool ShouldShowProfilerAfterProfilingGPU();
 	RHI_API float GetGPUHitchThreshold();
 }
+#endif
 
 /**
  * RHI capabilities.
diff --git a/Engine/Source/Runtime/RHI/Public/RHICommandList.h b/Engine/Source/Runtime/RHI/Public/RHICommandList.h
--- a/Engine/Source/Runtime/RHI/Public/RHICommandList.h
+++ b/Engine/Source/Runtime/RHI/Public/RHICommandList.h
@@ -2891,7 +2891,7 @@ public:
 
 	FORCEINLINE_DEBUGGABLE void BeginTransitions(TArrayView<const FRHITransition*> Transitions)
 	{
-#if WITH_PROFILEGPU
+#if WITH_PROFILEGPU && (RHI_NEW_GPU_PROFILER == 0)
 		extern RHI_API TAutoConsoleVariable<int32> GProfileGPUTransitions;
 		RHI_BREADCRUMB_EVENT_CONDITIONAL(*this, GProfileGPUTransitions.GetValueOnAnyThread() != 0, "RHIBeginTransitions");
 #endif
@@ -2930,7 +2930,7 @@ public:
 
 	FORCEINLINE_DEBUGGABLE void EndTransitions(TArrayView<const FRHITransition*> Transitions)
 	{
-#if WITH_PROFILEGPU
+#if WITH_PROFILEGPU && (RHI_NEW_GPU_PROFILER == 0)
 		extern RHI_API TAutoConsoleVariable<int32> GProfileGPUTransitions;
 		RHI_BREADCRUMB_EVENT_CONDITIONAL(*this, GProfileGPUTransitions.GetValueOnAnyThread() != 0, "RHIEndTransitions");
 #endif
