SBT RHI Validation improvements:
- index & vertex buffers used in the SBT are not validated for SRV state as well
- index buffer used in BLAS build is validated for SRV state
- fixup of different parts in the engine to make sure they are in correct SRV state

diff --git a/Engine/Plugins/Runtime/CableComponent/Source/CableComponent/Private/CableComponent.cpp b/Engine/Plugins/Runtime/CableComponent/Source/CableComponent/Private/CableComponent.cpp
--- a/Engine/Plugins/Runtime/CableComponent/Source/CableComponent/Private/CableComponent.cpp
+++ b/Engine/Plugins/Runtime/CableComponent/Source/CableComponent/Private/CableComponent.cpp
@@ -61,7 +61,7 @@ public:
 	virtual void InitRHI(FRHICommandListBase& RHICmdList) override
 	{
 		FRHIResourceCreateInfo CreateInfo(TEXT("FCableIndexBuffer"));
-		IndexBufferRHI = RHICmdList.CreateIndexBuffer(sizeof(int32), NumIndices * sizeof(int32), BUF_Dynamic, CreateInfo);
+		IndexBufferRHI = RHICmdList.CreateIndexBuffer(sizeof(int32), NumIndices * sizeof(int32), BUF_Dynamic | BUF_ShaderResource, CreateInfo);
 	}
 
 	int32 NumIndices;
diff --git a/Engine/Plugins/Runtime/HairStrands/Source/HairStrandsCore/Private/GroomResources.cpp b/Engine/Plugins/Runtime/HairStrands/Source/HairStrandsCore/Private/GroomResources.cpp
--- a/Engine/Plugins/Runtime/HairStrands/Source/HairStrandsCore/Private/GroomResources.cpp
+++ b/Engine/Plugins/Runtime/HairStrands/Source/HairStrandsCore/Private/GroomResources.cpp
@@ -990,7 +990,7 @@ void FHairCardIndexBuffer::InitRHI(FRHICommandListBase& RHICmdList)
 	const uint32 DataSizeInBytes = FHairCardsIndexFormat::SizeInByte * Indices.Num();
 
 	FRHIResourceCreateInfo CreateInfo(TEXT("FHairCardIndexBuffer"));
-	IndexBufferRHI = RHICmdList.CreateBuffer(DataSizeInBytes, BUF_Static | BUF_IndexBuffer, FHairCardsIndexFormat::SizeInByte, ERHIAccess::VertexOrIndexBuffer, CreateInfo);
+	IndexBufferRHI = RHICmdList.CreateBuffer(DataSizeInBytes, BUF_Static | BUF_IndexBuffer | BUF_ShaderResource, FHairCardsIndexFormat::SizeInByte, ERHIAccess::VertexOrIndexBuffer | ERHIAccess::SRVCompute, CreateInfo);
 	void* Buffer = RHICmdList.LockBuffer(IndexBufferRHI, 0, DataSizeInBytes, RLM_WriteOnly);
 	FMemory::Memcpy(Buffer, Indices.GetData(), DataSizeInBytes);
 	RHICmdList.UnlockBuffer(IndexBufferRHI);
@@ -1071,7 +1071,7 @@ void FHairCardsDeformedResource::InternalAllocate(FRDGBuilder& GraphBuilder)
 
 		// Manually transit to SRVs, in case of the cards are rendered not visible, but still rendered (in shadows for instance). In such a case, the cards deformation pass is not called, and thus the 
 		// buffers are never transit from UAV (clear) to SRV for rasterization. 
-		GraphBuilder.UseExternalAccessMode(Register(GraphBuilder, DeformedPositionBuffer[0], ERDGImportedBufferFlags::CreateSRV).Buffer, ERHIAccess::SRVMask);
+		GraphBuilder.UseExternalAccessMode(Register(GraphBuilder, DeformedPositionBuffer[0], ERDGImportedBufferFlags::CreateSRV).Buffer, ERHIAccess::SRVMask, ERHIPipeline::All);
 		GraphBuilder.UseExternalAccessMode(Register(GraphBuilder, DeformedPositionBuffer[1], ERDGImportedBufferFlags::CreateSRV).Buffer, ERHIAccess::SRVMask);
 	}
 
diff --git a/Engine/Source/Runtime/Engine/Private/Components/TextRenderComponent.cpp b/Engine/Source/Runtime/Engine/Private/Components/TextRenderComponent.cpp
--- a/Engine/Source/Runtime/Engine/Private/Components/TextRenderComponent.cpp
+++ b/Engine/Source/Runtime/Engine/Private/Components/TextRenderComponent.cpp
@@ -738,6 +738,9 @@ void FTextRenderSceneProxy::CreateRenderThreadResources(FRHICommandListBase& RHI
 		FName Name = FName(TEXT("FTextRenderSceneProxy ") + GetOwnerName().ToString());
 		VertexBuffers.SetOwnerName(Name);
 		IndexBuffer.SetOwnerName(Name);
+#endif
+#if RHI_RAYTRACING
+		EnumAddFlags(IndexBuffer.UsageFlags, EBufferUsageFlags::ShaderResource);
 #endif
 		VertexBuffers.InitFromDynamicVertex(RHICmdList, &VertexFactory, OutVertices);
 		// Enqueue initialization of render resources
diff --git a/Engine/Source/Runtime/Engine/Private/DynamicMeshBuilder.cpp b/Engine/Source/Runtime/Engine/Private/DynamicMeshBuilder.cpp
--- a/Engine/Source/Runtime/Engine/Private/DynamicMeshBuilder.cpp
+++ b/Engine/Source/Runtime/Engine/Private/DynamicMeshBuilder.cpp
@@ -199,12 +199,12 @@ TGlobalResource<FGlobalDynamicMeshVertexPool> GDynamicMeshVertexPool;
 
 void FDynamicMeshIndexBuffer32::InitRHI(FRHICommandListBase& RHICmdList)
 {
-	IndexBufferRHI = UE::RHIResourceUtils::CreateIndexBufferFromArray(RHICmdList, TEXT("FDynamicMeshIndexBuffer32"), EBufferUsageFlags::Static, MakeConstArrayView(Indices));
+	IndexBufferRHI = UE::RHIResourceUtils::CreateIndexBufferFromArray(RHICmdList, TEXT("FDynamicMeshIndexBuffer32"), UsageFlags, MakeConstArrayView(Indices));
 }
 
 void FDynamicMeshIndexBuffer16::InitRHI(FRHICommandListBase& RHICmdList)
 {
-	IndexBufferRHI = UE::RHIResourceUtils::CreateIndexBufferFromArray(RHICmdList, TEXT("FDynamicMeshIndexBuffer16"), EBufferUsageFlags::Static, MakeConstArrayView(Indices));
+	IndexBufferRHI = UE::RHIResourceUtils::CreateIndexBufferFromArray(RHICmdList, TEXT("FDynamicMeshIndexBuffer16"), UsageFlags, MakeConstArrayView(Indices));
 }
 
 /** FDynamicMeshBufferAllocator's base implementation. It always reallocates new buffers. */
diff --git a/Engine/Source/Runtime/Engine/Public/DynamicMeshBuilder.h b/Engine/Source/Runtime/Engine/Public/DynamicMeshBuilder.h
--- a/Engine/Source/Runtime/Engine/Public/DynamicMeshBuilder.h
+++ b/Engine/Source/Runtime/Engine/Public/DynamicMeshBuilder.h
@@ -227,6 +227,7 @@ class FDynamicMeshIndexBuffer32 : public FIndexBuffer
 {
 public:
 	TArray<uint32> Indices;
+	EBufferUsageFlags UsageFlags = EBufferUsageFlags::Static;
 
 	ENGINE_API virtual void InitRHI(FRHICommandListBase& RHICmdList) override;
 };
@@ -235,6 +236,7 @@ class FDynamicMeshIndexBuffer16 : public FIndexBuffer
 {
 public:
 	TArray<uint16> Indices;
+	EBufferUsageFlags UsageFlags = EBufferUsageFlags::Static;
 
 	ENGINE_API virtual void InitRHI(FRHICommandListBase& RHICmdList) override;
 };
diff --git a/Engine/Source/Runtime/RHI/Private/RHIValidation.cpp b/Engine/Source/Runtime/RHI/Private/RHIValidation.cpp
--- a/Engine/Source/Runtime/RHI/Private/RHIValidation.cpp
+++ b/Engine/Source/Runtime/RHI/Private/RHIValidation.cpp
@@ -487,6 +487,17 @@ namespace RHIValidation
 				if (EnumHasAnyFlags(ShaderBindingMode, ERayTracingShaderBindingMode::RTPSO))
 				{
 					CollectShaderBindingTableResources(RayTracingPipelineState, this, Binding, BindingType, Context.WorkerIndex);
+
+					// Also add SRV view requirement for all index and vertex buffers used in the SBT
+					const FRayTracingGeometryInitializer& BLASInitializer = Binding.Geometry->GetInitializer();
+					if (BLASInitializer.IndexBuffer)
+					{
+						AddSRV(BLASInitializer.IndexBuffer->GetWholeResourceIdentity(), Context.WorkerIndex);
+					}
+					for (int32 SegmentIndex = 0; SegmentIndex < BLASInitializer.Segments.Num(); ++SegmentIndex)
+					{
+						AddSRV(BLASInitializer.Segments[SegmentIndex].VertexBuffer->GetWholeResourceIdentity(), Context.WorkerIndex);
+					}
 				}
 			}
 		};
diff --git a/Engine/Source/Runtime/RHI/Public/RHIValidationContext.h b/Engine/Source/Runtime/RHI/Public/RHIValidationContext.h
--- a/Engine/Source/Runtime/RHI/Public/RHIValidationContext.h
+++ b/Engine/Source/Runtime/RHI/Public/RHIValidationContext.h
@@ -148,6 +148,26 @@ public:
 		RHIContext->RHISetShaderRootConstants(Constants);
 	}
 
+	virtual void RHIRayTraceDispatch(FRHIRayTracingPipelineState* RayTracingPipelineState, FRHIRayTracingShader* RayGenShader,
+		FRHIShaderBindingTable* SBT, const FRayTracingShaderBindings& GlobalResourceBindings,
+		uint32 Width, uint32 Height) override final
+	{
+		SBT->ValidateStateForDispatch(Tracker);
+		RHIContext->RHIRayTraceDispatch(RayTracingPipelineState, RayGenShader, SBT, GlobalResourceBindings, Width, Height);
+	}
+
+	virtual void RHIRayTraceDispatchIndirect(FRHIRayTracingPipelineState* RayTracingPipelineState, FRHIRayTracingShader* RayGenShader,
+		FRHIShaderBindingTable* SBT, const FRayTracingShaderBindings& GlobalResourceBindings,
+		FRHIBuffer* ArgumentBuffer, uint32 ArgumentOffset) override final
+	{
+		FValidationRHI::ValidateDispatchIndirectArgsBuffer(ArgumentBuffer, ArgumentOffset);
+		Tracker->Assert(ArgumentBuffer->GetWholeResourceIdentity(), ERHIAccess::IndirectArgs);
+		Tracker->Assert(ArgumentBuffer->GetWholeResourceIdentity(), ERHIAccess::SRVCompute);
+
+		SBT->ValidateStateForDispatch(Tracker);
+		RHIContext->RHIRayTraceDispatchIndirect(RayTracingPipelineState, RayGenShader, SBT, GlobalResourceBindings, ArgumentBuffer, ArgumentOffset);
+	}
+
 	virtual void RHIDispatchComputeShaderBundle(
 		FRHIShaderBundle* ShaderBundleRHI,
 		FRHIBuffer* RecordArgBuffer,
@@ -1197,12 +1217,14 @@ public:
 
 	virtual void RHIBuildAccelerationStructures(TConstArrayView<FRayTracingGeometryBuildParams> Params, const FRHIBufferRange& ScratchBufferRange) override final
 	{
-		// #yuriy_todo: explicit transitions and state validation for BLAS
 		for (const FRayTracingGeometryBuildParams& P : Params)
 		{
 			const FRayTracingGeometryInitializer& Initializer = P.Geometry->GetInitializer();
 
-//			TODO: Add ERHIAccess::SRVCompute assert for Initializer.IndexBuffer, still some resources in the wrong state eg. FHairCardIndexBuffer
+			if (Initializer.IndexBuffer)
+			{
+				Tracker->Assert(Initializer.IndexBuffer->GetWholeResourceIdentity(), ERHIAccess::SRVCompute);
+			}
 
 			for (const FRayTracingGeometrySegment& Segment : Initializer.Segments)
 			{
@@ -1215,7 +1237,6 @@ public:
 
 	virtual void RHIBuildAccelerationStructure(const FRayTracingSceneBuildParams& SceneBuildParams) override final
 	{
-		// #yuriy_todo: validate all referenced BLAS states
 		if (SceneBuildParams.Scene)
 		{
 			Tracker->Assert(SceneBuildParams.Scene->GetWholeResourceIdentity(), ERHIAccess::BVHWrite);
diff --git a/Engine/Source/Runtime/Renderer/Private/Tests/RayTracingTestbed.cpp b/Engine/Source/Runtime/Renderer/Private/Tests/RayTracingTestbed.cpp
--- a/Engine/Source/Runtime/Renderer/Private/Tests/RayTracingTestbed.cpp
+++ b/Engine/Source/Runtime/Renderer/Private/Tests/RayTracingTestbed.cpp
@@ -73,7 +73,7 @@ bool RunRayTracingTestbed_RenderThread(const FString& Parameters)
 			FVector3f(-1, -1, 0),
 		};
 
-		VertexBuffer = UE::RHIResourceUtils::CreateVertexBufferFromArray(RHICmdList, TEXT("RayTracingTestbedVB"), EBufferUsageFlags::Static, MakeConstArrayView(PositionData));
+		VertexBuffer = UE::RHIResourceUtils::CreateVertexBufferFromArray(RHICmdList, TEXT("RayTracingTestbedVB"), EBufferUsageFlags::Static | EBufferUsageFlags::ShaderResource, MakeConstArrayView(PositionData));
 	}
 
 	FBufferRHIRef IndexBuffer;
@@ -84,7 +84,7 @@ bool RunRayTracingTestbed_RenderThread(const FString& Parameters)
 			0, 1, 2
 		};
 
-		IndexBuffer = UE::RHIResourceUtils::CreateIndexBufferFromArray(RHICmdList, TEXT("RayTracingTestbedIB"), EBufferUsageFlags::Static, MakeConstArrayView(IndexData));
+		IndexBuffer = UE::RHIResourceUtils::CreateIndexBufferFromArray(RHICmdList, TEXT("RayTracingTestbedIB"), EBufferUsageFlags::Static | EBufferUsageFlags::ShaderResource, MakeConstArrayView(IndexData));
 	}
 
 	static constexpr uint32 NumRays = 4;
