Added "working memory" stat for integrated contexts
- Updated URendererSettings::SanatizeReflectionCaptureResolution and USkyLightComponent::SanitizeCubemapSize to use new metric

diff --git a/Engine/Source/Runtime/Engine/Private/Components/SkyLightComponent.cpp b/Engine/Source/Runtime/Engine/Private/Components/SkyLightComponent.cpp
--- a/Engine/Source/Runtime/Engine/Private/Components/SkyLightComponent.cpp
+++ b/Engine/Source/Runtime/Engine/Private/Components/SkyLightComponent.cpp
@@ -382,8 +382,10 @@ void USkyLightComponent::SanitizeCubemapSize()
 
 		FTextureMemoryStats TextureMemStats;
 		RHIGetTextureMemoryStats(TextureMemStats);
+		
+		int64 DeviceWorkingMemory = TextureMemStats.GetTotalDeviceWorkingMemory();
 
-		if (TextureMemStats.DedicatedVideoMemory > 0 && TexMemRequired > SIZE_T(TextureMemStats.DedicatedVideoMemory / 4))
+		if (DeviceWorkingMemory > 0 && TexMemRequired > SIZE_T(DeviceWorkingMemory / 4))
 		{
 			FNumberFormattingOptions FmtOpts = FNumberFormattingOptions()
 				.SetUseGrouping(false)
diff --git a/Engine/Source/Runtime/Engine/Private/RendererSettings.cpp b/Engine/Source/Runtime/Engine/Private/RendererSettings.cpp
--- a/Engine/Source/Runtime/Engine/Private/RendererSettings.cpp
+++ b/Engine/Source/Runtime/Engine/Private/RendererSettings.cpp
@@ -668,7 +668,9 @@ void URendererSettings::SanatizeReflectionCaptureResolution()
 		FTextureMemoryStats TextureMemStats;
 		RHIGetTextureMemoryStats(TextureMemStats);
 
-		if (TextureMemStats.DedicatedVideoMemory > 0 && TexMemRequired > SIZE_T(TextureMemStats.DedicatedVideoMemory / 8))
+		int64 DeviceWorkingMemory = TextureMemStats.GetTotalDeviceWorkingMemory();
+		
+		if (DeviceWorkingMemory > 0 && TexMemRequired > SIZE_T(DeviceWorkingMemory / 8))
 		{
 			FNumberFormattingOptions FmtOpts = FNumberFormattingOptions()
 				.SetUseGrouping(false)
diff --git a/Engine/Source/Runtime/RHI/Public/RHIStats.h b/Engine/Source/Runtime/RHI/Public/RHIStats.h
--- a/Engine/Source/Runtime/RHI/Public/RHIStats.h
+++ b/Engine/Source/Runtime/RHI/Public/RHIStats.h
@@ -5,6 +5,7 @@
 #include "Containers/StaticArray.h"
 #include "MultiGPU.h"
 #include "RHIDefinitions.h"
+#include "RHIGlobals.h"
 #include "Stats/Stats.h"
 
 struct FTextureMemoryStats
@@ -35,6 +36,17 @@ struct FTextureMemoryStats
 	// 0 if streaming pool size limitation is disabled, in bytes
 	int64 TexturePoolSize = 0;
 
+	int64 GetTotalDeviceWorkingMemory() const
+	{
+		if (GRHIDeviceIsIntegrated)
+		{
+			// Max in case the device failed to report the available working memory
+			return FMath::Max(TotalGraphicsMemory, DedicatedVideoMemory);
+		}
+		
+		return DedicatedVideoMemory;
+	}
+
 	bool AreHardwareStatsValid() const
 	{
 		// pardon the redundancy, have a broken compiler (__EMSCRIPTEN__) that needs these types spelled out...
