Move RayTraceDispatch & RayTraceDispatchIndirect from FRHICommandList to FRHIComputeCommandList to match the IRHIComputeContext

diff --git a/Engine/Source/Runtime/RHI/Private/RHICommandList.cpp b/Engine/Source/Runtime/RHI/Private/RHICommandList.cpp
--- a/Engine/Source/Runtime/RHI/Private/RHICommandList.cpp
+++ b/Engine/Source/Runtime/RHI/Private/RHICommandList.cpp
@@ -2181,7 +2181,7 @@ FRayTracingShaderBindings UE::RHI::ConvertRayTracingShaderBindings(const FRHIBat
 	return Result;
 }
 
-void FRHICommandList::RayTraceDispatch(
+void FRHIComputeCommandList::RayTraceDispatch(
 	FRayTracingPipelineState* Pipeline,
 	FRHIRayTracingShader* RayGenShader,
 	FRHIShaderBindingTable* SBT,
@@ -2202,7 +2202,7 @@ void FRHICommandList::RayTraceDispatch(
 	GlobalResourceBindings.Reset();
 }
 
-void FRHICommandList::RayTraceDispatchIndirect(
+void FRHIComputeCommandList::RayTraceDispatchIndirect(
 	FRayTracingPipelineState* Pipeline,
 	FRHIRayTracingShader* RayGenShader,
 	FRHIShaderBindingTable* SBT,
diff --git a/Engine/Source/Runtime/RHI/Public/RHICommandList.h b/Engine/Source/Runtime/RHI/Public/RHICommandList.h
--- a/Engine/Source/Runtime/RHI/Public/RHICommandList.h
+++ b/Engine/Source/Runtime/RHI/Public/RHICommandList.h
@@ -3609,6 +3609,74 @@
 		}
 #endif // WITH_MGPU
 	}
+	
+	UE_DEPRECATED(5.5, "Provide FRHIShaderBindingTable instead of FRayTracingScene.")
+	FORCEINLINE_DEBUGGABLE void RayTraceDispatch(FRayTracingPipelineState* Pipeline, FRHIRayTracingShader* RayGenShader, FRHIRayTracingScene* Scene, const FRayTracingShaderBindings& GlobalResourceBindings, uint32 Width, uint32 Height)
+	{
+		if (Bypass())
+		{
+			GetComputeContext().RHIRayTraceDispatch(GetRHIRayTracingPipelineState(Pipeline), RayGenShader, Scene, GlobalResourceBindings, Width, Height);
+		}
+		else
+		{
+			ALLOC_COMMAND(FRHICommandRayTraceDispatch)(Pipeline, RayGenShader, Scene, GlobalResourceBindings, Width, Height);
+		}
+	}
 
+	FORCEINLINE_DEBUGGABLE void RayTraceDispatch(FRayTracingPipelineState* Pipeline, FRHIRayTracingShader* RayGenShader, FRHIShaderBindingTable* SBT, const FRayTracingShaderBindings& GlobalResourceBindings, uint32 Width, uint32 Height)
+	{
+		check(SBT != nullptr);
+		if (Bypass())
+		{
+			GetComputeContext().RHIRayTraceDispatch(GetRHIRayTracingPipelineState(Pipeline), RayGenShader, SBT, GlobalResourceBindings, Width, Height);
+		}
+		else
+		{
+			ALLOC_COMMAND(FRHICommandRayTraceDispatch)(Pipeline, RayGenShader, SBT, GlobalResourceBindings, Width, Height);
+		}
+	}
+
+	/*
+	* Compatibility adaptor that operates on the new FRHIBatchedShaderParameters instead of legacy FRayTracingShaderBindings (planned for deprecation).
+	* This will become the default native code path in a future UE version.
+	*/
+	RHI_API void RayTraceDispatch(FRayTracingPipelineState* Pipeline, FRHIRayTracingShader* RayGenShader, FRHIShaderBindingTable* SBT, FRHIBatchedShaderParameters& GlobalResourceBindings, uint32 Width, uint32 Height);
+
+	/**
+	 * Trace rays using dimensions from a GPU buffer containing uint[3], interpreted as number of rays in X, Y and Z dimensions.
+	 * ArgumentBuffer must be in IndirectArgs|SRVCompute state.
+	 */
+	UE_DEPRECATED(5.5, "Provide FRHIShaderBindingTable instead of FRayTracingScene.")
+	FORCEINLINE_DEBUGGABLE void RayTraceDispatchIndirect(FRayTracingPipelineState* Pipeline, FRHIRayTracingShader* RayGenShader, FRHIRayTracingScene* Scene, const FRayTracingShaderBindings& GlobalResourceBindings, FRHIBuffer* ArgumentBuffer, uint32 ArgumentOffset)
+	{
+		if (Bypass())
+		{
+			GetComputeContext().RHIRayTraceDispatchIndirect(GetRHIRayTracingPipelineState(Pipeline), RayGenShader, Scene, GlobalResourceBindings, ArgumentBuffer, ArgumentOffset);
+		}
+		else
+		{
+			ALLOC_COMMAND(FRHICommandRayTraceDispatch)(Pipeline, RayGenShader, Scene, GlobalResourceBindings, ArgumentBuffer, ArgumentOffset);
+		}
+	}
+
+	FORCEINLINE_DEBUGGABLE void RayTraceDispatchIndirect(FRayTracingPipelineState* Pipeline, FRHIRayTracingShader* RayGenShader, FRHIShaderBindingTable* SBT, const FRayTracingShaderBindings& GlobalResourceBindings, FRHIBuffer* ArgumentBuffer, uint32 ArgumentOffset)
+	{
+		check(SBT != nullptr);
+		if (Bypass())
+		{
+			GetComputeContext().RHIRayTraceDispatchIndirect(GetRHIRayTracingPipelineState(Pipeline), RayGenShader, SBT, GlobalResourceBindings, ArgumentBuffer, ArgumentOffset);
+		}
+		else
+		{
+			ALLOC_COMMAND(FRHICommandRayTraceDispatch)(Pipeline, RayGenShader, SBT, GlobalResourceBindings, ArgumentBuffer, ArgumentOffset);
+		}
+	}
+
+	/*
+	* Compatibility adaptor that operates on the new FRHIBatchedShaderParameters instead of legacy FRayTracingShaderBindings (planned for deprecation).
+	* This will become the default native code path in a future UE version.
+	*/
+	RHI_API void RayTraceDispatchIndirect(FRayTracingPipelineState* Pipeline, FRHIRayTracingShader* RayGenShader, FRHIShaderBindingTable* SBT, FRHIBatchedShaderParameters& GlobalResourceBindings, FRHIBuffer* ArgumentBuffer, uint32 ArgumentOffset);
+
 	RHI_API void BuildAccelerationStructure(FRHIRayTracingGeometry* Geometry);
 	RHI_API void BuildAccelerationStructures(TConstArrayView<FRayTracingGeometryBuildParams> Params);
@@ -4305,75 +4373,7 @@
 			RHIThreadFence(true);
 		}
 	}
-
-	UE_DEPRECATED(5.5, "Provide FRHIShaderBindingTable instead of FRayTracingScene.")
-	FORCEINLINE_DEBUGGABLE void RayTraceDispatch(FRayTracingPipelineState* Pipeline, FRHIRayTracingShader* RayGenShader, FRHIRayTracingScene* Scene, const FRayTracingShaderBindings& GlobalResourceBindings, uint32 Width, uint32 Height)
-	{
-		if (Bypass())
-		{
-			GetContext().RHIRayTraceDispatch(GetRHIRayTracingPipelineState(Pipeline), RayGenShader, Scene, GlobalResourceBindings, Width, Height);
-		}
-		else
-		{
-			ALLOC_COMMAND(FRHICommandRayTraceDispatch)(Pipeline, RayGenShader, Scene, GlobalResourceBindings, Width, Height);
-		}
-	}
-
-	FORCEINLINE_DEBUGGABLE void RayTraceDispatch(FRayTracingPipelineState* Pipeline, FRHIRayTracingShader* RayGenShader, FRHIShaderBindingTable* SBT, const FRayTracingShaderBindings& GlobalResourceBindings, uint32 Width, uint32 Height)
-	{
-		check(SBT != nullptr);
-		if (Bypass())
-		{
-			GetContext().RHIRayTraceDispatch(GetRHIRayTracingPipelineState(Pipeline), RayGenShader, SBT, GlobalResourceBindings, Width, Height);
-		}
-		else
-		{
-			ALLOC_COMMAND(FRHICommandRayTraceDispatch)(Pipeline, RayGenShader, SBT, GlobalResourceBindings, Width, Height);
-		}
-	}
-
-	/*
-	* Compatibility adaptor that operates on the new FRHIBatchedShaderParameters instead of legacy FRayTracingShaderBindings (planned for deprecation).
-	* This will become the default native code path in a future UE version.
-	*/
-	RHI_API void RayTraceDispatch(FRayTracingPipelineState* Pipeline, FRHIRayTracingShader* RayGenShader, FRHIShaderBindingTable* SBT, FRHIBatchedShaderParameters& GlobalResourceBindings, uint32 Width, uint32 Height);
-
-	/**
-	 * Trace rays using dimensions from a GPU buffer containing uint[3], interpreted as number of rays in X, Y and Z dimensions.
-	 * ArgumentBuffer must be in IndirectArgs|SRVCompute state.
-	 */
-	UE_DEPRECATED(5.5, "Provide FRHIShaderBindingTable instead of FRayTracingScene.")
-	FORCEINLINE_DEBUGGABLE void RayTraceDispatchIndirect(FRayTracingPipelineState* Pipeline, FRHIRayTracingShader* RayGenShader, FRHIRayTracingScene* Scene, const FRayTracingShaderBindings& GlobalResourceBindings, FRHIBuffer* ArgumentBuffer, uint32 ArgumentOffset)
-	{
-		if (Bypass())
-		{
-			GetComputeContext().RHIRayTraceDispatchIndirect(GetRHIRayTracingPipelineState(Pipeline), RayGenShader, Scene, GlobalResourceBindings, ArgumentBuffer, ArgumentOffset);
-		}
-		else
-		{
-			ALLOC_COMMAND(FRHICommandRayTraceDispatch)(Pipeline, RayGenShader, Scene, GlobalResourceBindings, ArgumentBuffer, ArgumentOffset);
-		}
-	}
-
-	FORCEINLINE_DEBUGGABLE void RayTraceDispatchIndirect(FRayTracingPipelineState* Pipeline, FRHIRayTracingShader* RayGenShader, FRHIShaderBindingTable* SBT, const FRayTracingShaderBindings& GlobalResourceBindings, FRHIBuffer* ArgumentBuffer, uint32 ArgumentOffset)
-	{
-		check(SBT != nullptr);
-		if (Bypass())
-		{
-			GetComputeContext().RHIRayTraceDispatchIndirect(GetRHIRayTracingPipelineState(Pipeline), RayGenShader, SBT, GlobalResourceBindings, ArgumentBuffer, ArgumentOffset);
-		}
-		else
-		{
-			ALLOC_COMMAND(FRHICommandRayTraceDispatch)(Pipeline, RayGenShader, SBT, GlobalResourceBindings, ArgumentBuffer, ArgumentOffset);
-		}
-	}
-
-	/*
-	* Compatibility adaptor that operates on the new FRHIBatchedShaderParameters instead of legacy FRayTracingShaderBindings (planned for deprecation).
-	* This will become the default native code path in a future UE version.
-	*/
-	RHI_API void RayTraceDispatchIndirect(FRayTracingPipelineState* Pipeline, FRHIRayTracingShader* RayGenShader, FRHIShaderBindingTable* SBT, FRHIBatchedShaderParameters& GlobalResourceBindings, FRHIBuffer* ArgumentBuffer, uint32 ArgumentOffset);
-
+	
 	UE_DEPRECATED(5.5, "Use FRHIShaderBindingTable instead.")
 	FORCEINLINE_DEBUGGABLE void SetRayTracingBindings(
 		FRHIRayTracingScene* Scene, FRayTracingPipelineState* Pipeline,
