New GPU profiler - Split "stat gpu" stats into separate stat categories, one per GPU queue
- Stats are added automatically whenever a GPU stat is encountered on a GPU queue for the first time.
- Added FDynamicStat to StatsSystemTypes.h to achieve this, as only compile-time strings were possible before, and we need to generate a unique ANSI stat name per queue.
- Stat groups are named "GPU[gpu_index][queue_type][queue_index]. e.g. "stat GPU0Graphics0", "stat GPU0Compute0". Multiple queues are supported automatically.

Removed RHI_TEMP_USE_GPU_TRACE code.
- Temporary Insights support before the new API was added. We no longer need this.

Fixed missing semicolons and incorrect GPU stat macro use in various renderer files.

diff --git a/Engine/Plugins/FX/Niagara/Source/Niagara/Private/NiagaraGpuComputeDispatch.cpp b/Engine/Plugins/FX/Niagara/Source/Niagara/Private/NiagaraGpuComputeDispatch.cpp
--- a/Engine/Plugins/FX/Niagara/Source/Niagara/Private/NiagaraGpuComputeDispatch.cpp
+++ b/Engine/Plugins/FX/Niagara/Source/Niagara/Private/NiagaraGpuComputeDispatch.cpp
@@ -42,7 +42,6 @@
 DECLARE_CYCLE_STAT(TEXT("GPU Dispatch Setup [RT]"), STAT_NiagaraGPUDispatchSetup_RT, STATGROUP_Niagara);
 DECLARE_CYCLE_STAT(TEXT("GPU Emitter Dispatch [RT]"), STAT_NiagaraGPUSimTick_RT, STATGROUP_Niagara);
 DECLARE_CYCLE_STAT(TEXT("GPU Data Readback [RT]"), STAT_NiagaraGPUReadback_RT, STATGROUP_Niagara);
-DECLARE_FLOAT_COUNTER_STAT(TEXT("Niagara GPU Sim"), STAT_GPU_NiagaraSim, STATGROUP_GPU);
 DECLARE_DWORD_COUNTER_STAT(TEXT("# GPU Particles"), STAT_NiagaraGPUParticles, STATGROUP_Niagara);
 DECLARE_DWORD_COUNTER_STAT(TEXT("# GPU Sorted Particles"), STAT_NiagaraGPUSortedParticles, STATGROUP_Niagara);
 DECLARE_DWORD_COUNTER_STAT(TEXT("# GPU Sorted Buffers"), STAT_NiagaraGPUSortedBuffers, STATGROUP_Niagara);
diff --git a/Engine/Source/Runtime/Core/Public/Stats/GlobalStats.inl b/Engine/Source/Runtime/Core/Public/Stats/GlobalStats.inl
--- a/Engine/Source/Runtime/Core/Public/Stats/GlobalStats.inl
+++ b/Engine/Source/Runtime/Core/Public/Stats/GlobalStats.inl
@@ -21,7 +21,6 @@ DECLARE_STATS_GROUP(TEXT("DDC"), STATGROUP_DDC, STATCAT_Advanced);
 DECLARE_STATS_GROUP(TEXT("Default Stat Group"), STATGROUP_Default, STATCAT_Advanced);
 DECLARE_STATS_GROUP(TEXT("Engine"), STATGROUP_Engine, STATCAT_Advanced);
 DECLARE_STATS_GROUP(TEXT("FPS Chart"), STATGROUP_FPSChart, STATCAT_Advanced);
-DECLARE_STATS_GROUP(TEXT("GPU"), STATGROUP_GPU, STATCAT_Advanced);
 DECLARE_STATS_GROUP(TEXT("GPU Particles"), STATGROUP_GPUParticles, STATCAT_Advanced);
 DECLARE_STATS_GROUP(TEXT("Game"), STATGROUP_Game, STATCAT_Advanced);
 DECLARE_STATS_GROUP(TEXT("GPU Defrag"), STATGROUP_GPUDEFRAG, STATCAT_Advanced);
diff --git a/Engine/Source/Runtime/Core/Public/Stats/StatsSystemTypes.h b/Engine/Source/Runtime/Core/Public/Stats/StatsSystemTypes.h
--- a/Engine/Source/Runtime/Core/Public/Stats/StatsSystemTypes.h
+++ b/Engine/Source/Runtime/Core/Public/Stats/StatsSystemTypes.h
@@ -1782,10 +1782,65 @@ struct FThreadSafeStaticStatInner : public FThreadSafeStaticStatBase
 		const TStatIdData* LocalHighPerformanceEnable = HighPerformanceEnable.Load(EMemoryOrder::Relaxed);
 		if (!LocalHighPerformanceEnable)
 		{
-			LocalHighPerformanceEnable = DoSetup(TStatData::GetStatName(), TStatData::GetDescription(), TStatData::TGroup::GetGroupName(), TStatData::TGroup::GetGroupCategory(), TStatData::TGroup::GetDescription(), TStatData::TGroup::IsDefaultEnabled(), TStatData::IsClearEveryFrame(), TStatData::GetStatType(), TStatData::IsCycleStat(), TStatData::TGroup::GetSortByName(), TStatData::GetMemoryRegion());
+			LocalHighPerformanceEnable = DoSetup(
+				TStatData::GetStatName(),
+				TStatData::GetDescription(),
+				TStatData::TGroup::GetGroupName(),
+				TStatData::TGroup::GetGroupCategory(),
+				TStatData::TGroup::GetDescription(),
+				TStatData::TGroup::IsDefaultEnabled(),
+				TStatData::IsClearEveryFrame(),
+				TStatData::GetStatType(),
+				TStatData::IsCycleStat(),
+				TStatData::TGroup::GetSortByName(),
+				TStatData::GetMemoryRegion()
+			);
 		}
 		return TStatId(LocalHighPerformanceEnable);
 	}
+
+	FORCEINLINE FName GetStatFName() const
+	{
+		return GetStatId().GetName();
+	}
+};
+
+struct FDynamicStat : public FThreadSafeStaticStatBase
+{
+	FDynamicStat(
+		const char* InStatName,
+		const TCHAR* InStatDesc,
+		const char* InGroupName,
+		const char* InGroupCategory,
+		const TCHAR* InGroupDesc,
+		bool bDefaultEnable,
+		bool bShouldClearEveryFrame,
+		EStatDataType::Type InStatType,
+		bool bCycleStat,
+		bool bSortByName,
+		FPlatformMemory::EMemoryCounterRegion InMemoryRegion
+	)
+	{
+		HighPerformanceEnable = DoSetup(
+			InStatName,
+			InStatDesc,
+			InGroupName,
+			InGroupCategory,
+			InGroupDesc,
+			bDefaultEnable,
+			bShouldClearEveryFrame,
+			InStatType,
+			bCycleStat,
+			bSortByName,
+			InMemoryRegion
+		);
+	}
+
+	FORCEINLINE_STATS TStatId GetStatId() const
+	{
+		return TStatId(HighPerformanceEnable.Load(EMemoryOrder::Relaxed));
+	}
+
 	FORCEINLINE FName GetStatFName() const
 	{
 		return GetStatId().GetName();
diff --git a/Engine/Source/Runtime/Engine/Private/KismetRenderingLibrary.cpp b/Engine/Source/Runtime/Engine/Private/KismetRenderingLibrary.cpp
--- a/Engine/Source/Runtime/Engine/Private/KismetRenderingLibrary.cpp
+++ b/Engine/Source/Runtime/Engine/Private/KismetRenderingLibrary.cpp
@@ -797,7 +797,7 @@ void UKismetRenderingLibrary::BeginDrawCanvasToRenderTarget(UObject* WorldContex
 				](FRHICommandListImmediate& RHICmdList)
 			{
 #if  WANTS_DRAW_MESH_EVENTS
-				Breadcrumb->Emplace(RHICmdList, FRHIBreadcrumbData(__FILE__, __LINE__, TStatId(), NAME_None), TEXT("DrawCanvasToTarget: %s"), Name);
+				Breadcrumb->Emplace(RHICmdList, FRHIBreadcrumbData(__FILE__, __LINE__), TEXT("DrawCanvasToTarget: %s"), Name);
 #endif
 				RenderTargetResource->FlushDeferredResourceUpdate(RHICmdList);
 			});
diff --git a/Engine/Source/Runtime/Launch/Private/LaunchEngineLoop.cpp b/Engine/Source/Runtime/Launch/Private/LaunchEngineLoop.cpp
--- a/Engine/Source/Runtime/Launch/Private/LaunchEngineLoop.cpp
+++ b/Engine/Source/Runtime/Launch/Private/LaunchEngineLoop.cpp
@@ -5526,12 +5526,12 @@ static inline void BeginFrameRenderThread(FRHICommandListImmediate& RHICmdList,
 #if CSV_PROFILER
 	if (FCsvProfiler::Get()->IsCapturing_Renderthread())
 	{
-		GRHIFrameBreadcrumb.Emplace(RHICmdList, FRHIBreadcrumbData(__FILE__, __LINE__, TStatId(), NAME_None), TEXT("CsvFrame %d"), FCsvProfiler::Get()->GetCaptureFrameNumberRT());
+		GRHIFrameBreadcrumb.Emplace(RHICmdList, FRHIBreadcrumbData(__FILE__, __LINE__), TEXT("CsvFrame %d"), FCsvProfiler::Get()->GetCaptureFrameNumberRT());
 	}
 	else
 #endif
 	{
-		GRHIFrameBreadcrumb.Emplace(RHICmdList, FRHIBreadcrumbData(__FILE__, __LINE__, TStatId(), NAME_None), TEXT("Frame %d"), CurrentFrameCounter);
+		GRHIFrameBreadcrumb.Emplace(RHICmdList, FRHIBreadcrumbData(__FILE__, __LINE__), TEXT("Frame %d"), CurrentFrameCounter);
 	}
 #endif
 
diff --git a/Engine/Source/Runtime/RHI/Private/GPUProfiler.cpp b/Engine/Source/Runtime/RHI/Private/GPUProfiler.cpp
--- a/Engine/Source/Runtime/RHI/Private/GPUProfiler.cpp
+++ b/Engine/Source/Runtime/RHI/Private/GPUProfiler.cpp
@@ -10,6 +10,7 @@
 #include "Misc/CommandLine.h"
 #include "RHI.h"
 #include "GpuProfilerTrace.h"
+#include "Containers/AnsiString.h"
 
 #if !UE_BUILD_SHIPPING
 #include "VisualizerEvents.h"
@@ -17,33 +18,6 @@
 #include "Modules/ModuleManager.h"
 #endif
 
-#if HAS_GPU_STATS
-CSV_DEFINE_CATEGORY_MODULE(RHI_API, GPU, true);
-DECLARE_FLOAT_COUNTER_STAT(TEXT("[TOTAL]"), Stat_GPU_Total, STATGROUP_GPU);
-CSV_DEFINE_STAT(GPU, Total);
-#endif
-
-// Temporary function to resolve link issues with the above "Total" GPU stat moving from RenderCore to RHI.
-// We can remove this once the old GPU profiler code has been deleted, and RHI_NEW_GPU_PROFILER is set to 1 permanently.
-extern RHI_API void RHISetGPUStatTotals(bool bCsvStatsEnabled, double TotalMs)
-{
-#if HAS_GPU_STATS
-
-#if STATS
-	FThreadStats::AddMessage(GET_STATFNAME(Stat_GPU_Total), EStatOperation::Set, TotalMs);
-	TRACE_STAT_SET(GET_STATFNAME(Stat_GPU_Total), TotalMs);
-#endif
-
-#if CSV_PROFILER_STATS
-	if (bCsvStatsEnabled)
-	{
-		FCsvProfiler::Get()->RecordCustomStat(CSV_STAT_FNAME(Total), CSV_CATEGORY_INDEX(GPU), TotalMs, ECsvCustomStatOp::Set);
-	}
-#endif
-
-#endif // HAS_GPU_STATS
-}
-
 #define LOCTEXT_NAMESPACE "GpuProfiler"
 
 enum class EGPUProfileSortMode
@@ -745,16 +719,87 @@ bool FGPUTiming::GAreGlobalsInitialized = false;
 
 #else
 
-// Temporary. Adds Insights markers for the 0th GPU graphics queue
-// until we have a new API that is capable of displaying more info.
-#define RHI_TEMP_USE_GPU_TRACE (1 && GPUPROFILERTRACE_ENABLED)
-
-// When enabled, and running with a single GPU, repurposes the "GPU2" track
-// in Insights to show the single GPU's async compute queue.
-#define RHI_TEMP_USE_TRACK2_FOR_COMPUTE (1 && RHI_TEMP_USE_GPU_TRACE)
-
 namespace UE::RHI::GPUProfiler
 {
+#if HAS_GPU_STATS
+	// Total GPU time per queue.
+	static FGPUStat GPUStat_Total(TEXT("[TOTAL]"));
+#endif
+
+	TMap<FQueue, TUniquePtr<FGPUStat::FStatCategory>> FGPUStat::Categories;
+
+	FGPUStat::FStatCategory& FGPUStat::GetCategory(FQueue Queue)
+	{
+		TUniquePtr<FStatCategory>& Category = Categories.FindOrAdd(Queue);
+		if (!Category)
+		{
+			Category = MakeUnique<FStatCategory>(Queue);
+		}
+		return *Category;
+	}
+
+	FGPUStat::FStatCategory_Stats::FStatCategory_Stats(FQueue Queue)
+	#if STATS
+		: GroupName(FString::Printf(TEXT("STATGROUP_GPU%d%s%d"), Queue.GPU, Queue.GetTypeString(), Queue.Index))
+		, GroupDesc(FString::Printf(TEXT("GPU %d %s Queue %d Timing"), Queue.GPU, Queue.GetTypeString(), Queue.Index))
+	#endif
+	{}
+
+	FGPUStat::FStatCategory_Csv::FStatCategory_Csv(FQueue Queue)
+	#if CSV_PROFILER_STATS
+		: Csv(*FString::Printf(TEXT("GPU%d%s%d"), Queue.GPU, Queue.GetTypeString(), Queue.Index), true)
+	#endif
+	{}
+
+	FGPUStat::FStatCategory::FStatCategory(FQueue Queue)
+		: FStatCategory_Stats(Queue)
+		, FStatCategory_Csv(Queue)
+	{}
+
+#if STATS
+	TStatId FGPUStat::GetStatId(FQueue Queue)
+	{
+		FStatInstance& Instance = Instances.FindOrAdd(Queue);
+
+		if (!Instance.Stat)
+		{
+			FStatCategory& Category = GetCategory(Queue);
+
+			Instance.StatName = MakeUnique<FAnsiString>(FString::Printf(TEXT("STAT_GPU%d%s%d_%s"), Queue.GPU, Queue.GetTypeString(), Queue.Index, Name));
+			Instance.Stat = MakeUnique<FDynamicStat>(
+				**Instance.StatName,
+				Name,
+				*Category.GroupName,
+				"STATCAT_GPU",
+				*Category.GroupDesc,
+				true, // IsDefaultEnabled
+				true, // IsClearEveryFrame
+				EStatDataType::ST_double,
+				false, // IsCycleStat
+				false, // SortByName
+				FPlatformMemory::MCR_Invalid
+			);
+		}
+
+		return Instance.Stat->GetStatId();
+	}
+#endif
+
+#if CSV_PROFILER_STATS
+	FCsvDeclaredStat& FGPUStat::GetCsvStat(FQueue Queue)
+	{
+		FStatInstance& Instance = Instances.FindOrAdd(Queue);
+
+		if (!Instance.CsvStat.IsSet())
+		{
+			FStatCategory& Category = GetCategory(Queue);
+			Instance.CsvStat.Emplace(Name, Category.Csv.Index);
+		}
+
+		return *Instance.CsvStat;
+	}
+#endif
+
 	TLockFreePointerListUnordered<void, PLATFORM_CACHE_LINE_SIZE> FEventStream::FChunk::MemoryPool;
 
 	static TArray<FEventSink*>& GetSinks()
@@ -904,12 +949,44 @@ namespace UE::RHI::GPUProfiler
 			}
 		};
 
+		struct FStatState
+		{
+			uint64 BusyCycles = 0;
+			uint64 LastCycles = 0;
+
+			FStatState() = default;
+			FStatState(FStatState const&) = default;
+
+			FStatState(FStatState&& Other)
+				: FStatState(Other)
+			{
+				Other.BusyCycles = 0;
+				Other.LastCycles = 0;
+			}
+
+			void BeginWork(uint64 Timestamp)
+			{
+				LastCycles = FMath::Max(LastCycles, Timestamp);
+			}
+
+			void EndWork(uint64 Timestamp)
+			{
+				uint64 Delta = Timestamp > LastCycles
+					? Timestamp - LastCycles
+					: 0;
+
+				LastCycles = FMath::Max(LastCycles, Timestamp);
+				BusyCycles += Delta;
+			}
+		};
+
 		struct FQueueTimestamps
 		{
 			FTimestampStream Queue;
+			FStatState WholeQueueStat;
 
 		#if WITH_RHI_BREADCRUMBS
-			TMap<FRHIBreadcrumbData_Stats, FTimestampStream> Stats;
+			TMap<FRHIBreadcrumbData_Stats, FStatState> Stats;
 		#endif
 		};
 
@@ -935,255 +1012,17 @@ namespace UE::RHI::GPUProfiler
 		TMap<FQueue, FQueueState> QueueStates;
 		TMap<uint32, FFrameState> Frames;
 
-	#if RHI_TEMP_USE_GPU_TRACE
-		struct FInsightsTrack
-		{
-			uint32 const Index;
-			uint64 MaxTraceTime = 0;
-			uint32 FrameNumber = 0;
-
-		#if DO_CHECK
-			int32 EventCounter = 0;
-		#endif
-
-			FInsightsTrack(uint32 Index)
-				: Index(Index)
-			{}
-
-			uint64 GPUToTrace(uint64 GPUTimestamp)
-			{
-				uint64 TraceTime = (uint64)(FPlatformTime::ToMilliseconds64(GPUTimestamp) * 1000.0);
-
-				//
-				// Some platforms support top-of-pipe timestamps, meaning BeginWork/BeginBreadcrumb events
-				// that occur logically after EndWork/EndBreadcrumb events in the command stream can have
-				// a timestamp that is earlier than the subsequent begin event due the GPU workload overlap.
-				//
-				// The old Insights API cannot support this, and simply doesn't display the events if their
-				// timestamps aren't strictly sequential. Work around this by emitting the Max() of the current
-				// timestamp, and the largest timestamp we've seen before.
-				//
-
-				MaxTraceTime = FMath::Max(TraceTime, MaxTraceTime);
-				return MaxTraceTime;
-			}
-
-			bool bShowWork = false;
-			bool bEmittedGPUWorkName = false;
-
-			bool bNeedsEnd = false;
-			uint64 MaxEndTimeBOP = 0;
-			TOptional<uint64> LastBeginTimestampTOP;
-
-			// Emitting FEndWork events to Insights are deferred until we know there isn't an overlapping FBeginWork event
-			// that would otherwise prevent the GPU going idle. This is done to coalesce markers to make them less noisy.
-			bool EmitEndWork(FQueueState const& QueueState)
-			{
-				bool bEmitEnd = bNeedsEnd && (!LastBeginTimestampTOP.IsSet() || LastBeginTimestampTOP.GetValue() > MaxEndTimeBOP);
-				bool bNeedsBegin = !bNeedsEnd || bEmitEnd;
-
-				if (bEmitEnd)
-				{
-					uint64 TraceTime = GPUToTrace(MaxEndTimeBOP);
-
-				#if WITH_RHI_BREADCRUMBS
-					for (FRHIBreadcrumbNode* Current = QueueState.Breadcrumb; Current; Current = Current->GetParent())
-					{
-						check(EventCounter-- > 0);
-						FGpuProfilerTrace::EndEvent(TraceTime, Index);
-					}
-				#endif
-
-					if (bShowWork)
-					{
-						check(EventCounter-- > 0);
-						FGpuProfilerTrace::EndEvent(TraceTime, Index); // GPUWork event
-					}
-
-					bNeedsEnd = false;
-				}
-
-				return bNeedsBegin;
-			}
-
-			void EmitBeginWork(FQueueState const& QueueState)
-			{
-				if (LastBeginTimestampTOP.IsSet())
-				{
-					uint64 TraceTime = GPUToTrace(LastBeginTimestampTOP.GetValue());
-							
-					if (bShowWork)
-					{
-						static FName GraphicsWorkName("Graphics Work");
-						static FName ComputeWorkName("Compute Work");
-						FName WorkName = QueueState.Type == FQueue::EType::Graphics
-							? GraphicsWorkName
-							: ComputeWorkName;
-
-						if (!bEmittedGPUWorkName)
-						{
-							FGpuProfilerTrace::SpecifyEventByName(WorkName, Index);
-						}
-
-						FGpuProfilerTrace::BeginEventByName(WorkName, FrameNumber, TraceTime, Index);
-						check(++EventCounter);
-					}
-
-				#if WITH_RHI_BREADCRUMBS
-					{
-						FRHIBreadcrumb::FBuffer Buffer;
-						auto Recurse = [&](auto& Recurse, FRHIBreadcrumbNode* Current) -> void
-						{
-							if (!Current)
-								return;
-
-							Recurse(Recurse, Current->GetParent());
-
-							FName Name(Current->Name.GetTCHAR(Buffer));
-							FGpuProfilerTrace::BeginEventByName(Name, FrameNumber, TraceTime, Index);
-							check(++EventCounter);
-						};
-						Recurse(Recurse, QueueState.Breadcrumb);
-					}
-				#endif
-
-					LastBeginTimestampTOP.Reset();
-				}
-			}
-
-			void BeginWork(FQueueState const& QueueState, FEvent::FBeginWork const& BeginWork)
-			{
-				LastBeginTimestampTOP = BeginWork.GPUTimestampTOP;
-
-				if (bShowWork)
-				{
-					if (EmitEndWork(QueueState))
-					{
-						EmitBeginWork(QueueState);
-					}
-				}
-			}
-
-			void EndWork(FQueueState const& QueueState, FEvent::FEndWork const& EndWork)
-			{
-				MaxEndTimeBOP = FMath::Max(MaxEndTimeBOP, EndWork.GPUTimestampBOP);
-				bNeedsEnd = true;
-			}
-
-			void BeginBreadcrumb(FQueueState const& QueueState, FEvent::FBeginBreadcrumb const& BeginBreadcrumb)
-			{
-				FRHIBreadcrumb::FBuffer Buffer;
-				TCHAR const* Str = BeginBreadcrumb.Breadcrumb->Name.GetTCHAR(Buffer);
-				FName Name(Str);
-
-				FGpuProfilerTrace::SpecifyEventByName(Name, Index);
-				FGpuProfilerTrace::BeginEventByName(Name, FrameNumber, GPUToTrace(BeginBreadcrumb.GPUTimestampTOP), Index);
-				check(++EventCounter);
-			}
-
-			void EndBreadcrumb(FQueueState const& QueueState, FEvent::FEndBreadcrumb const& EndBreadcrumb)
-			{
-				check(EventCounter-- > 0);
-				FGpuProfilerTrace::EndEvent(GPUToTrace(EndBreadcrumb.GPUTimestampBOP), Index);
-			}
-
-			void FrameBoundary(FQueueState const& QueueState, FEvent::FFrameBoundary const& FrameBoundary)
-			{
-				// End the current Insights GPU frame + start the next one
-
-				// All breadcrumbs must be ended before the frame boundary can be emitted.
-				LastBeginTimestampTOP.Reset();
-				bool bNeedsBegin = EmitEndWork(QueueState);
-
-				check(FrameNumber == FrameBoundary.FrameNumber);
-				check(EventCounter == 0);
-				FGpuProfilerTrace::EndFrame(Index);
-
-				FrameNumber++;
-
-				// Use 1,1 calibration to disable any adjustments Insights makes.
-				// The timestamps we use in the GPU event stream are already in the CPU clock domain.
-				FGPUTimingCalibrationTimestamp Calibration{ 1, 1 };
-				FGpuProfilerTrace::BeginFrame(Calibration, Index);
-
-				if (!bShowWork)
-				{
-					if (bNeedsBegin)
-					{
-						LastBeginTimestampTOP = MaxEndTimeBOP;
-						EmitBeginWork(QueueState);
-					}
-				}
-			}
-
-			void Initialize()
-			{
-				// When enabled, adds "GPUWork" markers to the GPU Insights trace to show where the GPU is busy or idle.
-				// Causes breadcrumbs to be pushed / popped multiple times, breaking them up on the timeline.
-				bShowWork = FParse::Param(FCommandLine::Get(), TEXT("tracegpuwork"));
-
-				// Use 1,1 calibration to disable any adjustments Insights makes.
-				// The timestamps we use in the GPU event stream are already in the CPU clock domain.
-				FGPUTimingCalibrationTimestamp Calibration{ 1, 1 };
-				check(EventCounter == 0);
-				FGpuProfilerTrace::BeginFrame(Calibration, Index);
-			}
-
-		} InsightsTracks[2] { {0}, {1} };
-
-		FInsightsTrack* GetInsightsTrack(FQueue const& Queue)
-		{
-			if (GNumExplicitGPUsForRendering > 1)
-			{
-				// MGPU Mode - GPU0 Graphics + GPU1 Graphics
-				if (Queue.Type == FQueue::EType::Graphics && Queue.Index == 0 && Queue.GPU < 2)
-				{
-					return &InsightsTracks[Queue.GPU];
-				}
-			}
-			else
-			{
-				// GPU0 Graphics + GPU0 Compute Mode
-				if (Queue.GPU == 0 && Queue.Index == 0)
-				{
-					switch (Queue.Type)
-					{
-					case FQueue::EType::Graphics: return &InsightsTracks[0];
-				#if RHI_TEMP_USE_TRACK2_FOR_COMPUTE
-					case FQueue::EType::Compute : return &InsightsTracks[1];
-				#endif
-					}
-				}
-			}
-
-			return nullptr;
-		}
-
-	#endif // RHI_TEMP_USE_GPU_TRACE
-
 		void InitializeQueues(TConstArrayView<FQueue> Queues) override
 		{
 			for (FQueue const& Queue : Queues)
 			{
 				check(QueueStates.Find(Queue) == nullptr);
 				QueueStates.Add(Queue, Queue);
-
-			#if RHI_TEMP_USE_GPU_TRACE
-				// Start the first Insights GPU frame
-				FInsightsTrack* Track = GetInsightsTrack(Queue);
-				if (Track)
-				{
-					Track->Initialize();
-				}
-			#endif
 			}
 		}
 
 		void ProcessEvents(FQueue Queue, FEventStream const& EventStream) override
 		{
-		#if RHI_TEMP_USE_GPU_TRACE
-			FInsightsTrack* Track = GetInsightsTrack(Queue);
-		#endif
 			FQueueState& QueueState = QueueStates.FindChecked(Queue);
 
 			for (FEvent const* Event : EventStream)
@@ -1197,19 +1036,13 @@ namespace UE::RHI::GPUProfiler
 
 						FEvent::FBeginWork const& BeginWork = Event->Value.Get<FEvent::FBeginWork>();
 						QueueState.Timestamps.Queue.AddTimestamp(BeginWork.GPUTimestampTOP, true);
-
-					#if RHI_TEMP_USE_GPU_TRACE
-						if (Track)
-						{
-							Track->BeginWork(QueueState, BeginWork);
-						}
-					#endif
+						QueueState.Timestamps.WholeQueueStat.BeginWork(BeginWork.GPUTimestampTOP);
 
 					#if WITH_RHI_BREADCRUMBS
 						// Apply the timestamp to all active stats
 						for (auto const& [Stat, RefCount] : QueueState.ActiveStats)
 						{
-							QueueState.Timestamps.Stats.FindChecked(Stat).AddTimestamp(BeginWork.GPUTimestampTOP, true);
+							QueueState.Timestamps.Stats.FindChecked(Stat).BeginWork(BeginWork.GPUTimestampTOP);
 						}
 					#endif
 					}
@@ -1222,19 +1055,13 @@ namespace UE::RHI::GPUProfiler
 
 						FEvent::FEndWork const& EndWork = Event->Value.Get<FEvent::FEndWork>();
 						QueueState.Timestamps.Queue.AddTimestamp(EndWork.GPUTimestampBOP, false);
+						QueueState.Timestamps.WholeQueueStat.EndWork(EndWork.GPUTimestampBOP);
 
 					#if WITH_RHI_BREADCRUMBS
 						// Apply the timestamp to all active stats
 						for (auto const& [Stat, RefCount] : QueueState.ActiveStats)
 						{
-							QueueState.Timestamps.Stats.FindChecked(Stat).AddTimestamp(EndWork.GPUTimestampBOP, false);
-						}
-					#endif
-
-					#if RHI_TEMP_USE_GPU_TRACE
-						if (Track)
-						{
-							Track->EndWork(QueueState, EndWork);
+							QueueState.Timestamps.Stats.FindChecked(Stat).EndWork(EndWork.GPUTimestampBOP);
 						}
 					#endif
 					}
@@ -1248,13 +1075,6 @@ namespace UE::RHI::GPUProfiler
 						FEvent::FBeginBreadcrumb const& BeginBreadcrumb = Event->Value.Get<FEvent::FBeginBreadcrumb>();
 						FRHIBreadcrumbData_Stats const& Stat = BeginBreadcrumb.Breadcrumb->Name.Data;
 
-					#if RHI_TEMP_USE_GPU_TRACE
-						if (Track)
-						{
-							Track->BeginBreadcrumb(QueueState, BeginBreadcrumb);
-						}
-					#endif
-
 						if (Stat.ShouldComputeStat())
 						{
 							// Disregard the stat if it is nested within itself (i.e. its already in the ActiveStats map with a non-zero ref count).
@@ -1262,7 +1082,7 @@ namespace UE::RHI::GPUProfiler
 							int32 RefCount = QueueState.ActiveStats.FindOrAdd(Stat)++;
 							if (RefCount == 0)
 							{
-								QueueState.Timestamps.Stats.FindOrAdd(Stat).AddTimestamp(BeginBreadcrumb.GPUTimestampTOP, true);
+								QueueState.Timestamps.Stats.FindOrAdd(Stat).BeginWork(BeginBreadcrumb.GPUTimestampTOP);
 							}
 						}
 
@@ -1277,20 +1097,13 @@ namespace UE::RHI::GPUProfiler
 						FEvent::FEndBreadcrumb const& EndBreadcrumb = Event->Value.Get<FEvent::FEndBreadcrumb>();
 						FRHIBreadcrumbData_Stats const& Stat = EndBreadcrumb.Breadcrumb->Name.Data;
 
-					#if RHI_TEMP_USE_GPU_TRACE
-						if (Track)
-						{
-							Track->EndBreadcrumb(QueueState, EndBreadcrumb);
-						}
-					#endif
-
 						if (Stat.ShouldComputeStat())
 						{
 							// Pop the stat when the refcount hits zero.
 							int32 RefCount = --QueueState.ActiveStats.FindChecked(Stat);
 							if (RefCount == 0)
 							{
-								QueueState.Timestamps.Stats.FindChecked(Stat).AddTimestamp(EndBreadcrumb.GPUTimestampBOP, false);
+								QueueState.Timestamps.Stats.FindChecked(Stat).EndWork(EndBreadcrumb.GPUTimestampBOP);
 								QueueState.ActiveStats.FindAndRemoveChecked(Stat);
 							}							
 						}
@@ -1339,13 +1152,6 @@ namespace UE::RHI::GPUProfiler
 
 							Frames.Remove(FrameBoundary.FrameNumber);
 						}
-
-					#if RHI_TEMP_USE_GPU_TRACE
-						if (Track)
-						{
-							Track->FrameBoundary(QueueState, FrameBoundary);
-						}
-					#endif
 					}
 					break;
 				}
@@ -1354,8 +1160,6 @@ namespace UE::RHI::GPUProfiler
 
 		void ProcessFrame(FFrameState& FrameState)
 		{
-			TArray<FTimestampStream::FState, TInlineAllocator<GetRHIPipelineCount() * MAX_NUM_GPUS>> StreamPointers;
-
 		#if CSV_PROFILER_STATS
 			const bool bCsvStatsEnabled = !!CVarGPUCsvStatsEnabled.GetValueOnAnyThread();
 			FCsvProfiler* CsvProfiler = bCsvStatsEnabled ? FCsvProfiler::Get() : nullptr;
@@ -1363,47 +1167,40 @@ namespace UE::RHI::GPUProfiler
 			const bool bCsvStatsEnabled = false;
 		#endif
 
-			// Compute the individual GPU stats
-		#if WITH_RHI_BREADCRUMBS
-
-			TSet<FRHIBreadcrumbData_Stats> UniqueStats;
-			for (auto const& [Queue, State] : FrameState)
+			for (auto const& [Queue, QueueTimestamps] : FrameState)
 			{
-				for (auto const& [Stat, Timestamps] : State.Stats)
+			#if WITH_RHI_BREADCRUMBS
+				// Compute the individual GPU stats
+				for (auto const& [Stat, StatState] : QueueTimestamps.Stats)
 				{
-					UniqueStats.Add(Stat);
-				}
-			}
+					double StatMilliseconds = FPlatformTime::ToMilliseconds64(StatState.BusyCycles);
+					SET_FLOAT_STAT_FName(Stat.GPUStat->GetStatId(Queue).GetName(), StatMilliseconds);
 
-			for (FRHIBreadcrumbData_Stats const& Stat : UniqueStats)
-			{
-				StreamPointers.Reset();
-				for (auto const& [Queue, State] : FrameState)
-				{
-					FTimestampStream const* Stream = State.Stats.Find(Stat);
-					if (Stream)
+				#if CSV_PROFILER_STATS
+					if (CsvProfiler)
 					{
-						StreamPointers.Emplace(*Stream);
+						FCsvDeclaredStat& CsvStat = Stat.GPUStat->GetCsvStat(Queue);
+						CsvProfiler->RecordCustomStat(CsvStat.Name, CsvStat.CategoryIndex, StatMilliseconds, ECsvCustomStatOp::Set);
 					}
+				#endif
 				}
+			#endif // WITH_RHI_BREADCRUMBS
 
-				uint64 Union = FTimestampStream::ComputeUnion(StreamPointers);
-				double Milliseconds = FPlatformTime::ToMilliseconds64(Union);
-
-				SET_FLOAT_STAT_FName(Stat.StatId.GetName(), Milliseconds);
+				// Set the whole-frame per queue stat
+				double QueueMilliseconds = FPlatformTime::ToMilliseconds64(QueueTimestamps.WholeQueueStat.BusyCycles);
+				SET_FLOAT_STAT_FName(GPUStat_Total.GetStatId(Queue).GetName(), QueueMilliseconds);
 
 			#if CSV_PROFILER_STATS
 				if (CsvProfiler)
 				{
-					CsvProfiler->RecordCustomStat(Stat.CsvStat, CSV_CATEGORY_INDEX(GPU), Milliseconds, ECsvCustomStatOp::Set);
+					FCsvDeclaredStat& CsvStat = GPUStat_Total.GetCsvStat(Queue);
+					CsvProfiler->RecordCustomStat(CsvStat.Name, CsvStat.CategoryIndex, QueueMilliseconds, ECsvCustomStatOp::Set);
 				}
 			#endif
 			}
 
-		#endif // WITH_RHI_BREADCRUMBS
-
 			// Compute the whole-frame total GPU time.
-			StreamPointers.Reset();
+			TArray<FTimestampStream::FState, TInlineAllocator<GetRHIPipelineCount() * MAX_NUM_GPUS>> StreamPointers;
 			for (auto const& [Queue, State] : FrameState)
 			{
 				StreamPointers.Emplace(State.Queue);
@@ -1413,7 +1210,8 @@ namespace UE::RHI::GPUProfiler
 			// Update the global GPU frame time stats - need to convert to Cycles32 rather than Cycles64.
 			GGPUFrameTime = FPlatformMath::TruncToInt(FPlatformTime::ToSeconds64(WholeFrameUnion) / FPlatformTime::GetSecondsPerCycle());
 
-			RHISetGPUStatTotals(bCsvStatsEnabled, FPlatformTime::ToMilliseconds64(WholeFrameUnion));
+			// @todo set global csv GPU time
+			//RHISetGPUStatTotals(bCsvStatsEnabled, FPlatformTime::ToMilliseconds64(WholeFrameUnion));
 		}
 
 	} GGPUProfilerSink_StatSystem;
diff --git a/Engine/Source/Runtime/RHI/Private/RHIBreadcrumbs.cpp b/Engine/Source/Runtime/RHI/Private/RHIBreadcrumbs.cpp
--- a/Engine/Source/Runtime/RHI/Private/RHIBreadcrumbs.cpp
+++ b/Engine/Source/Runtime/RHI/Private/RHIBreadcrumbs.cpp
@@ -12,7 +12,7 @@
 FRHIBreadcrumbNode::FRHIBreadcrumbNode()
 	: Name([]() -> FRHIBreadcrumb const&
 	{
-		static TRHIBreadcrumb<> Breadcrumb(FRHIBreadcrumbData(__FILE__, __LINE__, TStatId(), NAME_None), TEXT("Sentinel"));
+		static TRHIBreadcrumb<> Breadcrumb(FRHIBreadcrumbData(__FILE__, __LINE__), TEXT("Sentinel"));
 		return Breadcrumb;
 	}())
 {}
diff --git a/Engine/Source/Runtime/RHI/Public/GPUProfiler.h b/Engine/Source/Runtime/RHI/Public/GPUProfiler.h
--- a/Engine/Source/Runtime/RHI/Public/GPUProfiler.h
+++ b/Engine/Source/Runtime/RHI/Public/GPUProfiler.h
@@ -11,10 +11,9 @@
 #include "ProfilingDebugging/CsvProfiler.h"
 
 #include "RHIBreadcrumbs.h"
+#include "RHIStats.h"
 
-#if HAS_GPU_STATS
-CSV_DECLARE_CATEGORY_MODULE_EXTERN(RHI_API, GPU);
-#endif
+#include "Containers/AnsiString.h"
 
 #if RHI_NEW_GPU_PROFILER
 
@@ -475,6 +474,83 @@ namespace UE::RHI::GPUProfiler
 
 	RHI_API void ProcessEvents(FQueue Queue, FEventStream EventStream);
 	RHI_API void InitializeQueues(TConstArrayView<FQueue> Queues);
+
+	struct FGPUStat
+	{
+	private:
+		struct FStatCategory_Stats
+		{
+		#if STATS
+			FAnsiString const GroupName;
+			FString     const GroupDesc;
+		#endif
+
+			FStatCategory_Stats(FQueue Queue);
+		};
+
+		struct FStatCategory_Csv
+		{
+		#if CSV_PROFILER_STATS
+			FCsvCategory const Csv;
+		#endif
+
+			FStatCategory_Csv(FQueue Queue);
+		};
+
+		struct FStatCategory
+			: public FStatCategory_Stats
+			, public FStatCategory_Csv
+		{
+			FStatCategory(FQueue Queue);
+		};
+
+		struct FStatInstance
+		{
+		#if STATS
+			TUniquePtr<FAnsiString> StatName;
+			TUniquePtr<FDynamicStat> Stat;
+		#endif
+		#if CSV_PROFILER_STATS
+			TOptional<FCsvDeclaredStat> CsvStat;
+		#endif
+		};
+
+		static TMap<FQueue, TUniquePtr<FStatCategory>> Categories;
+		static FStatCategory& GetCategory(FQueue Queue);
+
+		TMap<FQueue, FStatInstance> Instances;
+
+	public:
+		TCHAR const* const Name;
+
+		FGPUStat(TCHAR const* Name)
+			: Name(Name)
+		{}
+
+	#if STATS
+		TStatId GetStatId(FQueue Queue);
+	#endif
+
+	#if CSV_PROFILER_STATS
+		FCsvDeclaredStat& GetCsvStat(FQueue Queue);
+	#endif
+	};
+
+	template <typename TNameProvider>
+	struct TGPUStat : public FGPUStat
+	{
+		TGPUStat()
+			: FGPUStat(TNameProvider::GetName())
+		{}
+	};
+
+	template <typename TNameProvider>
+	struct TGPUStatWithDrawcallCategory : public TGPUStat<TNameProvider>
+	{
+	#if HAS_GPU_STATS
+		FRHIDrawStatsCategory DrawcallCategory;
+	#endif
+	};
 }
 
 #else
diff --git a/Engine/Source/Runtime/RHI/Public/RHIBreadcrumbs.h b/Engine/Source/Runtime/RHI/Public/RHIBreadcrumbs.h
--- a/Engine/Source/Runtime/RHI/Public/RHIBreadcrumbs.h
+++ b/Engine/Source/Runtime/RHI/Public/RHIBreadcrumbs.h
@@ -23,6 +23,13 @@
 // Whether to store the filename and line number of each RHI breadcrumb and emit this data to Insights.
 #define RHI_BREADCRUMBS_EMIT_LOCATION (WITH_RHI_BREADCRUMBS && (CPUPROFILERTRACE_ENABLED || GPUPROFILERTRACE_ENABLED) && 1)
 
+#if RHI_NEW_GPU_PROFILER && HAS_GPU_STATS
+namespace UE::RHI::GPUProfiler
+{
+	struct FGPUStat;
+}
+#endif
+
 #if WITH_RHI_BREADCRUMBS
 
 	//
@@ -49,55 +56,82 @@
 	//
 	struct FRHIBreadcrumbData_Stats
 	{
-#if STATS
+#if RHI_NEW_GPU_PROFILER && HAS_GPU_STATS
+
+		UE::RHI::GPUProfiler::FGPUStat* GPUStat;
+
+		FRHIBreadcrumbData_Stats(UE::RHI::GPUProfiler::FGPUStat* InGPUStat)
+			: GPUStat(InGPUStat)
+		{}
+
+		bool ShouldComputeStat() const
+		{
+			return GPUStat != nullptr;
+		}
+
+		bool operator == (FRHIBreadcrumbData_Stats const& RHS) const
+		{
+			return GPUStat == RHS.GPUStat;
+		}
+
+		friend uint32 GetTypeHash(FRHIBreadcrumbData_Stats const& Stats)
+		{
+			return PointerHash(Stats.GPUStat);
+		}
+
+#else
+
+	#if STATS
 		TStatId StatId {};
-#endif
-#if CSV_PROFILER_STATS
+	#endif
+	#if CSV_PROFILER_STATS
 		FName CsvStat = NAME_None;
-#endif
+	#endif
 
 		FRHIBreadcrumbData_Stats(TStatId InStatId, FName InCsvStat)
 		{
-#if STATS
+	#if STATS
 			StatId = InStatId;
-#endif
-#if CSV_PROFILER_STATS
+	#endif
+	#if CSV_PROFILER_STATS
 			CsvStat = InCsvStat;
-#endif
+	#endif
 		}
 
 		bool ShouldComputeStat() const
 		{
-#if STATS
+	#if STATS
 			return StatId.IsValidStat();
-#elif CSV_PROFILER_STATS
+	#elif CSV_PROFILER_STATS
 			return CsvStat != NAME_None;
-#else
+	#else
 			return false;
-#endif
+	#endif
 		}
 
 		bool operator == (FRHIBreadcrumbData_Stats const& RHS) const
 		{
-#if STATS
+	#if STATS
 			return StatId == RHS.StatId;
-#elif CSV_PROFILER_STATS
+	#elif CSV_PROFILER_STATS
 			return CsvStat == RHS.CsvStat;
-#else
+	#else
 			return true;
-#endif
+	#endif
 		}
 
 		friend uint32 GetTypeHash(FRHIBreadcrumbData_Stats const& Stats)
 		{
-#if STATS
+	#if STATS
 			return GetTypeHash(Stats.StatId);
-#elif CSV_PROFILER_STATS
+	#elif CSV_PROFILER_STATS
 			return GetTypeHash(Stats.CsvStat);
-#else
+	#else
 			return 0;
-#endif
+	#endif
 		}
+
+#endif
 	};
 
 	//
@@ -109,10 +143,21 @@
 		, public FRHIBreadcrumbData_Stats
 	{
 	public:
-		FRHIBreadcrumbData(ANSICHAR const* File, uint32 Line, TStatId StatId, FName CsvStat)
+
+	#if RHI_NEW_GPU_PROFILER && HAS_GPU_STATS
+		FRHIBreadcrumbData(ANSICHAR const* InFile, uint32 InLine, UE::RHI::GPUProfiler::FGPUStat* InGPUStat = nullptr)
+			: FRHIBreadcrumbData_Location(InFile, InLine)
+			, FRHIBreadcrumbData_Stats(InGPUStat)
+		{}
+
+	#else
+
+		FRHIBreadcrumbData(ANSICHAR const* File, uint32 Line, TStatId StatId = TStatId(), FName CsvStat = NAME_None)
 			: FRHIBreadcrumbData_Location(File, Line)
 			, FRHIBreadcrumbData_Stats(StatId, CsvStat)
 		{}
+
+	#endif
 	};
 
 	class FRHIBreadcrumb;
@@ -899,19 +944,25 @@
 		inline ~FRHIBreadcrumbEventScope();
 	};
 
-	#define RHI_BREADCRUMB_EVENT(                 RHICmdList_Or_RHIContext,                  Format, ...) FRHIBreadcrumbEventScope ANONYMOUS_VARIABLE(BreadcrumbEvent)(RHICmdList_Or_RHIContext, FRHIBreadcrumbData(__FILE__, __LINE__, TStatId(), NAME_None), true     , TEXT(Format), ## __VA_ARGS__)
-	#define RHI_BREADCRUMB_EVENT_CONDITIONAL(     RHICmdList_Or_RHIContext,       Condition, Format, ...) FRHIBreadcrumbEventScope ANONYMOUS_VARIABLE(BreadcrumbEvent)(RHICmdList_Or_RHIContext, FRHIBreadcrumbData(__FILE__, __LINE__, TStatId(), NAME_None), Condition, TEXT(Format), ## __VA_ARGS__)
+#if RHI_NEW_GPU_PROFILER
+	#define RHI_GPU_STAT_ARGS(StatName) &GPUStat_##StatName
+#else
+	#define RHI_GPU_STAT_ARGS(StatName) GET_STATID(Stat_GPU_##StatName), CSV_STAT_FNAME(StatName)
+#endif
+
+	#define RHI_BREADCRUMB_EVENT(                 RHICmdList_Or_RHIContext,                  Format, ...) FRHIBreadcrumbEventScope ANONYMOUS_VARIABLE(BreadcrumbEvent)(RHICmdList_Or_RHIContext, FRHIBreadcrumbData(__FILE__, __LINE__), true     , TEXT(Format), ## __VA_ARGS__)
+	#define RHI_BREADCRUMB_EVENT_CONDITIONAL(     RHICmdList_Or_RHIContext,       Condition, Format, ...) FRHIBreadcrumbEventScope ANONYMOUS_VARIABLE(BreadcrumbEvent)(RHICmdList_Or_RHIContext, FRHIBreadcrumbData(__FILE__, __LINE__), Condition, TEXT(Format), ## __VA_ARGS__)
 #if HAS_GPU_STATS
-	#define RHI_BREADCRUMB_EVENT_STAT(            RHICmdList_Or_RHIContext, Stat,            Format, ...) FRHIBreadcrumbEventScope ANONYMOUS_VARIABLE(BreadcrumbEvent)(RHICmdList_Or_RHIContext, FRHIBreadcrumbData(__FILE__, __LINE__, GET_STATID(Stat_GPU_##Stat), CSV_STAT_FNAME(Stat)), true     , TEXT(Format), ## __VA_ARGS__)
-	#define RHI_BREADCRUMB_EVENT_CONDITIONAL_STAT(RHICmdList_Or_RHIContext, Stat, Condition, Format, ...) FRHIBreadcrumbEventScope ANONYMOUS_VARIABLE(BreadcrumbEvent)(RHICmdList_Or_RHIContext, FRHIBreadcrumbData(__FILE__, __LINE__, GET_STATID(Stat_GPU_##Stat), CSV_STAT_FNAME(Stat)), Condition, TEXT(Format), ## __VA_ARGS__)
+	#define RHI_BREADCRUMB_EVENT_STAT(            RHICmdList_Or_RHIContext, Stat,            Format, ...) FRHIBreadcrumbEventScope ANONYMOUS_VARIABLE(BreadcrumbEvent)(RHICmdList_Or_RHIContext, FRHIBreadcrumbData(__FILE__, __LINE__, RHI_GPU_STAT_ARGS(Stat)), true     , TEXT(Format), ## __VA_ARGS__)
+	#define RHI_BREADCRUMB_EVENT_CONDITIONAL_STAT(RHICmdList_Or_RHIContext, Stat, Condition, Format, ...) FRHIBreadcrumbEventScope ANONYMOUS_VARIABLE(BreadcrumbEvent)(RHICmdList_Or_RHIContext, FRHIBreadcrumbData(__FILE__, __LINE__, RHI_GPU_STAT_ARGS(Stat)), Condition, TEXT(Format), ## __VA_ARGS__)
 #else
-	#define RHI_BREADCRUMB_EVENT_STAT(            RHICmdList_Or_RHIContext, Stat,            Format, ...) FRHIBreadcrumbEventScope ANONYMOUS_VARIABLE(BreadcrumbEvent)(RHICmdList_Or_RHIContext, FRHIBreadcrumbData(__FILE__, __LINE__, TStatId(), NAME_None), true     , TEXT(Format), ## __VA_ARGS__)
-	#define RHI_BREADCRUMB_EVENT_CONDITIONAL_STAT(RHICmdList_Or_RHIContext, Stat, Condition, Format, ...) FRHIBreadcrumbEventScope ANONYMOUS_VARIABLE(BreadcrumbEvent)(RHICmdList_Or_RHIContext, FRHIBreadcrumbData(__FILE__, __LINE__, TStatId(), NAME_None), Condition, TEXT(Format), ## __VA_ARGS__)
+	#define RHI_BREADCRUMB_EVENT_STAT(            RHICmdList_Or_RHIContext, Stat,            Format, ...) FRHIBreadcrumbEventScope ANONYMOUS_VARIABLE(BreadcrumbEvent)(RHICmdList_Or_RHIContext, FRHIBreadcrumbData(__FILE__, __LINE__), true     , TEXT(Format), ## __VA_ARGS__)
+	#define RHI_BREADCRUMB_EVENT_CONDITIONAL_STAT(RHICmdList_Or_RHIContext, Stat, Condition, Format, ...) FRHIBreadcrumbEventScope ANONYMOUS_VARIABLE(BreadcrumbEvent)(RHICmdList_Or_RHIContext, FRHIBreadcrumbData(__FILE__, __LINE__), Condition, TEXT(Format), ## __VA_ARGS__)
 #endif
 
 	// Used only for back compat with SCOPED_DRAW_EVENTF
-	#define RHI_BREADCRUMB_EVENT_STR_DEPRECATED(            RHICmdList_Or_RHIContext,            Format, ...) FRHIBreadcrumbEventScope ANONYMOUS_VARIABLE(BreadcrumbEvent)(RHICmdList_Or_RHIContext, FRHIBreadcrumbData(__FILE__, __LINE__, TStatId(), NAME_None), true     , Format, ## __VA_ARGS__)
-	#define RHI_BREADCRUMB_EVENT_CONDITIONAL_STR_DEPRECATED(RHICmdList_Or_RHIContext, Condition, Format, ...) FRHIBreadcrumbEventScope ANONYMOUS_VARIABLE(BreadcrumbEvent)(RHICmdList_Or_RHIContext, FRHIBreadcrumbData(__FILE__, __LINE__, TStatId(), NAME_None), Condition, Format, ## __VA_ARGS__)
+	#define RHI_BREADCRUMB_EVENT_STR_DEPRECATED(            RHICmdList_Or_RHIContext,            Format, ...) FRHIBreadcrumbEventScope ANONYMOUS_VARIABLE(BreadcrumbEvent)(RHICmdList_Or_RHIContext, FRHIBreadcrumbData(__FILE__, __LINE__), true     , Format, ## __VA_ARGS__)
+	#define RHI_BREADCRUMB_EVENT_CONDITIONAL_STR_DEPRECATED(RHICmdList_Or_RHIContext, Condition, Format, ...) FRHIBreadcrumbEventScope ANONYMOUS_VARIABLE(BreadcrumbEvent)(RHICmdList_Or_RHIContext, FRHIBreadcrumbData(__FILE__, __LINE__), Condition, Format, ## __VA_ARGS__)
 
 #else // WITH_RHI_BREADCRUMBS == 0
 
diff --git a/Engine/Source/Runtime/RenderCore/Private/ProfilingDebugging/RealtimeGPUProfiler.cpp b/Engine/Source/Runtime/RenderCore/Private/ProfilingDebugging/RealtimeGPUProfiler.cpp
--- a/Engine/Source/Runtime/RenderCore/Private/ProfilingDebugging/RealtimeGPUProfiler.cpp
+++ b/Engine/Source/Runtime/RenderCore/Private/ProfilingDebugging/RealtimeGPUProfiler.cpp
@@ -368,6 +368,11 @@ void TraverseEventTree(
 }
 #endif
 
+CSV_DEFINE_CATEGORY_MODULE(RENDERCORE_API, GPU, true);
+CSV_DEFINE_STAT(GPU, Total);
+
+DECLARE_FLOAT_COUNTER_STAT(TEXT("[TOTAL]"), Stat_GPU_Total, STATGROUP_GPU);
+
 /*-----------------------------------------------------------------------------
 FRealtimeGPUProfilerFrame class
 Container for a single frame's GPU stats
@@ -512,8 +517,19 @@ public:
 			}
 		}
 
-		extern RHI_API void RHISetGPUStatTotals(bool bCsvStatsEnabled, double TotalMs);
-		RHISetGPUStatTotals(bCsvStatsEnabled, TotalUs / 1000.0);
+		double TotalMs = TotalUs / 1000.0;
+
+#if STATS
+		FThreadStats::AddMessage(GET_STATFNAME(Stat_GPU_Total), EStatOperation::Set, TotalMs);
+		TRACE_STAT_SET(GET_STATFNAME(Stat_GPU_Total), TotalMs);
+#endif
+
+#if CSV_PROFILER_STATS
+		if (CsvProfiler)
+		{
+			CsvProfiler->RecordCustomStat(CSV_STAT_FNAME(Total), CSV_CATEGORY_INDEX(GPU), TotalMs, ECsvCustomStatOp::Set);
+		}
+#endif
 
 #if GPUPROFILERTRACE_ENABLED
 		TArray<TArray<int32>> GpuProfilerEventChildrenIndices;
diff --git a/Engine/Source/Runtime/RenderCore/Private/RenderCaptureInterface.cpp b/Engine/Source/Runtime/RenderCore/Private/RenderCaptureInterface.cpp
--- a/Engine/Source/Runtime/RenderCore/Private/RenderCaptureInterface.cpp
+++ b/Engine/Source/Runtime/RenderCore/Private/RenderCaptureInterface.cpp
@@ -35,7 +35,7 @@ namespace RenderCaptureInterface
 #if WITH_RHI_BREADCRUMBS
 				if (bPushEvent)
 				{
-					Breadcrumb->Emplace(RHICommandListLocal, FRHIBreadcrumbData(__FILE__, __LINE__, TStatId(), NAME_None), TEXT("%s"), EventName);
+					Breadcrumb->Emplace(RHICommandListLocal, FRHIBreadcrumbData(__FILE__, __LINE__), TEXT("%s"), EventName);
 				}
 #endif
 			});
@@ -58,7 +58,7 @@ namespace RenderCaptureInterface
 			if (bEvent)
 			{
 				RHIBreadcrumb = MakeUnique<TOptional<FRHIBreadcrumbEventManual>>();
-				RHIBreadcrumb->Emplace(*RHICommandList, FRHIBreadcrumbData(__FILE__, __LINE__, TStatId(), NAME_None), TEXT("%s"), FString(InEventName));
+				RHIBreadcrumb->Emplace(*RHICommandList, FRHIBreadcrumbData(__FILE__, __LINE__), TEXT("%s"), FString(InEventName));
 			}
 #endif
 		}
@@ -85,7 +85,7 @@ namespace RenderCaptureInterface
 #if RDG_EVENTS
 			if (bEvent)
 			{
-				RDGEvent.Emplace(*GraphBuilder, ERDGScopeFlags::None, FRHIBreadcrumbData(__FILE__, __LINE__, TStatId(), NAME_None), RDG_EVENT_NAME("%s", InEventName));
+				RDGEvent.Emplace(*GraphBuilder, ERDGScopeFlags::None, FRHIBreadcrumbData(__FILE__, __LINE__), RDG_EVENT_NAME("%s", InEventName));
 			}
 #endif
 		}
diff --git a/Engine/Source/Runtime/RenderCore/Private/RenderGraphBuilder.cpp b/Engine/Source/Runtime/RenderCore/Private/RenderGraphBuilder.cpp
--- a/Engine/Source/Runtime/RenderCore/Private/RenderGraphBuilder.cpp
+++ b/Engine/Source/Runtime/RenderCore/Private/RenderGraphBuilder.cpp
@@ -2600,7 +2600,7 @@ FRDGPass* FRDGBuilder::SetupParameterPass(FRDGPass* Pass)
 	if (ScopeState.ScopeMode == ERDGScopeMode::AllEventsAndPassNames)
 	{
 		FRDGEventName Name = Pass->GetEventName();
-		PassNameScope.Emplace(*this, ERDGScopeFlags::None, FRHIBreadcrumbData(__FILE__, __LINE__, TStatId(), NAME_None), MoveTemp(Name));
+		PassNameScope.Emplace(*this, ERDGScopeFlags::None, FRHIBreadcrumbData(__FILE__, __LINE__), MoveTemp(Name));
 	}
 #endif
 
diff --git a/Engine/Source/Runtime/RenderCore/Public/ProfilingDebugging/RealtimeGPUProfiler.h b/Engine/Source/Runtime/RenderCore/Public/ProfilingDebugging/RealtimeGPUProfiler.h
--- a/Engine/Source/Runtime/RenderCore/Public/ProfilingDebugging/RealtimeGPUProfiler.h
+++ b/Engine/Source/Runtime/RenderCore/Public/ProfilingDebugging/RealtimeGPUProfiler.h
@@ -67,11 +67,11 @@
 		}
 	};
 
-	#define RHI_BREADCRUMB_EVENT_GAMETHREAD(                       Format, ...) FRHIBreadcrumbEvent_GameThread ANONYMOUS_VARIABLE(BreadcrumbEvent_GameThread)(true     , FRHIBreadcrumbData(__FILE__, __LINE__, TStatId(), NAME_None), TEXT(Format), ##__VA_ARGS__)
-	#define RHI_BREADCRUMB_EVENT_CONDITIONAL_GAMETHREAD(Condition, Format, ...) FRHIBreadcrumbEvent_GameThread ANONYMOUS_VARIABLE(BreadcrumbEvent_GameThread)(Condition, FRHIBreadcrumbData(__FILE__, __LINE__, TStatId(), NAME_None), TEXT(Format), ##__VA_ARGS__)
+	#define RHI_BREADCRUMB_EVENT_GAMETHREAD(                       Format, ...) FRHIBreadcrumbEvent_GameThread ANONYMOUS_VARIABLE(BreadcrumbEvent_GameThread)(true     , FRHIBreadcrumbData(__FILE__, __LINE__), TEXT(Format), ##__VA_ARGS__)
+	#define RHI_BREADCRUMB_EVENT_CONDITIONAL_GAMETHREAD(Condition, Format, ...) FRHIBreadcrumbEvent_GameThread ANONYMOUS_VARIABLE(BreadcrumbEvent_GameThread)(Condition, FRHIBreadcrumbData(__FILE__, __LINE__), TEXT(Format), ##__VA_ARGS__)
 
 	// Used only for back compat with SCOPED_DRAW_EVENTF_GAMETHREAD
-	#define RHI_BREADCRUMB_EVENT_GAMETHREAD_STR_DEPRECATED(Format, ...) FRHIBreadcrumbEvent_GameThread ANONYMOUS_VARIABLE(BreadcrumbEvent_GameThread)(true, FRHIBreadcrumbData(__FILE__, __LINE__, TStatId(), NAME_None), Format, ##__VA_ARGS__)
+	#define RHI_BREADCRUMB_EVENT_GAMETHREAD_STR_DEPRECATED(Format, ...) FRHIBreadcrumbEvent_GameThread ANONYMOUS_VARIABLE(BreadcrumbEvent_GameThread)(true, FRHIBreadcrumbData(__FILE__, __LINE__), Format, ##__VA_ARGS__)
 
 #else
 
@@ -141,26 +141,41 @@
 
 	#if HAS_GPU_STATS
 
-		// @todo
-		#define DECLARE_GPU_STAT(                StatName            ) DECLARE_FLOAT_COUNTER_STAT(TEXT(#StatName)       , Stat_GPU_##StatName, STATGROUP_GPU  ); CSV_DEFINE_STAT(GPU,StatName);
-		#define DECLARE_GPU_STAT_NAMED(          StatName, NameString) DECLARE_FLOAT_COUNTER_STAT(NameString            , Stat_GPU_##StatName, STATGROUP_GPU  ); CSV_DEFINE_STAT(GPU,StatName);
-		#define DECLARE_GPU_DRAWCALL_STAT(       StatName            ) DECLARE_FLOAT_COUNTER_STAT(TEXT(#StatName)       , Stat_GPU_##StatName, STATGROUP_GPU  ); CSV_DEFINE_STAT(GPU,StatName);
-		#define DECLARE_GPU_DRAWCALL_STAT_NAMED( StatName, NameString) DECLARE_FLOAT_COUNTER_STAT(NameString            , Stat_GPU_##StatName, STATGROUP_GPU  ); CSV_DEFINE_STAT(GPU,StatName);
-		#define DECLARE_GPU_DRAWCALL_STAT_EXTERN(StatName            ) DECLARE_FLOAT_COUNTER_STAT_EXTERN(TEXT(#StatName), Stat_GPU_##StatName, STATGROUP_GPU, ); CSV_DECLARE_STAT_EXTERN(GPU,StatName);
+		#define DECLARE_GPU_STAT_NAME_TYPE(StatName, NameString) \
+			struct TRHIGPUStatNameProvider_##StatName            \
+			{                                                    \
+				static constexpr TCHAR const* GetName()          \
+				{                                                \
+					return NameString;                           \
+				}                                                \
+			}
 
 		// Extern GPU stats are needed where a stat is used in multiple CPPs. Use the DECLARE_GPU_STAT_NAMED_EXTERN in the header and DEFINE_GPU_STAT in the CPPs
-		#define DECLARE_GPU_STAT_NAMED_EXTERN(StatName, NameString) DECLARE_FLOAT_COUNTER_STAT_EXTERN(NameString, Stat_GPU_##StatName, STATGROUP_GPU, ); CSV_DECLARE_STAT_EXTERN(GPU,StatName);
-		#define DEFINE_GPU_STAT(              StatName            ) DEFINE_STAT(Stat_GPU_##StatName);                                                    CSV_DEFINE_STAT(GPU,StatName);
-		#define DEFINE_GPU_DRAWCALL_STAT(     StatName            ) DEFINE_STAT(Stat_GPU_##StatName);                                                    CSV_DEFINE_STAT(GPU,StatName);
+		#define DECLARE_GPU_STAT_NAMED(                StatName, NameString) DECLARE_GPU_STAT_NAME_TYPE(StatName, NameString); static UE::RHI::GPUProfiler::TGPUStat                    <TRHIGPUStatNameProvider_##StatName> GPUStat_##StatName
+		#define DECLARE_GPU_STAT_NAMED_EXTERN(         StatName, NameString) DECLARE_GPU_STAT_NAME_TYPE(StatName, NameString); extern UE::RHI::GPUProfiler::TGPUStat                    <TRHIGPUStatNameProvider_##StatName> GPUStat_##StatName
+		#define DECLARE_GPU_DRAWCALL_STAT_NAMED(       StatName, NameString) DECLARE_GPU_STAT_NAME_TYPE(StatName, NameString); static UE::RHI::GPUProfiler::TGPUStatWithDrawcallCategory<TRHIGPUStatNameProvider_##StatName> GPUStat_##StatName
+		#define DECLARE_GPU_DRAWCALL_STAT_NAMED_EXTERN(StatName, NameString) DECLARE_GPU_STAT_NAME_TYPE(StatName, NameString); extern UE::RHI::GPUProfiler::TGPUStatWithDrawcallCategory<TRHIGPUStatNameProvider_##StatName> GPUStat_##StatName
+
+		#define DECLARE_GPU_STAT(                StatName) DECLARE_GPU_STAT_NAMED(                StatName, TEXT(#StatName))
+		#define DECLARE_GPU_STAT_EXTERN(         StatName) DECLARE_GPU_STAT_NAMED_EXTERN(         StatName, TEXT(#StatName))
+		#define DECLARE_GPU_DRAWCALL_STAT(       StatName) DECLARE_GPU_DRAWCALL_STAT_NAMED(       StatName, TEXT(#StatName))
+		#define DECLARE_GPU_DRAWCALL_STAT_EXTERN(StatName) DECLARE_GPU_DRAWCALL_STAT_NAMED_EXTERN(StatName, TEXT(#StatName))
+
+		#define DEFINE_GPU_STAT(         StatName) UE::RHI::GPUProfiler::TGPUStat                    <TRHIGPUStatNameProvider_##StatName> GPUStat_##StatName
+		#define DEFINE_GPU_DRAWCALL_STAT(StatName) UE::RHI::GPUProfiler::TGPUStatWithDrawcallCategory<TRHIGPUStatNameProvider_##StatName> GPUStat_##StatName
 
 	#else
 
-		#define DECLARE_GPU_STAT(...)
 		#define DECLARE_GPU_STAT_NAMED(...)
-		#define DECLARE_GPU_DRAWCALL_STAT(...)
+		#define DECLARE_GPU_STAT_NAMED_EXTERN(...)
 		#define DECLARE_GPU_DRAWCALL_STAT_NAMED(...)
+		#define DECLARE_GPU_DRAWCALL_STAT_NAMED_EXTERN(...)
+
+		#define DECLARE_GPU_STAT(...)
+		#define DECLARE_GPU_STAT_EXTERN(...)
+		#define DECLARE_GPU_DRAWCALL_STAT(...)
 		#define DECLARE_GPU_DRAWCALL_STAT_EXTERN(...)
-		#define DECLARE_GPU_STAT_NAMED_EXTERN(...)
+
 		#define DEFINE_GPU_STAT(...)
 		#define DEFINE_GPU_DRAWCALL_STAT(...)
 
@@ -180,14 +195,18 @@ class FRealtimeGPUProfiler;
 class FRealtimeGPUProfilerEvent;
 class FRealtimeGPUProfilerFrame;
 
+DECLARE_STATS_GROUP(TEXT("GPU"), STATGROUP_GPU, STATCAT_Advanced);
+CSV_DECLARE_CATEGORY_MODULE_EXTERN(RENDERCORE_API, GPU);
+
 #if HAS_GPU_STATS
 
 	// The DECLARE_GPU_STAT macros both declare and define a stat (for use in a single CPP)
-	#define DECLARE_GPU_STAT(StatName)                            DECLARE_FLOAT_COUNTER_STAT(TEXT(#StatName)       , Stat_GPU_##StatName, STATGROUP_GPU  ); CSV_DEFINE_STAT(GPU,StatName);         static FRHIDrawStatsCategory DrawcallCountCategory_##StatName;
-	#define DECLARE_GPU_STAT_NAMED(StatName, NameString)          DECLARE_FLOAT_COUNTER_STAT(NameString            , Stat_GPU_##StatName, STATGROUP_GPU  ); CSV_DEFINE_STAT(GPU,StatName);         static FRHIDrawStatsCategory DrawcallCountCategory_##StatName;
-	#define DECLARE_GPU_DRAWCALL_STAT(StatName)                   DECLARE_FLOAT_COUNTER_STAT(TEXT(#StatName)       , Stat_GPU_##StatName, STATGROUP_GPU  ); CSV_DEFINE_STAT(GPU,StatName);         static FRHIDrawStatsCategory DrawcallCountCategory_##StatName((TCHAR*)TEXT(#StatName));
-	#define DECLARE_GPU_DRAWCALL_STAT_NAMED(StatName, NameString) DECLARE_FLOAT_COUNTER_STAT(NameString            , Stat_GPU_##StatName, STATGROUP_GPU  ); CSV_DEFINE_STAT(GPU,StatName);         static FRHIDrawStatsCategory DrawcallCountCategory_##StatName((TCHAR*)TEXT(#StatName));
-	#define DECLARE_GPU_DRAWCALL_STAT_EXTERN(StatName)            DECLARE_FLOAT_COUNTER_STAT_EXTERN(TEXT(#StatName), Stat_GPU_##StatName, STATGROUP_GPU, ); CSV_DECLARE_STAT_EXTERN(GPU,StatName); extern FRHIDrawStatsCategory DrawcallCountCategory_##StatName;
+	#define DECLARE_GPU_STAT(StatName)                                   DECLARE_FLOAT_COUNTER_STAT(TEXT(#StatName)       , Stat_GPU_##StatName, STATGROUP_GPU  ); CSV_DEFINE_STAT(GPU,StatName);         static FRHIDrawStatsCategory DrawcallCountCategory_##StatName;
+	#define DECLARE_GPU_STAT_NAMED(StatName, NameString)                 DECLARE_FLOAT_COUNTER_STAT(NameString            , Stat_GPU_##StatName, STATGROUP_GPU  ); CSV_DEFINE_STAT(GPU,StatName);         static FRHIDrawStatsCategory DrawcallCountCategory_##StatName;
+	#define DECLARE_GPU_DRAWCALL_STAT(StatName)                          DECLARE_FLOAT_COUNTER_STAT(TEXT(#StatName)       , Stat_GPU_##StatName, STATGROUP_GPU  ); CSV_DEFINE_STAT(GPU,StatName);         static FRHIDrawStatsCategory DrawcallCountCategory_##StatName((TCHAR*)TEXT(#StatName));
+	#define DECLARE_GPU_DRAWCALL_STAT_NAMED(StatName, NameString)        DECLARE_FLOAT_COUNTER_STAT(NameString            , Stat_GPU_##StatName, STATGROUP_GPU  ); CSV_DEFINE_STAT(GPU,StatName);         static FRHIDrawStatsCategory DrawcallCountCategory_##StatName((TCHAR*)TEXT(#StatName));
+	#define DECLARE_GPU_DRAWCALL_STAT_EXTERN(StatName)                   DECLARE_FLOAT_COUNTER_STAT_EXTERN(TEXT(#StatName), Stat_GPU_##StatName, STATGROUP_GPU, ); CSV_DECLARE_STAT_EXTERN(GPU,StatName); extern FRHIDrawStatsCategory DrawcallCountCategory_##StatName;
+	#define DECLARE_GPU_DRAWCALL_STAT_NAMED_EXTERN(StatName, NameString) DECLARE_FLOAT_COUNTER_STAT_EXTERN(NameString     , Stat_GPU_##StatName, STATGROUP_GPU, ); CSV_DECLARE_STAT_EXTERN(GPU,StatName); extern FRHIDrawStatsCategory DrawcallCountCategory_##StatName;
 
 	// Extern GPU stats are needed where a stat is used in multiple CPPs. Use the DECLARE_GPU_STAT_NAMED_EXTERN in the header and DEFINE_GPU_STAT in the CPPs
 	#define DECLARE_GPU_STAT_NAMED_EXTERN(StatName, NameString) DECLARE_FLOAT_COUNTER_STAT_EXTERN(NameString, Stat_GPU_##StatName, STATGROUP_GPU, ); CSV_DECLARE_STAT_EXTERN(GPU,StatName); extern FRHIDrawStatsCategory DrawcallCountCategory_##StatName;
@@ -218,6 +237,7 @@ class FRealtimeGPUProfilerFrame;
 	#define DECLARE_GPU_DRAWCALL_STAT(StatName)
 	#define DECLARE_GPU_DRAWCALL_STAT_NAMED(StatName, NameString)
 	#define DECLARE_GPU_DRAWCALL_STAT_EXTERN(StatName)
+	#define DECLARE_GPU_DRAWCALL_STAT_NAMED_EXTERN(StatName, NameString)
 
 	#define DECLARE_GPU_STAT_NAMED_EXTERN(StatName, NameString)
 	#define DEFINE_GPU_STAT(StatName)
diff --git a/Engine/Source/Runtime/RenderCore/Public/RenderGraphEvent.h b/Engine/Source/Runtime/RenderCore/Public/RenderGraphEvent.h
--- a/Engine/Source/Runtime/RenderCore/Public/RenderGraphEvent.h
+++ b/Engine/Source/Runtime/RenderCore/Public/RenderGraphEvent.h
@@ -348,7 +348,7 @@ private:
 		TRDGEventScopeGuard<FRDGScope_RHI> PREPROCESSOR_JOIN(__RDG_ScopeRef_,__LINE__)( \
 			(GraphBuilder)                                                              \
 			, ERDGScopeFlags::None                                                      \
-			, FRHIBreadcrumbData(__FILE__, __LINE__, TStatId(), NAME_None)              \
+			, FRHIBreadcrumbData(__FILE__, __LINE__)                                    \
 			, RDG_EVENT_NAME(Format, ##__VA_ARGS__)                                     \
 		)
 
@@ -360,8 +360,7 @@ private:
 				, FRHIBreadcrumbData(                                                       \
 					  __FILE__                                                              \
 					, __LINE__                                                              \
-					, GET_STATID(Stat_GPU_##StatName)                                       \
-					, CSV_STAT_FNAME(StatName)                                              \
+					, RHI_GPU_STAT_ARGS(StatName)                                           \
 				)                                                                           \
 				, RDG_EVENT_NAME(Format, ##__VA_ARGS__)                                     \
 			)
@@ -382,8 +381,6 @@ private:
 					, FRHIBreadcrumbData(                                                          \
 						  __FILE__                                                                 \
 						, __LINE__                                                                 \
-						, TStatId()                                                                \
-						, NAME_None                                                                \
 					)                                                                              \
 					, RDG_EVENT_NAME(Format, ##__VA_ARGS__)                                        \
 				);                                                                                 \
@@ -403,8 +400,7 @@ private:
 						, FRHIBreadcrumbData(                                                          \
 							  __FILE__                                                                 \
 							, __LINE__                                                                 \
-							, GET_STATID(Stat_GPU_##StatName)                                          \
-							, CSV_STAT_FNAME(StatName)                                                 \
+							, RHI_GPU_STAT_ARGS(StatName)                                              \
 						)                                                                              \
 						, RDG_EVENT_NAME(Format, ##__VA_ARGS__)                                        \
 					);                                                                                 \
@@ -424,7 +420,7 @@ private:
 			PREPROCESSOR_JOIN(__RDG_ScopeRef_,__LINE__).Emplace(                                   \
 				(GraphBuilder)                                                                     \
 				, ERDGScopeFlags::Final                                                            \
-				, FRHIBreadcrumbData(__FILE__, __LINE__, TStatId(), NAME_None)                     \
+				, FRHIBreadcrumbData(__FILE__, __LINE__)                                           \
 				, RDG_EVENT_NAME(Format, ##__VA_ARGS__)                                            \
 			);                                                                                     \
 		} while (false)
diff --git a/Engine/Source/Runtime/Renderer/Private/GPUScene.h b/Engine/Source/Runtime/Renderer/Private/GPUScene.h
--- a/Engine/Source/Runtime/Renderer/Private/GPUScene.h
+++ b/Engine/Source/Runtime/Renderer/Private/GPUScene.h
@@ -37,7 +37,7 @@ namespace UE::Renderer::Private
 	class IShadowInvalidatingInstances;
 }
 
-DECLARE_GPU_STAT_NAMED_EXTERN(GPUSceneUpdate, TEXT("GPUSceneUpdate"))
+DECLARE_GPU_STAT_NAMED_EXTERN(GPUSceneUpdate, TEXT("GPUSceneUpdate"));
 
 #if !ENABLE_SCENE_DATA_DX11_UB_ERROR_WORKAROUND
 // Note: this is declared in GPUSceneWriter.h (a public header) until such a time as the UB path works.
diff --git a/Engine/Source/Runtime/Renderer/Private/PathTracingSpatialTemporalDenoising.cpp b/Engine/Source/Runtime/Renderer/Private/PathTracingSpatialTemporalDenoising.cpp
--- a/Engine/Source/Runtime/Renderer/Private/PathTracingSpatialTemporalDenoising.cpp
+++ b/Engine/Source/Runtime/Renderer/Private/PathTracingSpatialTemporalDenoising.cpp
@@ -361,7 +361,7 @@ namespace {
 		ECVF_RenderThreadSafe
 	);
 
-	DECLARE_GPU_STAT(PathTracingSpatialTemporalDenoising)
+	DECLARE_GPU_STAT(PathTracingSpatialTemporalDenoising);
 }
 
 int GetPathTracingDenoiserMode(const FViewInfo& View)
diff --git a/Engine/Source/Runtime/Renderer/Private/PostProcess/DiaphragmDOF.cpp b/Engine/Source/Runtime/Renderer/Private/PostProcess/DiaphragmDOF.cpp
--- a/Engine/Source/Runtime/Renderer/Private/PostProcess/DiaphragmDOF.cpp
+++ b/Engine/Source/Runtime/Renderer/Private/PostProcess/DiaphragmDOF.cpp
@@ -29,7 +29,7 @@
 namespace
 {
 
-DECLARE_GPU_STAT(DepthOfField)
+DECLARE_GPU_STAT(DepthOfField);
 
 TAutoConsoleVariable<int32> CVarDOFGatherResDivisor(
 	TEXT("r.DOF.Gather.ResolutionDivisor"), 2,
diff --git a/Engine/Source/Runtime/Renderer/Private/PostProcess/PostProcessMotionBlur.cpp b/Engine/Source/Runtime/Renderer/Private/PostProcess/PostProcessMotionBlur.cpp
--- a/Engine/Source/Runtime/Renderer/Private/PostProcess/PostProcessMotionBlur.cpp
+++ b/Engine/Source/Runtime/Renderer/Private/PostProcess/PostProcessMotionBlur.cpp
@@ -108,7 +108,7 @@ namespace
 		return MotionBlurQuality;
 	}
 
-DECLARE_GPU_STAT(MotionBlur)
+	DECLARE_GPU_STAT(MotionBlur);
 
 }
 
diff --git a/Engine/Source/Runtime/Renderer/Private/PostProcess/PostProcessSubsurface.cpp b/Engine/Source/Runtime/Renderer/Private/PostProcess/PostProcessSubsurface.cpp
--- a/Engine/Source/Runtime/Renderer/Private/PostProcess/PostProcessSubsurface.cpp
+++ b/Engine/Source/Runtime/Renderer/Private/PostProcess/PostProcessSubsurface.cpp
@@ -145,7 +145,7 @@ namespace
 		TEXT("Normalized distance (0..1) at which the surface color is interpreted as transmittance color to compute extinction coefficients."),
 		ECVF_RenderThreadSafe);
 
-	DECLARE_GPU_STAT(SubsurfaceScattering)
+	DECLARE_GPU_STAT(SubsurfaceScattering);
 }
 
 // Define to use a custom ps to clear UAV.
diff --git a/Engine/Source/Runtime/Renderer/Private/PostProcess/TemporalAA.cpp b/Engine/Source/Runtime/Renderer/Private/PostProcess/TemporalAA.cpp
--- a/Engine/Source/Runtime/Renderer/Private/PostProcess/TemporalAA.cpp
+++ b/Engine/Source/Runtime/Renderer/Private/PostProcess/TemporalAA.cpp
@@ -477,7 +477,7 @@ void SetupSampleWeightParameters(FTemporalAA::FParameters* OutTAAParameters, con
 	}
 }
 
-DECLARE_GPU_STAT(TAA)
+DECLARE_GPU_STAT(TAA);
 
 const TCHAR* const kTAAOutputNames[] = {
 	TEXT("TAA.History"),
diff --git a/Engine/Source/Runtime/Renderer/Private/PostProcess/TemporalSuperResolution.cpp b/Engine/Source/Runtime/Renderer/Private/PostProcess/TemporalSuperResolution.cpp
--- a/Engine/Source/Runtime/Renderer/Private/PostProcess/TemporalSuperResolution.cpp
+++ b/Engine/Source/Runtime/Renderer/Private/PostProcess/TemporalSuperResolution.cpp
@@ -1167,7 +1167,7 @@ IMPLEMENT_GLOBAL_SHADER(FTSRUpdateHistoryCS,         "/Engine/Private/TemporalSu
 IMPLEMENT_GLOBAL_SHADER(FTSRResolveHistoryCS,        "/Engine/Private/TemporalSuperResolution/TSRResolveHistory.usf",        "MainCS", SF_Compute);
 IMPLEMENT_GLOBAL_SHADER(FTSRVisualizeCS,             "/Engine/Private/TemporalSuperResolution/TSRVisualize.usf",             "MainCS", SF_Compute);
 
-DECLARE_GPU_STAT(TemporalSuperResolution)
+DECLARE_GPU_STAT(TemporalSuperResolution);
 
 } //! namespace
 
diff --git a/Engine/Source/Runtime/Renderer/Private/SceneRenderBuilder.cpp b/Engine/Source/Runtime/Renderer/Private/SceneRenderBuilder.cpp
--- a/Engine/Source/Runtime/Renderer/Private/SceneRenderBuilder.cpp
+++ b/Engine/Source/Runtime/Renderer/Private/SceneRenderBuilder.cpp
@@ -114,7 +114,7 @@ class FSceneRenderProcessor
 #if WITH_RHI_BREADCRUMBS
 			if (Group && IsLocationActive(Location))
 			{
-				FRHIBreadcrumbData BreadcrumbData(__FILE__, __LINE__, TStatId(), NAME_None);
+				FRHIBreadcrumbData BreadcrumbData(__FILE__, __LINE__);
 
 				// User specified an explicit group name.
 				if (!Group->Name.IsEmpty())
diff --git a/Engine/Source/Runtime/Renderer/Private/SceneRendering.h b/Engine/Source/Runtime/Renderer/Private/SceneRendering.h
--- a/Engine/Source/Runtime/Renderer/Private/SceneRendering.h
+++ b/Engine/Source/Runtime/Renderer/Private/SceneRendering.h
@@ -111,8 +111,8 @@ struct FPersistentViewId
 extern bool ShouldUseStereoLumenOptimizations();
 
 DECLARE_GPU_DRAWCALL_STAT_EXTERN(VirtualTextureUpdate);
-DECLARE_GPU_STAT_NAMED_EXTERN(Postprocessing, TEXT("Postprocessing"))
-DECLARE_GPU_STAT_NAMED_EXTERN(CustomRenderPasses, TEXT("CustomRenderPasses"))
+DECLARE_GPU_STAT_NAMED_EXTERN(Postprocessing, TEXT("Postprocessing"));
+DECLARE_GPU_STAT_NAMED_EXTERN(CustomRenderPasses, TEXT("CustomRenderPasses"));
 
 /** Mobile only. Information used to determine whether static meshes will be rendered with CSM shaders or not. */
 class FMobileCSMVisibilityInfo
diff --git a/Engine/Source/Runtime/Renderer/Private/ScreenSpaceDenoise.cpp b/Engine/Source/Runtime/Renderer/Private/ScreenSpaceDenoise.cpp
--- a/Engine/Source/Runtime/Renderer/Private/ScreenSpaceDenoise.cpp
+++ b/Engine/Source/Runtime/Renderer/Private/ScreenSpaceDenoise.cpp
@@ -141,10 +141,10 @@ static_assert(IScreenSpaceDenoiser::kMaxBatchSize <= kMaxBufferProcessingCount,
 
 const IScreenSpaceDenoiser* GScreenSpaceDenoiser = nullptr;
 
-DECLARE_GPU_STAT(ReflectionsDenoiser)
-DECLARE_GPU_STAT(ShadowsDenoiser)
-DECLARE_GPU_STAT(AmbientOcclusionDenoiser)
-DECLARE_GPU_STAT(DiffuseIndirectDenoiser)
+DECLARE_GPU_STAT(ReflectionsDenoiser);
+DECLARE_GPU_STAT(ShadowsDenoiser);
+DECLARE_GPU_STAT(AmbientOcclusionDenoiser);
+DECLARE_GPU_STAT(DiffuseIndirectDenoiser);
 
 namespace
 {
diff --git a/Engine/Source/Runtime/Renderer/Private/ShadowDepthRendering.cpp b/Engine/Source/Runtime/Renderer/Private/ShadowDepthRendering.cpp
--- a/Engine/Source/Runtime/Renderer/Private/ShadowDepthRendering.cpp
+++ b/Engine/Source/Runtime/Renderer/Private/ShadowDepthRendering.cpp
@@ -1597,7 +1597,7 @@ void FSceneRenderer::RenderShadowDepthMapAtlases(FRDGBuilder& GraphBuilder)
 				FString LightNameWithLevel;
 				GetLightNameForDrawEvent(CurrentLightForDrawEvent, LightNameWithLevel);
 
-				RDGScope.Emplace(GraphBuilder, ERDGScopeFlags::None, FRHIBreadcrumbData(__FILE__, __LINE__, TStatId(), NAME_None), RDG_EVENT_NAME("%s", *LightNameWithLevel));
+				RDGScope.Emplace(GraphBuilder, ERDGScopeFlags::None, FRHIBreadcrumbData(__FILE__, __LINE__), RDG_EVENT_NAME("%s", *LightNameWithLevel));
 			}
 		#endif
 		};
diff --git a/Engine/Source/Runtime/Renderer/Private/ShadowRendering.h b/Engine/Source/Runtime/Renderer/Private/ShadowRendering.h
--- a/Engine/Source/Runtime/Renderer/Private/ShadowRendering.h
+++ b/Engine/Source/Runtime/Renderer/Private/ShadowRendering.h
@@ -46,7 +46,7 @@ class FVirtualShadowMapPerLightCacheEntry;
 class FLightTileIntersectionParameters;
 class FDistanceFieldCulledObjectBufferParameters;
 
-DECLARE_GPU_STAT_NAMED_EXTERN(ShadowDepths, TEXT("Shadow Depths"));
+DECLARE_GPU_DRAWCALL_STAT_NAMED_EXTERN(ShadowDepths, TEXT("Shadow Depths"));
 
 /** Renders a cone with a spherical cap, used for rendering spot lights in deferred passes. */
 extern void DrawStencilingCone(const FMatrix& ConeToWorld, float ConeAngle, float SphereRadius, const FVector& PreViewTranslation);
