Fixed PSO stat registration crash caused by separate synchronization scopes

diff --git a/Engine/Source/Runtime/RHI/Private/PipelineFileCache.cpp b/Engine/Source/Runtime/RHI/Private/PipelineFileCache.cpp
--- a/Engine/Source/Runtime/RHI/Private/PipelineFileCache.cpp
+++ b/Engine/Source/Runtime/RHI/Private/PipelineFileCache.cpp
@@ -3557,7 +3557,7 @@ void FPipelineFileCacheManager::BroadcastNewPSOsDelegate()
 	}
 }
 
-void FPipelineFileCacheManager::CacheGraphicsPSO(uint32 RunTimeHash, FGraphicsPipelineStateInitializer const& Initializer, bool bWasPSOPrecached)
+void FPipelineFileCacheManager::CacheGraphicsPSO(uint32 RunTimeHash, FGraphicsPipelineStateInitializer const& Initializer, bool bWasPSOPrecached, FPipelineStateStats** OutStats)
 {
 	if(IsPipelineFileCacheEnabled() && (LogPSOtoFileCache() || ReportNewPSOs()))
 	{
@@ -3639,10 +3639,16 @@ void FPipelineFileCacheManager::CacheGraphicsPSO(uint32 RunTimeHash, FGraphicsPi
 				RegisterPSOUsageDataUpdateForNextSave(*PSOUsage);
 			}
 		}
+
+		// Optionally supply stats to share the lock scope
+		if (OutStats)
+		{
+			*OutStats = RegisterPSOStatsInternal(Lock, RunTimeHash);
+		}
 	}
 }
 
-void FPipelineFileCacheManager::CacheComputePSO(uint32 RunTimeHash, FRHIComputeShader const* Initializer, bool bWasPSOPrecached)
+void FPipelineFileCacheManager::CacheComputePSO(uint32 RunTimeHash, FRHIComputeShader const* Initializer, bool bWasPSOPrecached, FPipelineStateStats** OutStats)
 {
 	if(IsPipelineFileCacheEnabled() && (LogPSOtoFileCache() || ReportNewPSOs()))
 	{
@@ -3715,6 +3721,12 @@ void FPipelineFileCacheManager::CacheComputePSO(uint32 RunTimeHash, FRHIComputeS
 				}
 			}
 		}
+
+		// Optionally supply stats to share the lock scope
+		if (OutStats)
+		{
+			*OutStats = RegisterPSOStatsInternal(Lock, RunTimeHash);
+		}
 	}
 }
 
@@ -3857,28 +3869,42 @@ void FPipelineFileCacheManager::RegisterPSOCompileFailure(uint32 RunTimeHash, FG
 
 FPipelineStateStats* FPipelineFileCacheManager::RegisterPSOStats(uint32 RunTimeHash)
 {
-	FPipelineStateStats* Stat = nullptr;
 	if(IsPipelineFileCacheEnabled() && LogPSOtoFileCache())
 	{
 		FRWScopeLock Lock(FileCacheLock, SLT_ReadOnly);
+		return RegisterPSOStatsInternal(Lock, RunTimeHash);
+	}
+	
+	return nullptr;
+}
 
-		uint32 PSOHash = RunTimeToPSOUsage.FindChecked(RunTimeHash).PSOHash;
+FPipelineStateStats* FPipelineFileCacheManager::RegisterPSOStatsInternal(FRWScopeLock& Lock, uint32 RunTimeHash)
+{
+	// May fail registration if the user cache has been closed
+	FPSOUsageData* PSOUsage = RunTimeToPSOUsage.Find(RunTimeHash);
+	if (!PSOUsage)
+	{
+		UE_LOG(LogRHI, Display, TEXT("PSO (%u) missing from cache, please instead use the OutStats parameter in CacheGraphicsPSO/CacheComputePSO"), RunTimeHash);
+		return nullptr;
+	}
+
+	uint32 PSOHash = PSOUsage->PSOHash;
+	FPipelineStateStats* Stat = Stats.FindRef(PSOHash);
+	if (!Stat)
+	{
+		Lock.ReleaseReadOnlyLockAndAcquireWriteLock_USE_WITH_CAUTION();
 		Stat = Stats.FindRef(PSOHash);
 		if (!Stat)
 		{
-			Lock.ReleaseReadOnlyLockAndAcquireWriteLock_USE_WITH_CAUTION();
-			Stat = Stats.FindRef(PSOHash);
-			if (!Stat)
-			{
-				Stat = new FPipelineStateStats;
-				Stat->PSOHash = PSOHash;
-				Stats.Add(PSOHash, Stat);
+			Stat = new FPipelineStateStats;
+			Stat->PSOHash = PSOHash;
+			Stats.Add(PSOHash, Stat);
 
-				INC_MEMORY_STAT_BY(STAT_PSOStatMemory, sizeof(FPipelineStateStats) + sizeof(uint32));
-			}
+			INC_MEMORY_STAT_BY(STAT_PSOStatMemory, sizeof(FPipelineStateStats) + sizeof(uint32));
 		}
-		Stat->CreateCount++;
 	}
+	Stat->CreateCount++;
+
 	return Stat;
 }
 
diff --git a/Engine/Source/Runtime/RHI/Private/PipelineStateCache.cpp b/Engine/Source/Runtime/RHI/Private/PipelineStateCache.cpp
--- a/Engine/Source/Runtime/RHI/Private/PipelineStateCache.cpp
+++ b/Engine/Source/Runtime/RHI/Private/PipelineStateCache.cpp
@@ -3272,11 +3272,13 @@ FComputePipelineState* PipelineStateCache::GetAndOrCreateComputePipelineState(FR
 		EPSOPrecacheResult PSOPrecacheResult = CheckPipelineStateInCache(ComputeShader);
 
 		bool bWasPSOPrecached = PSOPrecacheResult == EPSOPrecacheResult::Active || PSOPrecacheResult == EPSOPrecacheResult::Complete;
-		FPipelineFileCacheManager::CacheComputePSO(GetTypeHash(ComputeShader), ComputeShader, bWasPSOPrecached);
+		
+		FPipelineStateStats* Stats = nullptr;
+		FPipelineFileCacheManager::CacheComputePSO(GetTypeHash(ComputeShader), ComputeShader, bWasPSOPrecached, &Stats);
 
 		// create new compute state
 		OutCachedState = new FComputePipelineState(ComputeShader, nullptr);
-		OutCachedState->Stats = FPipelineFileCacheManager::RegisterPSOStats(GetTypeHash(ComputeShader));
+		OutCachedState->Stats = Stats;
 		if (DoAsyncCompile)
 		{
 			OutCachedState->SetCompletionEvent(FGraphEvent::CreateGraphEvent());
@@ -3792,11 +3794,13 @@ FGraphicsPipelineState* PipelineStateCache::GetAndOrCreateGraphicsPipelineState(
 		bool DoAsyncCompile = IsAsyncCompilationAllowed(RHICmdList, Initializer.bFromPSOFileCache);
 
 		bool bWasPSOPrecached = PSOPrecacheResult == EPSOPrecacheResult::Active || PSOPrecacheResult == EPSOPrecacheResult::Complete;
-		FPipelineFileCacheManager::CacheGraphicsPSO(GetTypeHash(Initializer), Initializer, bWasPSOPrecached);
+
+		FPipelineStateStats* Stats = nullptr;
+		FPipelineFileCacheManager::CacheGraphicsPSO(GetTypeHash(Initializer), Initializer, bWasPSOPrecached, &Stats);
 
 		// create new graphics state
 		OutCachedState = new FGraphicsPipelineState();
-		OutCachedState->Stats = FPipelineFileCacheManager::RegisterPSOStats(GetTypeHash(Initializer));
+		OutCachedState->Stats = Stats;
 		if (DoAsyncCompile)
 		{
 			OutCachedState->SetCompletionEvent(FGraphEvent::CreateGraphEvent());
diff --git a/Engine/Source/Runtime/RHI/Public/PipelineFileCache.h b/Engine/Source/Runtime/RHI/Public/PipelineFileCache.h
--- a/Engine/Source/Runtime/RHI/Public/PipelineFileCache.h
+++ b/Engine/Source/Runtime/RHI/Public/PipelineFileCache.h
@@ -342,13 +342,14 @@ public:
 
 	RHI_API static void CloseUserPipelineFileCache();
 
-	RHI_API static void CacheGraphicsPSO(uint32 RunTimeHash, FGraphicsPipelineStateInitializer const& Initializer, bool bWasPSOPrecached);
-	RHI_API static void CacheComputePSO(uint32 RunTimeHash, FRHIComputeShader const* Initializer, bool bWasPSOPrecached);
+	RHI_API static void CacheGraphicsPSO(uint32 RunTimeHash, FGraphicsPipelineStateInitializer const& Initializer, bool bWasPSOPrecached, FPipelineStateStats** OutStats = nullptr);
+	RHI_API static void CacheComputePSO(uint32 RunTimeHash, FRHIComputeShader const* Initializer, bool bWasPSOPrecached, FPipelineStateStats** OutStats = nullptr);
 	RHI_API static void CacheRayTracingPSO(const FRayTracingPipelineStateInitializer& Initializer, ERayTracingPipelineCacheFlags Flags);
 
 	// true if the named PSOFC is currently open.
 	RHI_API static bool HasPipelineFileCache(const FString& PSOCacheKey);
 
+	UE_DEPRECATED(5.6, "Do not call this function, please instead use the OutStats parameter in CacheGraphicsPSO/CacheComputePSO")
 	RHI_API static FPipelineStateStats* RegisterPSOStats(uint32 RunTimeHash);
 	
 	/*
@@ -414,6 +415,9 @@ private:
 	static void LogNewGraphicsPSOToConsoleAndCSV(FPipelineCacheFileFormatPSO& PSO, uint32 PSOHash, bool bWasPSOPrecached);
 	static void LogNewComputePSOToConsoleAndCSV(FPipelineCacheFileFormatPSO& PSO, uint32 PSOHash, bool bWasPSOPrecached);
 	static void LogNewRaytracingPSOToConsole(FPipelineCacheFileFormatPSO& PSO, uint32 PSOHash, bool bIsNonBlockingPSO);
+
+	static FPipelineStateStats* RegisterPSOStatsInternal(class FRWScopeLock& Lock, uint32 RunTimeHash);
+	
 private:
 	static FRWLock FileCacheLock;
 
