RHI Validation fixes

diff --git a/Engine/Source/Runtime/RHI/Private/RHIValidation.cpp b/Engine/Source/Runtime/RHI/Private/RHIValidation.cpp
--- a/Engine/Source/Runtime/RHI/Private/RHIValidation.cpp
+++ b/Engine/Source/Runtime/RHI/Private/RHIValidation.cpp
@@ -482,10 +482,11 @@ namespace RHIValidation
 			const FRayTracingLocalShaderBindings& Binding = Bindings[CurrentIndex];
 
 			// only collect shader binding data if RTPSO & hit group indexing mode
-			if (HitGroupIndexingMode == ERayTracingHitGroupIndexingMode::Allow && Binding.Geometry)
+			bool bValidBinding = (Binding.BindingType == ERayTracingLocalShaderBindingType::Persistent || Binding.BindingType == ERayTracingLocalShaderBindingType::Validation);
+			if (HitGroupIndexingMode == ERayTracingHitGroupIndexingMode::Allow && bValidBinding)
 			{
 				if (EnumHasAnyFlags(ShaderBindingMode, ERayTracingShaderBindingMode::RTPSO))
-				{
+				{					
 					CollectShaderBindingTableResources(RayTracingPipelineState, this, Binding, BindingType, Context.WorkerIndex);
 
 					// Also add SRV view requirement for all index and vertex buffers used in the SBT
@@ -2319,18 +2320,21 @@ namespace RHIValidation
 				}
 			}
 
-			uint32 LastMip = SubresourceRange.MipIndex + SubresourceRange.NumMips;
-			uint32 LastArraySlice = SubresourceRange.ArraySlice + SubresourceRange.NumArraySlices;
-			uint32 LastPlaneIndex = SubresourceRange.PlaneIndex + SubresourceRange.NumPlanes;
-
-			for (uint32 PlaneIndex = SubresourceRange.PlaneIndex; PlaneIndex < LastPlaneIndex; ++PlaneIndex)
+			if (SubresourceStates.Num() == 0)
 			{
-				for (uint32 MipIndex = SubresourceRange.MipIndex; MipIndex < LastMip; ++MipIndex)
+				uint32 LastMip = SubresourceRange.MipIndex + SubresourceRange.NumMips;
+				uint32 LastArraySlice = SubresourceRange.ArraySlice + SubresourceRange.NumArraySlices;
+				uint32 LastPlaneIndex = SubresourceRange.PlaneIndex + SubresourceRange.NumPlanes;
+
+				for (uint32 PlaneIndex = SubresourceRange.PlaneIndex; PlaneIndex < LastPlaneIndex; ++PlaneIndex)
 				{
-					for (uint32 ArraySlice = SubresourceRange.ArraySlice; ArraySlice < LastArraySlice; ++ArraySlice)
+					for (uint32 MipIndex = SubresourceRange.MipIndex; MipIndex < LastMip; ++MipIndex)
 					{
-						uint32 SubresourceIndex = PlaneIndex + (MipIndex + ArraySlice * NumMips) * NumPlanes;
-						Callback(SubresourceStates[SubresourceIndex], FSubresourceIndex(MipIndex, ArraySlice, PlaneIndex));
+						for (uint32 ArraySlice = SubresourceRange.ArraySlice; ArraySlice < LastArraySlice; ++ArraySlice)
+						{
+							uint32 SubresourceIndex = PlaneIndex + (MipIndex + ArraySlice * NumMips) * NumPlanes;
+							Callback(SubresourceStates[SubresourceIndex], FSubresourceIndex(MipIndex, ArraySlice, PlaneIndex));
+						}
 					}
 				}
 			}
diff --git a/Engine/Source/Runtime/RenderCore/Private/GlobalRenderResources.cpp b/Engine/Source/Runtime/RenderCore/Private/GlobalRenderResources.cpp
--- a/Engine/Source/Runtime/RenderCore/Private/GlobalRenderResources.cpp
+++ b/Engine/Source/Runtime/RenderCore/Private/GlobalRenderResources.cpp
@@ -778,7 +778,7 @@ public:
 	{
 		check(!IsValidRef(BufferType::GetRHI()));
 		FRHIResourceCreateInfo CreateInfo(TEXT("FDynamicBuffer"));
-		BufferType::SetRHI(RHICmdList.CreateBuffer(BufferSize, Stride == 0 ? EBufferUsageFlags::VertexBuffer : EBufferUsageFlags::IndexBuffer, Stride, ERHIAccess::VertexOrIndexBuffer, CreateInfo));
+		BufferType::SetRHI(RHICmdList.CreateBuffer(BufferSize, Stride == 0 ? EBufferUsageFlags::VertexBuffer : EBufferUsageFlags::IndexBuffer, Stride, ERHIAccess::VertexOrIndexBuffer | ERHIAccess::SRVCompute, CreateInfo));
 		MappedBuffer = NULL;
 		AllocatedByteCount = 0;
 	}
