RHI Validation, rollback Discard/Acquire check logic

diff --git a/Engine/Source/Runtime/RHI/Private/RHIValidation.cpp b/Engine/Source/Runtime/RHI/Private/RHIValidation.cpp
--- a/Engine/Source/Runtime/RHI/Private/RHIValidation.cpp
+++ b/Engine/Source/Runtime/RHI/Private/RHIValidation.cpp
@@ -2004,7 +2004,7 @@ namespace RHIValidation
 	void FTransientState::Acquire(FResource* Resource, void* CreateTrace)
 	{
 		RHI_VALIDATION_CHECK(bTransient, *GetReasonString_AcquireNonTransient(Resource));
-		RHI_VALIDATION_CHECK(Status != EStatus::Acquired, *GetReasonString_DuplicateAcquireTransient(Resource, AcquireBacktrace, CreateTrace));
+		RHI_VALIDATION_CHECK(Status == EStatus::None, *GetReasonString_DuplicateAcquireTransient(Resource, AcquireBacktrace, CreateTrace));
 		Status = EStatus::Acquired;
 
 		if (!AcquireBacktrace)
@@ -2049,7 +2049,7 @@ namespace RHIValidation
 
 		// Acquire should validate whether ResourceAfter is transient. We assume it is here.
 		RHI_VALIDATION_CHECK(TransientStateBefore.bTransient, *GetReasonString_AliasingOverlapNonTransient(ResourceBefore, ResourceAfter));
-		//RHI_VALIDATION_CHECK(TransientStateBefore.IsDiscarded(), *GetReasonString_AliasingOverlapNonDiscarded(ResourceBefore, ResourceAfter, CreateTrace));
+		RHI_VALIDATION_CHECK(TransientStateBefore.IsDiscarded(), *GetReasonString_AliasingOverlapNonDiscarded(ResourceBefore, ResourceAfter, CreateTrace));
 
 		if (ResourceBefore->LoggingMode != ELoggingMode::None)
 		{
diff --git a/Engine/Source/Runtime/RHI/Public/RHIValidationCommon.h b/Engine/Source/Runtime/RHI/Public/RHIValidationCommon.h
--- a/Engine/Source/Runtime/RHI/Public/RHIValidationCommon.h
+++ b/Engine/Source/Runtime/RHI/Public/RHIValidationCommon.h
@@ -296,7 +296,6 @@ namespace RHIValidation
 
 		FTransientState(ERHIAccess InitialAccess)
 			: bTransient(InitialAccess == ERHIAccess::Discard)
-			, Status((InitialAccess == ERHIAccess::Discard) ? EStatus::Discarded : EStatus::None)
 		{}
 
 		void* AcquireBacktrace = nullptr;
