Fix wrong condition in RHIValidation that was generating false positive while using -rhivalidation

diff --git a/Engine/Source/Runtime/RHI/Private/RHIValidation.cpp b/Engine/Source/Runtime/RHI/Private/RHIValidation.cpp
--- a/Engine/Source/Runtime/RHI/Private/RHIValidation.cpp
+++ b/Engine/Source/Runtime/RHI/Private/RHIValidation.cpp
@@ -2320,7 +2320,7 @@ namespace RHIValidation
 				}
 			}
 
-			if (SubresourceStates.Num() == 0)
+			if (SubresourceStates.Num() != 0)
 			{
 				uint32 LastMip = SubresourceRange.MipIndex + SubresourceRange.NumMips;
 				uint32 LastArraySlice = SubresourceRange.ArraySlice + SubresourceRange.NumArraySlices;
