Add a specific slot in FBoundShaderStateInput for workgraph mesh shaders.
Also adding PLATFORM_SUPPORTS_WORKGRAPH_SHADERS to avoid paying cost on non-windows platforms.

diff --git a/Engine/Source/Runtime/Core/Public/HAL/Platform.h b/Engine/Source/Runtime/Core/Public/HAL/Platform.h
--- a/Engine/Source/Runtime/Core/Public/HAL/Platform.h
+++ b/Engine/Source/Runtime/Core/Public/HAL/Platform.h
@@ -382,6 +382,10 @@
 	#define PLATFORM_SUPPORTS_MESH_SHADERS 0
 #endif
 
+#ifndef PLATFORM_SUPPORTS_WORKGRAPH_SHADERS
+	#define PLATFORM_SUPPORTS_WORKGRAPH_SHADERS 0
+#endif
+
 #ifndef PLATFORM_SUPPORTS_BINDLESS_RENDERING
 	#define PLATFORM_SUPPORTS_BINDLESS_RENDERING 0
 #endif
diff --git a/Engine/Source/Runtime/Core/Public/Windows/WindowsPlatform.h b/Engine/Source/Runtime/Core/Public/Windows/WindowsPlatform.h
--- a/Engine/Source/Runtime/Core/Public/Windows/WindowsPlatform.h
+++ b/Engine/Source/Runtime/Core/Public/Windows/WindowsPlatform.h
@@ -106,6 +106,7 @@ typedef FWindowsPlatformTypes FPlatformTypes;
 #define PLATFORM_SUPPORTS_VIRTUAL_TEXTURE_STREAMING			1
 #define PLATFORM_SUPPORTS_VARIABLE_RATE_SHADING				1
 #define PLATFORM_SUPPORTS_MESH_SHADERS						1
+#define PLATFORM_SUPPORTS_WORKGRAPH_SHADERS					1
 #define PLATFORM_SUPPORTS_BINDLESS_RENDERING				1
 #define PLATFORM_USES__ALIGNED_MALLOC						1
 
diff --git a/Engine/Source/Runtime/RHI/Private/PipelineStateCache.cpp b/Engine/Source/Runtime/RHI/Private/PipelineStateCache.cpp
--- a/Engine/Source/Runtime/RHI/Private/PipelineStateCache.cpp
+++ b/Engine/Source/Runtime/RHI/Private/PipelineStateCache.cpp
@@ -65,6 +65,9 @@ static inline uint32 GetTypeHash(const FBoundShaderStateInput& Input)
 	Hash = HashCombineFast(Hash, GetTypeHash(Input.GetMeshShader()));
 	Hash = HashCombineFast(Hash, GetTypeHash(Input.GetAmplificationShader()));
 #endif
+#if PLATFORM_SUPPORTS_WORKGRAPH_SHADERS
+	Hash = HashCombineFast(Hash, GetTypeHash(Input.GetWorkGraphShader()));
+#endif
 #if PLATFORM_SUPPORTS_GEOMETRY_SHADERS
 	Hash = HashCombineFast(Hash, GetTypeHash(Input.GetGeometryShader()));
 #endif
diff --git a/Engine/Source/Runtime/RHI/Public/RHIResources.h b/Engine/Source/Runtime/RHI/Public/RHIResources.h
--- a/Engine/Source/Runtime/RHI/Public/RHIResources.h
+++ b/Engine/Source/Runtime/RHI/Public/RHIResources.h
@@ -4088,6 +4088,7 @@ struct FBoundShaderStateInput
 		{
 			check(VertexDeclarationRHI == nullptr);
 			check(VertexShaderRHI == nullptr);
+			check(GetWorkGraphShader() == nullptr);
 			GetMeshShader()->AddRef();
 
 			if (GetAmplificationShader())
@@ -4095,6 +4096,13 @@ struct FBoundShaderStateInput
 				GetAmplificationShader()->AddRef();
 			}
 		}
+		else if (GetWorkGraphShader())
+		{
+			check(VertexDeclarationRHI == nullptr);
+			check(VertexShaderRHI == nullptr);
+			check(GetMeshShader() == nullptr);
+			GetWorkGraphShader()->AddRef();
+		}
 		else
 		{
 			check(VertexDeclarationRHI);
@@ -4121,6 +4129,7 @@ struct FBoundShaderStateInput
 		{
 			check(VertexDeclarationRHI == nullptr);
 			check(VertexShaderRHI == nullptr);
+			check(GetWorkGraphShader() == nullptr);
 			GetMeshShader()->Release();
 
 			if (GetAmplificationShader())
@@ -4128,6 +4137,13 @@ struct FBoundShaderStateInput
 				GetAmplificationShader()->Release();
 			}
 		}
+		else if (GetWorkGraphShader())
+		{
+			check(VertexDeclarationRHI == nullptr);
+			check(VertexShaderRHI == nullptr);
+			check(GetMeshShader() == nullptr);
+			GetWorkGraphShader()->Release();
+		}
 		else
 		{
 			check(VertexDeclarationRHI);
@@ -4171,6 +4187,14 @@ struct FBoundShaderStateInput
 	void SetGeometryShader(FRHIGeometryShader*) {}
 #endif
 
+#if PLATFORM_SUPPORTS_WORKGRAPH_SHADERS
+	FRHIWorkGraphShader* GetWorkGraphShader() const { return WorkGraphMeshShaderRHI; }
+	void SetWorkGraphShader(FRHIWorkGraphShader* InWorkGraphMeshShader) { WorkGraphMeshShaderRHI = InWorkGraphMeshShader; }
+#else
+	FRHIWorkGraphShader* GetWorkGraphShader() const { return nullptr; }
+	void SetWorkGraphShader(FRHIWorkGraphShader* InWorkGraphMeshShader) {}
+#endif
+
 	FRHIVertexDeclaration* VertexDeclarationRHI = nullptr;
 	FRHIVertexShader* VertexShaderRHI = nullptr;
 	FRHIPixelShader* PixelShaderRHI = nullptr;
@@ -4182,6 +4206,9 @@ private:
 #if PLATFORM_SUPPORTS_GEOMETRY_SHADERS
 	FRHIGeometryShader* GeometryShaderRHI = nullptr;
 #endif
+#if PLATFORM_SUPPORTS_WORKGRAPH_SHADERS
+	FRHIWorkGraphShader* WorkGraphMeshShaderRHI = nullptr;
+#endif
 };
 
 // Hints for some RHIs that support subpasses
@@ -4333,6 +4360,7 @@ public:
 			BoundShaderState.PixelShaderRHI != rhs.BoundShaderState.PixelShaderRHI ||
 			BoundShaderState.GetMeshShader() != rhs.BoundShaderState.GetMeshShader() ||
 			BoundShaderState.GetAmplificationShader() != rhs.BoundShaderState.GetAmplificationShader() ||
+			BoundShaderState.GetWorkGraphShader() != rhs.BoundShaderState.GetWorkGraphShader() ||
 			BoundShaderState.GetGeometryShader() != rhs.BoundShaderState.GetGeometryShader() ||
 			BlendState != rhs.BlendState ||
 			RasterizerState != rhs.RasterizerState ||
