If a FRHITextureReference is used in a FRHITransitionInfo return the FRHIViewableResource from the referenced texture instead.

diff --git a/Engine/Source/Runtime/RHI/Private/RHITransition.cpp b/Engine/Source/Runtime/RHI/Private/RHITransition.cpp
--- a/Engine/Source/Runtime/RHI/Private/RHITransition.cpp
+++ b/Engine/Source/Runtime/RHI/Private/RHITransition.cpp
@@ -11,9 +11,19 @@ FRHIViewableResource* GetViewableResource(const FRHITransitionInfo& Info)
 	switch (Info.Type)
 	{
 	case FRHITransitionInfo::EType::Buffer:
+		return Info.Buffer;
+
 	case FRHITransitionInfo::EType::Texture:
-		return Info.ViewableResource;
+	{
+		FRHITextureReference* TextureReference = Info.Texture->GetTextureReference(); 
 
+		if (TextureReference)
+		{
+			return TextureReference->GetReferencedTexture();
+		}
+
+		return Info.Texture;
+	}
 	case FRHITransitionInfo::EType::UAV:
 		return Info.UAV ? Info.UAV->GetResource() : nullptr;
 	}
