Fix crash when transitioning a null texture.

diff --git a/Engine/Source/Runtime/RHI/Private/RHITransition.cpp b/Engine/Source/Runtime/RHI/Private/RHITransition.cpp
--- a/Engine/Source/Runtime/RHI/Private/RHITransition.cpp
+++ b/Engine/Source/Runtime/RHI/Private/RHITransition.cpp
@@ -10,22 +10,31 @@ FRHIViewableResource* GetViewableResource(const FRHITransitionInfo& Info)
 {
 	switch (Info.Type)
 	{
-	case FRHITransitionInfo::EType::Buffer:
-		return Info.Buffer;
-
-	case FRHITransitionInfo::EType::Texture:
-	{
-		FRHITextureReference* TextureReference = Info.Texture->GetTextureReference(); 
+		case FRHITransitionInfo::EType::Buffer:
+		{
+			return Info.Buffer;
+		}
 
-		if (TextureReference)
+		case FRHITransitionInfo::EType::Texture:
 		{
-			return TextureReference->GetReferencedTexture();
+			if (!Info.Texture)
+			{
+				return nullptr;
+			}
+
+			FRHITextureReference* TextureReference = Info.Texture->GetTextureReference(); 
+			if (TextureReference)
+			{
+				return TextureReference->GetReferencedTexture();
+			}
+
+			return Info.Texture;
 		}
 
-		return Info.Texture;
-	}
-	case FRHITransitionInfo::EType::UAV:
-		return Info.UAV ? Info.UAV->GetResource() : nullptr;
+		case FRHITransitionInfo::EType::UAV:
+		{
+			return Info.UAV ? Info.UAV->GetResource() : nullptr;
+		}
 	}
 
 	return nullptr;
