Increased max number of bound UAVs to 16 on RHIs that support it
- While devices may support more than 16 UAVs per stage, a variable limit would complicate root signature / layout creation. Since they already have a limit of 16 internally, just bump it for now

diff --git a/Engine/Source/Runtime/D3D12RHI/Private/D3D12Adapter.cpp b/Engine/Source/Runtime/D3D12RHI/Private/D3D12Adapter.cpp
--- a/Engine/Source/Runtime/D3D12RHI/Private/D3D12Adapter.cpp
+++ b/Engine/Source/Runtime/D3D12RHI/Private/D3D12Adapter.cpp
@@ -942,6 +942,15 @@ void FD3D12Adapter::InitializeDevices()
 				MaxNonSamplerDescriptors = D3D12_MAX_SHADER_VISIBLE_DESCRIPTOR_HEAP_SIZE_TIER_2;
 				MaxSamplerDescriptors = D3D12_MAX_SHADER_VISIBLE_SAMPLER_HEAP_SIZE;
 			}
+			
+			if (GetResourceBindingTier() == D3D12_RESOURCE_BINDING_TIER_1)
+			{
+				GRHIGlobals.MaxSimultaneousUAVs = Desc.MaxSupportedFeatureLevel > D3D_FEATURE_LEVEL_11_0 ? MAX_UAVS : 8;
+			}
+			else
+			{
+				GRHIGlobals.MaxSimultaneousUAVs = MAX_UAVS;
+			}
 
 			const bool bRenderDocPresent = D3D12RHI_IsRenderDocPresent(RootDevice);
 
diff --git a/Engine/Source/Runtime/D3D12RHI/Private/D3D12Commands.cpp b/Engine/Source/Runtime/D3D12RHI/Private/D3D12Commands.cpp
--- a/Engine/Source/Runtime/D3D12RHI/Private/D3D12Commands.cpp
+++ b/Engine/Source/Runtime/D3D12RHI/Private/D3D12Commands.cpp
@@ -1435,8 +1435,6 @@ static D3D12_SHADING_RATE_COMBINER ConvertShadingRateCombiner(EVRSRateCombiner I
 
 void FD3D12CommandContext::SetRenderTargetsAndClear(const FRHISetRenderTargetsInfo& RenderTargetsInfo)
 {
-	FRHIUnorderedAccessView* UAVs[MaxSimultaneousUAVs] = {};
-
 	this->SetRenderTargets(RenderTargetsInfo.NumColorRenderTargets,
 		RenderTargetsInfo.ColorRenderTarget,
 		&RenderTargetsInfo.DepthStencilRenderTarget);
diff --git a/Engine/Source/Runtime/RHI/Private/RHIValidation.cpp b/Engine/Source/Runtime/RHI/Private/RHIValidation.cpp
--- a/Engine/Source/Runtime/RHI/Private/RHIValidation.cpp
+++ b/Engine/Source/Runtime/RHI/Private/RHIValidation.cpp
@@ -304,8 +304,7 @@ namespace RHIValidation
 		// The barrier tracking expects us to call Assert() only once per unique resource.
 		// However, multiple UAVs may be bound, all referencing the same resource.
 		// Find the unique resources to ensure we only do the tracking once per resource.
-		uint32 NumUniqueIdentities = 0;
-		FResourceIdentity UniqueIdentities[MaxSimultaneousUAVs];
+		TArray<FResourceIdentity, TInlineAllocator<FRHIGlobals::MinGuaranteedSimultaneousUAVs>> UniqueIdentities;
 
 		for (int32 UAVIndex = 0; UAVIndex < UAVs.Num(); ++UAVIndex)
 		{
@@ -315,15 +314,15 @@ namespace RHIValidation
 
 				// Check if we've already seen this resource.
 				bool bFound = false;
-				for (uint32 Index = 0; !bFound && Index < NumUniqueIdentities; ++Index)
+				for (int32 Index = 0; !bFound && Index < UniqueIdentities.Num(); ++Index)
 				{
 					bFound = UniqueIdentities[Index] == Identity;
 				}
 
 				if (!bFound)
 				{
-					check(NumUniqueIdentities < UE_ARRAY_COUNT(UniqueIdentities));
-					UniqueIdentities[NumUniqueIdentities++] = Identity;
+					check(UniqueIdentities.Num() < GRHIGlobals.MaxSimultaneousUAVs);
+					UniqueIdentities.Add(Identity);
 
 					// Assert unique resources have the required state.
 					BarrierTracker->AddOp(FOperation::Assert(Identity, RequiredState));
diff --git a/Engine/Source/Runtime/RHI/Public/RHIDefinitions.h b/Engine/Source/Runtime/RHI/Public/RHIDefinitions.h
--- a/Engine/Source/Runtime/RHI/Public/RHIDefinitions.h
+++ b/Engine/Source/Runtime/RHI/Public/RHIDefinitions.h
@@ -234,7 +234,10 @@ enum
 static_assert(MaxSimultaneousRenderTargets <= (1 << MaxSimultaneousRenderTargets_NumBits), "MaxSimultaneousRenderTargets will not fit on MaxSimultaneousRenderTargets_NumBits");
 
 /** The number of UAVs that may be simultaneously bound to a shader. */
-enum { MaxSimultaneousUAVs = 8 };
+enum
+{
+	MaxSimultaneousUAVs UE_DEPRECATED(5.6, "Do not use MaxSimultaneousUAVs, this is now device reported by GRHIGlobals.MaxSimultaneousUAVs") = 8
+};
 
 enum class ERHIZBuffer
 {
diff --git a/Engine/Source/Runtime/RHI/Public/RHIGlobals.h b/Engine/Source/Runtime/RHI/Public/RHIGlobals.h
--- a/Engine/Source/Runtime/RHI/Public/RHIGlobals.h
+++ b/Engine/Source/Runtime/RHI/Public/RHIGlobals.h
@@ -340,6 +340,12 @@ struct FRHIGlobals
 
 	int32 MaxTextureSamplers = 16;
 
+	/** The guaranteed minimum number of simultaneously bound UAV's across all stages in a single pipeline */
+	static constexpr int32 MinGuaranteedSimultaneousUAVs = 8;
+
+	/** The maximum number of simultaneously bound UAV's across all stages in a single pipeline */
+	int32 MaxSimultaneousUAVs = MinGuaranteedSimultaneousUAVs;
+
 	/** The maximum work group invocations allowed for compute shader. */
 	TRHIGlobal<int32> MaxWorkGroupInvocations = 1024;
 
diff --git a/Engine/Source/Runtime/RHI/Public/RHIResources.h b/Engine/Source/Runtime/RHI/Public/RHIResources.h
--- a/Engine/Source/Runtime/RHI/Public/RHIResources.h
+++ b/Engine/Source/Runtime/RHI/Public/RHIResources.h
@@ -3933,7 +3933,6 @@ public:
 			bool bClearStencil;
 			bool bClearColor;
 			bool bHasResolveAttachments;
-			FRHIUnorderedAccessView* UnorderedAccessView[MaxSimultaneousUAVs];
 			uint8 MultiViewCount;
 
 			void Set(const FRHISetRenderTargetsInfo& RTInfo)
diff --git a/Engine/Source/Runtime/RHI/Public/RHIValidationCommon.h b/Engine/Source/Runtime/RHI/Public/RHIValidationCommon.h
--- a/Engine/Source/Runtime/RHI/Public/RHIValidationCommon.h
+++ b/Engine/Source/Runtime/RHI/Public/RHIValidationCommon.h
@@ -986,7 +986,7 @@ namespace RHIValidation
 		public:
 			FUAVTracker()
 			{
-				UAVs.Reserve(MaxSimultaneousUAVs);
+				UAVs.Reserve(GRHIGlobals.MinGuaranteedSimultaneousUAVs);
 			}
 
 			inline FRHIUnorderedAccessView*& operator[](int32 Slot)
diff --git a/Engine/Source/Runtime/Renderer/Private/SceneTextureReductions.cpp b/Engine/Source/Runtime/Renderer/Private/SceneTextureReductions.cpp
--- a/Engine/Source/Runtime/Renderer/Private/SceneTextureReductions.cpp
+++ b/Engine/Source/Runtime/Renderer/Private/SceneTextureReductions.cpp
@@ -164,9 +164,8 @@ void BuildHZB(
 	int32 MaxMipBatchSizeFirstPass = MaxMipBatchSize;
 
 	// Avoid overflowing MaxSimultaneousUAVs by reducing the number of levels in the first pass (the only one that writes both near/far textures at once)
-	if (bUseCompute && bReduceClosestDepth && OutFroxelViewData != nullptr && MaxSimultaneousUAVs < 10)
+	if (bUseCompute && bReduceClosestDepth && OutFroxelViewData != nullptr && GRHIGlobals.MaxSimultaneousUAVs < 10)
 	{
-		ensure(MaxSimultaneousUAVs >= 8);
 		// 6 UAVs for the ouput mip levels + 2 from the froxel builder
 		MaxMipBatchSizeFirstPass = 3;
 	}
diff --git a/Engine/Source/Runtime/VulkanRHI/Private/VulkanRHI.cpp b/Engine/Source/Runtime/VulkanRHI/Private/VulkanRHI.cpp
--- a/Engine/Source/Runtime/VulkanRHI/Private/VulkanRHI.cpp
+++ b/Engine/Source/Runtime/VulkanRHI/Private/VulkanRHI.cpp
@@ -921,6 +921,12 @@ void FVulkanDynamicRHI::InitInstance()
 		GRHIMaxDispatchThreadGroupsPerDimension.Y = FMath::Min<uint32>(Limits.maxComputeWorkGroupCount[1], 0x7fffffff);
 		GRHIMaxDispatchThreadGroupsPerDimension.Z = FMath::Min<uint32>(Limits.maxComputeWorkGroupCount[2], 0x7fffffff);
 
+		// Note: While the 2022/2024 profile limits state a minimum of 16, other profiles (even core) go down to 4
+		//       (see https://vulkan.lunarg.com/doc/view/1.3.290.0/windows/profiles_definitions.html)
+		// Since the RHI has historically always supported 8 UAV's, let's leave those specific devices out
+		uint32 DeviceMaxStorageDescriptorPerStage = FMath::Min(Props.limits.maxPerStageDescriptorStorageBuffers, Props.limits.maxPerStageDescriptorStorageImages);
+		GRHIGlobals.MaxSimultaneousUAVs = DeviceMaxStorageDescriptorPerStage >= 16 ? 16 : 8;
+
 		FVulkanPlatform::SetupFeatureLevels(GRHIGlobals.ShaderPlatformForFeatureLevel);
 
 		GRHIRequiresRenderTargetForPixelShaderUAVs = true;
