Add new ERequestedGPUCrash mode (cmd list corruption). Only implemented on certain platforms for the moment. Use "gpucrash cmdlistcorruption" cmd with WITH_GPUDEBUGCRASH. Note, functionality is stripped out in Shipping.

#tests XSX

diff --git a/Engine/Source/Runtime/Engine/Private/UnrealEngine.cpp b/Engine/Source/Runtime/Engine/Private/UnrealEngine.cpp
--- a/Engine/Source/Runtime/Engine/Private/UnrealEngine.cpp
+++ b/Engine/Source/Runtime/Engine/Private/UnrealEngine.cpp
@@ -6282,6 +6282,10 @@ bool UEngine::HandleGPUDebugCrashCommand(const TCHAR* Cmd, FOutputDevice& Ar)
 	{
 		GRHIGlobals.TriggerGPUCrash = ERequestedGPUCrash::Type_Assert;
 	}
+	if (FParse::Command(&Cmd, TEXT("cmdlistcorruption")))
+	{
+		GRHIGlobals.TriggerGPUCrash = ERequestedGPUCrash::Type_CmdListCorruption;
+	}
 	if (FParse::Command(&Cmd, TEXT("hang")) || !bSupportsPageFaultsAndPlatformBreak)
 	{
 		if (!bSupportsPageFaultsAndPlatformBreak)
diff --git a/Engine/Source/Runtime/RHI/Public/RHICommandList.h b/Engine/Source/Runtime/RHI/Public/RHICommandList.h
--- a/Engine/Source/Runtime/RHI/Public/RHICommandList.h
+++ b/Engine/Source/Runtime/RHI/Public/RHICommandList.h
@@ -2188,6 +2188,12 @@ FRHICOMMAND_MACRO(FRHICommandSetDepthBounds)
 	RHI_EXECUTE_API void Execute(FRHICommandListBase& CmdList);
 };
 
+FRHICOMMAND_MACRO(FRHIGpuHangCommandListCorruption)
+{
+	FORCEINLINE_DEBUGGABLE FRHIGpuHangCommandListCorruption(){}
+	RHI_EXECUTE_API void Execute(FRHICommandListBase & CmdList);
+};
+
 FRHICOMMAND_MACRO(FRHICommandSetShadingRate)
 {
 	EVRSShadingRate   ShadingRate;
@@ -4163,6 +4169,16 @@ public:
 		ALLOC_COMMAND(FRHICommandSetDepthBounds)(MinDepth, MaxDepth);
 	}
 	
+	FORCEINLINE_DEBUGGABLE void GpuHangCommandListCorruption()
+	{
+		if (Bypass())
+		{
+			GetContext().RHIGpuHangCommandListCorruption();
+			return;
+		}
+		ALLOC_COMMAND(FRHIGpuHangCommandListCorruption)();
+	}
+
 	FORCEINLINE_DEBUGGABLE void SetShadingRate(EVRSShadingRate ShadingRate, EVRSRateCombiner Combiner)
 	{
 #if PLATFORM_SUPPORTS_VARIABLE_RATE_SHADING
diff --git a/Engine/Source/Runtime/RHI/Public/RHICommandListCommandExecutes.inl b/Engine/Source/Runtime/RHI/Public/RHICommandListCommandExecutes.inl
--- a/Engine/Source/Runtime/RHI/Public/RHICommandListCommandExecutes.inl
+++ b/Engine/Source/Runtime/RHI/Public/RHICommandListCommandExecutes.inl
@@ -345,6 +345,12 @@ void FRHICommandSetDepthBounds::Execute(FRHICommandListBase& CmdList)
 	INTERNAL_DECORATOR(RHISetDepthBounds)(MinDepth, MaxDepth);
 }
 
+void FRHIGpuHangCommandListCorruption::Execute(FRHICommandListBase& CmdList)
+{
+	RHISTAT(GpuHangCommandListCorruption);
+	INTERNAL_DECORATOR(RHIGpuHangCommandListCorruption)();
+}
+
 void FRHICommandClearUAVFloat::Execute(FRHICommandListBase& CmdList)
 {
 	RHISTAT(ClearUAV);
diff --git a/Engine/Source/Runtime/RHI/Public/RHIContext.h b/Engine/Source/Runtime/RHI/Public/RHIContext.h
--- a/Engine/Source/Runtime/RHI/Public/RHIContext.h
+++ b/Engine/Source/Runtime/RHI/Public/RHIContext.h
@@ -800,6 +800,11 @@ public:
 	*/
 	virtual void RHISetDepthBounds(float MinDepth, float MaxDepth) = 0;
 
+	virtual void RHIGpuHangCommandListCorruption()
+	{
+		/* empty default implementation */
+	}
+
 	virtual void RHISetShadingRate(EVRSShadingRate ShadingRate, EVRSRateCombiner Combiner)
 	{
 		/* empty default implementation */
diff --git a/Engine/Source/Runtime/RHI/Public/RHIDefinitions.h b/Engine/Source/Runtime/RHI/Public/RHIDefinitions.h
--- a/Engine/Source/Runtime/RHI/Public/RHIDefinitions.h
+++ b/Engine/Source/Runtime/RHI/Public/RHIDefinitions.h
@@ -1376,9 +1376,10 @@ enum class ERequestedGPUCrash : uint8
 	Type_PageFault = 1 << 1,
 	Type_PlatformBreak = 1 << 2,
 	Type_Assert = 1 << 3,
+	Type_CmdListCorruption = 1 << 4,
 
-	Queue_Direct = 1 << 4,
-	Queue_Compute = 1 << 5,
+	Queue_Direct = 1 << 5,
+	Queue_Compute = 1 << 6,
 };
 ENUM_CLASS_FLAGS(ERequestedGPUCrash);
 
diff --git a/Engine/Source/Runtime/RenderCore/Private/GPUDebugCrashUtils.cpp b/Engine/Source/Runtime/RenderCore/Private/GPUDebugCrashUtils.cpp
--- a/Engine/Source/Runtime/RenderCore/Private/GPUDebugCrashUtils.cpp
+++ b/Engine/Source/Runtime/RenderCore/Private/GPUDebugCrashUtils.cpp
@@ -71,6 +71,17 @@ void ScheduleGPUDebugCrash(FRDGBuilder& GraphBuilder)
 		EnumHasAnyFlags(GRHIGlobals.TriggerGPUCrash, ERequestedGPUCrash::Queue_Compute) ? TEXT("ComputeQueue") : TEXT("DirectQueue"),
 		*CrashTypeString);
 
+	if(EnumHasAnyFlags(GRHIGlobals.TriggerGPUCrash, ERequestedGPUCrash::Type_CmdListCorruption))
+	{ 
+		GraphBuilder.AddPass(
+			RDG_EVENT_NAME("GPUDebugCrashCmdCorruption"),
+			ERDGPassFlags::None,
+			[](FRDGAsyncTask, FRHICommandList& RHICmdList)
+			{
+				RHICmdList.GpuHangCommandListCorruption();
+			});
+	}
+
 	FComputeShaderUtils::AddPass(
 		GraphBuilder,
 		RDG_EVENT_NAME("GPUDebugCrash"),
