Fix GetLanguage to not assert when using a source build cooked data that includes shader platforms not included in the source build.

#rnx

diff --git a/Engine/Source/Runtime/RHI/Public/DataDrivenShaderPlatformInfo.h b/Engine/Source/Runtime/RHI/Public/DataDrivenShaderPlatformInfo.h
--- a/Engine/Source/Runtime/RHI/Public/DataDrivenShaderPlatformInfo.h
+++ b/Engine/Source/Runtime/RHI/Public/DataDrivenShaderPlatformInfo.h
@@ -214,8 +214,11 @@ public:
 
 	static FORCEINLINE_DEBUGGABLE const FName GetLanguage(const FStaticShaderPlatform Platform)
 	{
-		check(IsValid(Platform));
-		return Infos[Platform].Language;
+		if(IsValid(Platform))
+		{
+			return Infos[Platform].Language;
+		}
+		return NAME_None;
 	}
 
 	static FORCEINLINE_DEBUGGABLE const ERHIFeatureLevel::Type GetMaxFeatureLevel(const FStaticShaderPlatform Platform)
