Skip the SRV cache when creating SRV's for ray tracing scenes, since we are currently recreating the ray tracing scene every frame causing it to grow unbounded.

diff --git a/Engine/Source/Runtime/RHI/Private/RHI.cpp b/Engine/Source/Runtime/RHI/Private/RHI.cpp
--- a/Engine/Source/Runtime/RHI/Private/RHI.cpp
+++ b/Engine/Source/Runtime/RHI/Private/RHI.cpp
@@ -1973,11 +1973,15 @@ FRHIUnorderedAccessView* FRHITextureViewCache::GetOrCreateUAV(FRHICommandListBas
 
 FRHIShaderResourceView* FRHIBufferViewCache::GetOrCreateSRV(FRHICommandListBase& RHICmdList, FRHIBuffer* Buffer, const FRHIBufferSRVCreateInfo& SRVCreateInfo)
 {
-	for (const auto& KeyValue : SRVs)
+	// Bypass the cache when using the ray tracing scene, as there is no guarantee the scene will not be recreated, which would grow the cache unbounded.
+	if (!SRVCreateInfo.RayTracingScene)
 	{
-		if (KeyValue.Key == SRVCreateInfo)
+		for (const auto& KeyValue : SRVs)
 		{
-			return KeyValue.Value.GetReference();
+			if (KeyValue.Key == SRVCreateInfo)
+			{
+				return KeyValue.Value.GetReference();
+			}
 		}
 	}
 
@@ -2012,7 +2016,12 @@ FRHIShaderResourceView* FRHIBufferViewCache::GetOrCreateSRV(FRHICommandListBase&
 	FShaderResourceViewRHIRef RHIShaderResourceView = RHICmdList.CreateShaderResourceView(Buffer, CreateDesc);
 
 	FRHIShaderResourceView* View = RHIShaderResourceView.GetReference();
-	SRVs.Emplace(SRVCreateInfo, MoveTemp(RHIShaderResourceView));
+
+	if (!SRVCreateInfo.RayTracingScene)
+	{
+		SRVs.Emplace(SRVCreateInfo, MoveTemp(RHIShaderResourceView));
+	}
+
 	return View;
 }
 
