RHI: New NeedsAllPlanes mode added to RHIAccess tracking for RHIs that need the state of both planes when transitionning DepthStencil textures.

#rnx

diff --git a/Engine/Source/Runtime/RHI/Private/RHICommandList.cpp b/Engine/Source/Runtime/RHI/Private/RHICommandList.cpp
--- a/Engine/Source/Runtime/RHI/Private/RHICommandList.cpp
+++ b/Engine/Source/Runtime/RHI/Private/RHICommandList.cpp
@@ -2223,6 +2223,7 @@ void FRHIComputeCommandList::RayTraceDispatchIndirect(
 	GlobalResourceBindings.Reset();
 }
 
+
 // track all transitions:
 // - if a BeforeAccess is Unknown add the transition to a list for patching before dispatch
 // - save each resource state to use at dispatch time
@@ -2259,7 +2260,7 @@ void FRHICommandListBase::TrackResourceState(TArrayView<const FRHITransition*> T
 			}
 
 			TRHIResourceState<FRHICmdListBaseArrayAllocator>* TrackedResourceState = TrackedResources.Find(TransitionViewableResource);
-			FRHISubresourceRange SubResourceIndex = *TransitionInfo;
+			const FRHISubresourceRange SubResourceIndex = *TransitionInfo;
 
 			// Both Unknown is not a supported case
 			check(!((TransitionInfo->AccessBefore == ERHIAccess::Unknown) && (TransitionInfo->AccessAfter == ERHIAccess::Unknown)));
@@ -2278,7 +2279,7 @@ void FRHICommandListBase::TrackResourceState(TArrayView<const FRHITransition*> T
 					PreviousState = TrackedResourceState->GetSubresourceState(SubResourceIndex);
 				}
 
-				// Note: even if PreviuosState is a valid state we cannot write into TransitionInfo since we are recording the cmdlist
+				// Note: even if PreviousState is a valid state we cannot write into TransitionInfo since we are recording the cmdlist
 				TransitionInfoToPatch.Push(FTransitionPatchInfo(TransitionInfo, PreviousState));
 			}
 	
@@ -2309,6 +2310,26 @@ void FRHICommandListBase::TrackResourceState(TArrayView<const FRHITransition*> T
 					TrackedResources.Add(TransitionViewableResource, TRHIResourceState<FRHICmdListBaseArrayAllocator>(SubResourceIndex, TransitionInfo->AccessAfter, TransitionViewableResource->GetTrackedSubResourceCount(), true));
 				}
 			}
+
+			// Some RHI need all planes of a DepthStencil texture to transition
+			if (NeedsAllPlanes() &&
+				(TransitionInfo->Type == FRHITransitionInfo::EType::Texture) &&
+				!TransitionInfo->IsAllPlaneSlices())
+			{
+				// Only applies to texture formats with a stencil
+				FRHITexture* Texture = (FRHITexture*)TransitionViewableResource;
+				if (IsStencilFormat(TransitionInfo->Texture->GetDesc().Format))
+				{
+					// Set the OtherPlaneState for all single plane transitions, less work than tracking the set planes per subresource
+					const FRHISubresourceRange OtherResourceIndex(
+						SubResourceIndex.MipIndex,
+						SubResourceIndex.ArraySlice,
+						(SubResourceIndex.PlaneSlice == FRHISubresourceRange::kDepthPlaneSlice) ? FRHISubresourceRange::kStencilPlaneSlice : FRHISubresourceRange::kDepthPlaneSlice
+					);
+					const ERHIAccess OtherPlaneState = TrackedResourceState ? TrackedResourceState->GetSubresourceState(OtherResourceIndex) : ERHIAccess::Unknown;
+					PlaneTransitionInfoToPatch.Push(FTransitionPatchInfo(TransitionInfo, OtherPlaneState));
+				}
+			}
 		}
 	}
 }
@@ -2342,7 +2363,9 @@ void FRHICommandListBase::UpdateTransitionsAndResourceState()
 					FRHIViewableResource* TransitionViewableResource = GetViewableResource(*TransitionInfo);
 
 					if (TransitionViewableResource)
+					{
 						ResourceState = TransitionViewableResource->TrackedAccess.GetSubresourceState(SubResourceIndex);
+					}
 
 					TransitionInfo->DiscardAccessBefore = ResourceState;
 				}
@@ -2360,7 +2383,9 @@ void FRHICommandListBase::UpdateTransitionsAndResourceState()
 					FRHIViewableResource* TransitionViewableResource = GetViewableResource(*TransitionInfo);
 
 					if (TransitionViewableResource)
+					{
 						ResourceState = TransitionViewableResource->TrackedAccess.GetSubresourceState(SubResourceIndex);
+					}
 
 					TransitionInfo->AccessBefore = ResourceState;
 				}
@@ -2378,7 +2403,9 @@ void FRHICommandListBase::UpdateTransitionsAndResourceState()
 					FRHIViewableResource* TransitionViewableResource = GetViewableResource(*TransitionInfo);
 
 					if (TransitionViewableResource)
+					{
 						ResourceState = TransitionViewableResource->TrackedAccess.GetSubresourceState(SubResourceIndex);
+					}
 
 					TransitionInfo->AccessAfter = ResourceState;
 				}
@@ -2388,6 +2415,47 @@ void FRHICommandListBase::UpdateTransitionsAndResourceState()
 			check(TransitionInfo->AccessAfter != ERHIAccess::Unknown);
 		}
 
+		if (NeedsAllPlanes())
+		{
+			// Iterate over all PlaneTransitionInfo to patch any missing plane information
+			for (FTransitionPatchInfo& PatchInfo : PlaneTransitionInfoToPatch)
+			{
+				PatchInfo.TransitionInfo->Flags |= EResourceTransitionFlags::FillOtherPlaneState;
+				if (PatchInfo.State != ERHIAccess::Unknown)
+				{
+					PatchInfo.TransitionInfo->OtherPlaneAccess = PatchInfo.State;
+				}
+				else
+				{
+					FRHIViewableResource* TransitionViewableResource = GetViewableResource(*PatchInfo.TransitionInfo);
+					if (TransitionViewableResource)
+					{
+						const FRHISubresourceRange OtherResourceIndex(
+							PatchInfo.TransitionInfo->MipIndex,
+							PatchInfo.TransitionInfo->ArraySlice,
+							(PatchInfo.TransitionInfo->PlaneSlice == FRHISubresourceRange::kDepthPlaneSlice) ? FRHISubresourceRange::kStencilPlaneSlice : FRHISubresourceRange::kDepthPlaneSlice
+						);
+						PatchInfo.TransitionInfo->OtherPlaneAccess = TransitionViewableResource->TrackedAccess.GetSubresourceState(OtherResourceIndex);
+					}
+				}
+			}
+			PlaneTransitionInfoToPatch.Empty();
+
+			// Iterate over all RenderPassInfo to patch any missing plane states
+			for (FRenderPassPatchInfo& PatchInfo : RenderPassInfoToPatch)
+			{
+				if (PatchInfo.State != ERHIAccess::Unknown)
+				{
+					PatchInfo.RenderPassInfo->NopAccess = PatchInfo.State;
+				}
+				else
+				{
+					ResolveNopAccessInPlace(*PatchInfo.RenderPassInfo);
+				}
+			}
+			RenderPassInfoToPatch.Empty();
+		}
+
 		// Iterate over all tracked Resources to update their TrackedState with the latest state
 		for (TPair<FRHIViewableResource*, TRHIResourceState<FRHICmdListBaseArrayAllocator>>& Pair : TrackedResources)
 		{
@@ -2399,4 +2467,4 @@ void FRHICommandListBase::UpdateTransitionsAndResourceState()
 		TransitionInfoToPatch.Empty();
 		TrackedResources.Empty();
 	}
-}
\ No newline at end of file
+}
diff --git a/Engine/Source/Runtime/RHI/Private/RHIStrings.cpp b/Engine/Source/Runtime/RHI/Private/RHIStrings.cpp
--- a/Engine/Source/Runtime/RHI/Private/RHIStrings.cpp
+++ b/Engine/Source/Runtime/RHI/Private/RHIStrings.cpp
@@ -508,6 +508,7 @@ FString GetResourceTransitionFlagsName(EResourceTransitionFlags Flags)
 		case EResourceTransitionFlags::Discard:				return TEXT("Discard");
 		case EResourceTransitionFlags::Clear:				return TEXT("Clear");
 		case EResourceTransitionFlags::IgnoreAfterState:	return TEXT("IgnoreAfterState");
+		case EResourceTransitionFlags::FillOtherPlaneState:	return TEXT("FillOtherPlaneState");
 		}
 	});
 }
diff --git a/Engine/Source/Runtime/RHI/Public/RHICommandList.h b/Engine/Source/Runtime/RHI/Public/RHICommandList.h
--- a/Engine/Source/Runtime/RHI/Public/RHICommandList.h
+++ b/Engine/Source/Runtime/RHI/Public/RHICommandList.h
@@ -1460,7 +1460,10 @@ public:
 		}
 	}
 
-	FORCEINLINE_DEBUGGABLE bool NeedTransitionsPatching() { return (PersistentState.bAllowResourceStateTracking) && (GRHIGlobals.EnableHighLevelTransitionPatching); }
+	FORCEINLINE_DEBUGGABLE bool NeedTransitionsPatching() const { return (PersistentState.bAllowResourceStateTracking) && (GRHIGlobals.EnableHighLevelTransitionPatching); }
+
+	// Returns true if the RHI always needs the state of both planes (depth|stencil) for transitions
+	FORCEINLINE_DEBUGGABLE bool NeedsAllPlanes() const { return GRHIGlobals.NeedsAllPlanes && GRHIGlobals.EnableHighLevelTransitionPatching; }
 
 private:
 	FRHICommandListBase(FPersistentState const& InPersistentState);
@@ -1495,11 +1498,25 @@ private:
 		ERHIAccess			State;
 	};
 
+	struct FRenderPassPatchInfo
+	{
+		FRenderPassPatchInfo(FRHIRenderPassInfo* InRenderPassInfo, ERHIAccess InState) :
+			RenderPassInfo(InRenderPassInfo),
+			State(InState) {}
+
+		FRHIRenderPassInfo* RenderPassInfo;
+		ERHIAccess			State;
+	};
+
 	// Using FRHICmdListBaseArrayAllocator since elements Lifetime for all the following containers are 1 frame, between Record and Submit
 	// Track every resource state change happening during the command list Record
 	TMap<FRHIViewableResource*, TRHIResourceState<FRHICmdListBaseArrayAllocator>, FRHICmdListBaseSetAllocator> TrackedResources;
 	// Transitions with Unknown or Discard (when needed) that will patched at Dispatch time
 	TArray<FTransitionPatchInfo, FRHICmdListBaseArrayAllocator> TransitionInfoToPatch;
+	// Transitions with a single plane of DepthStencil specified where the other plane needs to be resolved with tracking
+	TArray<FTransitionPatchInfo, FRHICmdListBaseArrayAllocator> PlaneTransitionInfoToPatch;
+	// RenderPassInfo with a single plane of DepthStencil specified where the other plane needs to be resolved with tracking
+	TArray<FRenderPassPatchInfo, FRHICmdListBaseArrayAllocator> RenderPassInfoToPatch;
 	
 	// This will scan the array and keep track of transitions that needs patching later at Dispatch time
 	RHI_API void TrackResourceState(TArrayView<const FRHITransition*> Transitions);
@@ -1510,6 +1527,10 @@ private:
 	// Bypass needs to patch Transitions on the spot
 	inline void PatchTransitionsInPlace(FRHITransitionInfo* TransitionInfo, const ERHIPipeline SrcPipeline, const ERHIPipeline DestPipeline);
 	inline void PatchTransitionsInPlace(TArrayView<const FRHITransition*> Transitions);
+
+protected:
+	inline void ResolveNopAccessInPlace(FRHIRenderPassInfo& InInfo);
+	inline void TrackNopAccess(const FRHIRenderPassInfo& InInfo);
 };
 
 #if WITH_RHI_BREADCRUMBS
@@ -1881,10 +1902,10 @@ FRHICOMMAND_MACRO(FRHICommandSetScissorRect)
 
 FRHICOMMAND_MACRO(FRHICommandBeginRenderPass)
 {
-	FRHIRenderPassInfo Info;
+	FRHIRenderPassInfo& Info;
 	const TCHAR* Name;
 
-	FRHICommandBeginRenderPass(const FRHIRenderPassInfo& InInfo, const TCHAR* InName)
+	FRHICommandBeginRenderPass(FRHIRenderPassInfo& InInfo, const TCHAR* InName)
 		: Info(InInfo)
 		, Name(InName)
 	{
@@ -4247,14 +4268,34 @@ public:
 
 		InInfo.Validate();
 
+		// If the ExclusiveDepthStencil contains a Nop, RHIs that need the state of both planes 
+		// will need tracking to resolve the RHIAccess for the missing plane
+		const bool bNeedsNopResolve = NeedsAllPlanes() &&
+			InInfo.DepthStencilRenderTarget.DepthStencilTarget &&
+			(!InInfo.DepthStencilRenderTarget.ExclusiveDepthStencil.IsUsingDepth() ||
+			 !InInfo.DepthStencilRenderTarget.ExclusiveDepthStencil.IsUsingStencil());
+
 		if (Bypass())
 		{
+			if (bNeedsNopResolve)
+			{
+				ResolveNopAccessInPlace(const_cast<FRHIRenderPassInfo&>(InInfo));
+			}
 			GetContext().RHIBeginRenderPass(InInfo, Name);
 		}
 		else
 		{
+			// Copy the transition array into the command list
+			FRHIRenderPassInfo* InfoCopy = (FRHIRenderPassInfo*)Alloc(sizeof(FRHIRenderPassInfo), alignof(FRHIRenderPassInfo));
+			FMemory::Memcpy(InfoCopy, &InInfo, sizeof(InInfo));
+
+			if (bNeedsNopResolve)
+			{
+				TrackNopAccess(*InfoCopy);
+			}
+
 			TCHAR* NameCopy  = AllocString(Name);
-			ALLOC_COMMAND(FRHICommandBeginRenderPass)(InInfo, NameCopy);
+			ALLOC_COMMAND(FRHICommandBeginRenderPass)(*InfoCopy, NameCopy);
 		}
 
 		CacheActiveRenderTargets(InInfo);
diff --git a/Engine/Source/Runtime/RHI/Public/RHICommandList.inl b/Engine/Source/Runtime/RHI/Public/RHICommandList.inl
--- a/Engine/Source/Runtime/RHI/Public/RHICommandList.inl
+++ b/Engine/Source/Runtime/RHI/Public/RHICommandList.inl
@@ -233,7 +233,7 @@ FORCEINLINE_DEBUGGABLE void FRHICommandListBase::PatchTransitionsInPlace(FRHITra
 		return;
 	}
 
-	FRHISubresourceRange SubResourceIndex = *TransitionInfo;
+	const FRHISubresourceRange SubResourceIndex = *TransitionInfo;
 	FRHIViewableResource* TransitionViewableResource = GetViewableResource(*TransitionInfo);
 
 	// Both Unknown is not a supported case
@@ -281,6 +281,25 @@ FORCEINLINE_DEBUGGABLE void FRHICommandListBase::PatchTransitionsInPlace(FRHITra
 		
 		check(TransitionInfo->AccessBefore != ERHIAccess::Unknown);
 		check(TransitionInfo->AccessAfter != ERHIAccess::Unknown);
+
+		// Some RHI need all planes of a DepthStencil texture to transition
+		if (NeedsAllPlanes() && (TransitionInfo->Type == FRHITransitionInfo::EType::Texture) && !SubResourceIndex.IsAllPlaneSlices())
+		{
+			// Only applies to texture formats with a stencil
+			FRHITexture* Texture = (FRHITexture*)TransitionViewableResource;
+			if (IsStencilFormat(TransitionInfo->Texture->GetDesc().Format))
+			{
+				// Set the OtherPlaneState for all single plane transitions, less work than tracking the set planes per subresource
+				const FRHISubresourceRange OtherResourceIndex(
+					SubResourceIndex.MipIndex,
+					SubResourceIndex.ArraySlice,
+					(SubResourceIndex.PlaneSlice == FRHISubresourceRange::kDepthPlaneSlice) ? FRHISubresourceRange::kStencilPlaneSlice : FRHISubresourceRange::kDepthPlaneSlice
+				);
+
+				TransitionInfo->OtherPlaneAccess = TransitionViewableResource->TrackedAccess.GetSubresourceState(OtherResourceIndex);
+				TransitionInfo->Flags |= EResourceTransitionFlags::FillOtherPlaneState;
+			}
+		}
 	}
 }
 
@@ -301,3 +320,44 @@ FORCEINLINE_DEBUGGABLE void FRHICommandListBase::PatchTransitionsInPlace(TArrayV
 		}
 	}
 }
+
+FORCEINLINE_DEBUGGABLE void FRHICommandListBase::ResolveNopAccessInPlace(FRHIRenderPassInfo& InInfo)
+{
+	if (!InInfo.DepthStencilRenderTarget.ExclusiveDepthStencil.IsUsingDepth())
+	{
+		const FRHISubresourceRange SubResourceIndex(0, 0, FRHISubresourceRange::kDepthPlaneSlice);
+		InInfo.NopAccess = InInfo.DepthStencilRenderTarget.DepthStencilTarget->TrackedAccess.GetSubresourceState(SubResourceIndex);
+	}
+	else if (!InInfo.DepthStencilRenderTarget.ExclusiveDepthStencil.IsUsingStencil())
+	{
+		const FRHISubresourceRange SubResourceIndex(0, 0, FRHISubresourceRange::kStencilPlaneSlice);
+		InInfo.NopAccess = InInfo.DepthStencilRenderTarget.DepthStencilTarget->TrackedAccess.GetSubresourceState(SubResourceIndex);
+	}
+}
+
+FORCEINLINE_DEBUGGABLE void FRHICommandListBase::TrackNopAccess(const FRHIRenderPassInfo& InInfo)
+{
+	checkSlow(NeedsAllPlanes());
+	if (InInfo.DepthStencilRenderTarget.DepthStencilTarget)
+	{
+		TRHIResourceState<FRHICmdListBaseArrayAllocator>* TrackedResourceState = TrackedResources.Find(InInfo.DepthStencilRenderTarget.DepthStencilTarget);
+
+		ERHIAccess NopState = ERHIAccess::Unknown;
+		if (TrackedResourceState)
+		{
+			if (!InInfo.DepthStencilRenderTarget.ExclusiveDepthStencil.IsUsingDepth())
+			{
+				const FRHISubresourceRange SubResourceIndex(0, 0, FRHISubresourceRange::kDepthPlaneSlice);
+				NopState = TrackedResourceState->GetSubresourceState(SubResourceIndex);
+			}
+			else if (!InInfo.DepthStencilRenderTarget.ExclusiveDepthStencil.IsUsingStencil())
+			{
+				const FRHISubresourceRange SubResourceIndex(0, 0, FRHISubresourceRange::kStencilPlaneSlice);
+				NopState = TrackedResourceState->GetSubresourceState(SubResourceIndex);
+			}
+		}
+
+		// Note: even if PreviousState is a valid state we cannot write into TransitionInfo since we are recording the cmdlist
+		RenderPassInfoToPatch.Push(FRenderPassPatchInfo((FRHIRenderPassInfo*)&InInfo, NopState));
+	}
+}
diff --git a/Engine/Source/Runtime/RHI/Public/RHIDefinitions.h b/Engine/Source/Runtime/RHI/Public/RHIDefinitions.h
--- a/Engine/Source/Runtime/RHI/Public/RHIDefinitions.h
+++ b/Engine/Source/Runtime/RHI/Public/RHIDefinitions.h
@@ -1362,7 +1362,10 @@ enum class EResourceTransitionFlags
 	MaintainCompression = 1 << 0, // Specifies that the transition should not decompress the resource, allowing us to read a compressed resource directly in its compressed state.
 	Discard				= 1 << 1, // Specifies that the data in the resource should be discarded during the transition - used for transient resource acquire when the resource will be fully overwritten
 	Clear				= 1 << 2, // Specifies that the data in the resource should be cleared during the transition - used for transient resource acquire when the resource might not be fully overwritten
+
+	// Flags below aren't used by RDG and therefore not included in the Last/Mask calculations
 	IgnoreAfterState	= 1 << 3, // The transition will be patched for Unknown if needed but the After State will be ignored by the tracker
+	FillOtherPlaneState = 1 << 4, // The other plane state of a depth|stencil texture has been filled by tracking
 
 	Last = Clear,
 	Mask = (Last << 1) - 1
diff --git a/Engine/Source/Runtime/RHI/Public/RHIGlobals.h b/Engine/Source/Runtime/RHI/Public/RHIGlobals.h
--- a/Engine/Source/Runtime/RHI/Public/RHIGlobals.h
+++ b/Engine/Source/Runtime/RHI/Public/RHIGlobals.h
@@ -713,7 +713,10 @@ struct FRHIGlobals
 
 	/** True if RHI needs the transitions from the high level for COPYSRC/COPYDEST*/
 	bool NeedExtraTransitions = false;
-	
+
+	/** True if RHI always needs the state of both planes (depth|stencil) for transitions */
+	bool NeedsAllPlanes = false;
+
 	/** True if lossy/fixed-rate compression of render targets is supporterd. */
 	bool SupportsLossyFramebufferCompression = false;
 };
diff --git a/Engine/Source/Runtime/RHI/Public/RHIResources.h b/Engine/Source/Runtime/RHI/Public/RHIResources.h
--- a/Engine/Source/Runtime/RHI/Public/RHIResources.h
+++ b/Engine/Source/Runtime/RHI/Public/RHIResources.h
@@ -1326,7 +1326,7 @@ public:
 		TrackedAccess = Access;
 	}
 
-	ERHIAccess GetTrackedAccess_Unsafe()
+	ERHIAccess GetTrackedAccess_Unsafe() const
 	{
 		return TrackedAccess.GetResourceState();
 	}
@@ -4938,6 +4938,9 @@ struct FRHIRenderPassInfo
 	FTextureRHIRef ShadingRateTexture = nullptr;
 	EVRSRateCombiner ShadingRateTextureCombiner = VRSRB_Passthrough;
 
+	// Some RHIs need to know the layout of all planes in the depth target
+	ERHIAccess NopAccess = ERHIAccess::Unknown;
+
 	// Some RHIs require a hint that occlusion queries will be used in this render pass
 	uint32 NumOcclusionQueries = 0;
 	bool bOcclusionQueries = false;
diff --git a/Engine/Source/Runtime/RHI/Public/RHITransition.h b/Engine/Source/Runtime/RHI/Public/RHITransition.h
--- a/Engine/Source/Runtime/RHI/Public/RHITransition.h
+++ b/Engine/Source/Runtime/RHI/Public/RHITransition.h
@@ -139,6 +139,7 @@ struct FRHITransitionInfo : public FRHISubresourceRange
 	ERHIAccess AccessBefore = ERHIAccess::Unknown;
 	ERHIAccess AccessAfter = ERHIAccess::Unknown;
 	ERHIAccess DiscardAccessBefore = ERHIAccess::Unknown;
+	ERHIAccess OtherPlaneAccess = ERHIAccess::Unknown;  // used with GRHIGlobals.NeedsAllPlanes
 	EResourceTransitionFlags Flags = EResourceTransitionFlags::None;
 	TOptional<FRHICommitResourceInfo> CommitInfo;
 
