CIS fixes from fallout due to check() argument reference changes.

#rnx

diff --git a/Engine/Source/Runtime/RHI/Public/RHICommandList.h b/Engine/Source/Runtime/RHI/Public/RHICommandList.h
--- a/Engine/Source/Runtime/RHI/Public/RHICommandList.h
+++ b/Engine/Source/Runtime/RHI/Public/RHICommandList.h
@@ -3312,7 +3312,9 @@ public:
 		check(Breadcrumb && Breadcrumb != FRHIBreadcrumbNode::Sentinel);
 		check(IsSingleRHIPipeline(Pipeline));
 		check(EnumHasAllFlags(ActivePipelines, Pipeline));
+#if DO_CHECK
 		check(!EnumHasAnyFlags(ERHIPipeline(Breadcrumb->BeginPipes.fetch_or(std::underlying_type_t<ERHIPipeline>(Pipeline))), Pipeline));
+#endif
 
 		BreadcrumbAllocatorRefs.AddUnique(Breadcrumb->Allocator);
 
@@ -3339,7 +3341,9 @@ public:
 		check(Breadcrumb && Breadcrumb != FRHIBreadcrumbNode::Sentinel);
 		check(IsSingleRHIPipeline(Pipeline));
 		check(EnumHasAllFlags(ActivePipelines, Pipeline));
+#if DO_CHECK
 		check(!EnumHasAnyFlags(ERHIPipeline(Breadcrumb->EndPipes.fetch_or(std::underlying_type_t<ERHIPipeline>(Pipeline))), Pipeline));
+#endif
 
 		BreadcrumbAllocatorRefs.AddUnique(Breadcrumb->Allocator);
 
diff --git a/Engine/Source/Runtime/RHI/Public/RHICommandList.inl b/Engine/Source/Runtime/RHI/Public/RHICommandList.inl
--- a/Engine/Source/Runtime/RHI/Public/RHICommandList.inl
+++ b/Engine/Source/Runtime/RHI/Public/RHICommandList.inl
@@ -179,7 +179,9 @@ inline FRHIResourceReplaceBatcher::~FRHIResourceReplaceBatcher()
 		, ThreadId(FPlatformTLS::GetCurrentThreadId())
 	#endif
 	{
+	#if DO_CHECK
 		check(Pipeline != ERHIPipeline::None);
+	#endif
 
 		Node->SetParent(RHICmdList.PersistentState.LocalBreadcrumb);
 		RHICmdList.BeginBreadcrumbCPU(Node.Get(), true);
@@ -189,10 +191,12 @@ inline FRHIResourceReplaceBatcher::~FRHIResourceReplaceBatcher()
 	inline void FRHIBreadcrumbEventManual::End(FRHIComputeCommandList& RHICmdList)
 	{
 		checkf(Node, TEXT("Manual breadcrumb was already ended."));
+	#if DO_CHECK
 		checkf(Pipeline == RHICmdList.GetPipeline(), TEXT("Manual breadcrumb was started and ended on different pipelines. Start: %s, End: %s")
 			, *GetRHIPipelineName(Pipeline)
 			, *GetRHIPipelineName(RHICmdList.GetPipeline())
 		);
+	#endif
 
 		checkf(ThreadId == FPlatformTLS::GetCurrentThreadId(), TEXT("Manual breadcrumbs must be started and ended on the same thread."));
 
