VulkanRHI: Only update tracking in EndTransitions when running with NeedsAllPlanes since it's possible to begin a RenderPass in between Begin and End.

#rnx

diff --git a/Engine/Source/Runtime/RHI/Private/RHICommandList.cpp b/Engine/Source/Runtime/RHI/Private/RHICommandList.cpp
--- a/Engine/Source/Runtime/RHI/Private/RHICommandList.cpp
+++ b/Engine/Source/Runtime/RHI/Private/RHICommandList.cpp
@@ -2101,7 +2101,9 @@ void FRHICommandListBase::TransitionInternal(TConstArrayView<FRHITransitionInfo>
 
 			if (NeedTransitionsPatching())
 			{
-				PatchTransitionsInPlace((FRHITransitionInfo*)&LocalTransitionInfo[i], Pipeline, Pipeline);
+				// Because we know the Begin and End are executed consecutively, we don't have to worry about RenderPasses in between for NeedAllPlanes
+				const bool bForceTrackingUpdate = true;
+				PatchTransitionsInPlace((FRHITransitionInfo*)&LocalTransitionInfo[i], Pipeline, Pipeline, bForceTrackingUpdate);
 			}
 		}
 
@@ -2122,7 +2124,8 @@ void FRHICommandListBase::TransitionInternal(TConstArrayView<FRHITransitionInfo>
 	else
 	{
 		FRHITransition* Transition = RHICreateTransition(FRHITransitionCreateInfo(Pipeline, Pipeline, CreateFlags | ERHITransitionCreateFlags::NoSplit, Infos));
-		TrackResourceState(MakeArrayView((const FRHITransition**)&Transition, 1));
+		TrackResourceStateBegin(MakeArrayView((const FRHITransition**)&Transition, 1));
+		TrackResourceStateEnd(MakeArrayView((const FRHITransition**)&Transition, 1));
 		ALLOC_COMMAND(FRHICommandResourceTransition)(Transition);
 	}
 }
@@ -2229,11 +2232,13 @@ void FRHIComputeCommandList::RayTraceDispatchIndirect(
 // track all transitions:
 // - if a BeforeAccess is Unknown add the transition to a list for patching before dispatch
 // - save each resource state to use at dispatch time
-void FRHICommandListBase::TrackResourceState(TArrayView<const FRHITransition*> Transitions)
+void FRHICommandListBase::TrackResourceStateBegin(TArrayView<const FRHITransition*> Transitions)
 {
 	// If cmdlist tracking is off or RHI can handle Unknown resource, skip the tracking
 	if (NeedTransitionsPatching() == false)
+	{
 		return;
+	}
 
 	for (const FRHITransition* Transition : Transitions)
 	{
@@ -2299,7 +2304,8 @@ void FRHICommandListBase::TrackResourceState(TArrayView<const FRHITransition*> T
 			const bool bShouldSkipAfterState = ((EnumHasAnyFlags(TransitionInfo->Flags, EResourceTransitionFlags::IgnoreAfterState)) || (TransitionInfo->AccessAfter == ERHIAccess::Unknown));
 			const bool bNeedUpdateAfterState = (bShouldSkipAfterState == false) && (bDiscardOnAsync == false);
 
-			if (bNeedUpdateAfterState)
+			const bool bTrackOnEnd = NeedsAllPlanes() && EnumHasAnyFlags(TransitionCreateInfoData->SrcPipelines, ERHIPipeline::Graphics) && EnumHasAnyFlags(TransitionCreateInfoData->DstPipelines, ERHIPipeline::Graphics);
+			if (bNeedUpdateAfterState && !bTrackOnEnd)
 			{
 				// Update the TrackedResource State
 				if (TrackedResourceState)
@@ -2336,6 +2342,61 @@ void FRHICommandListBase::TrackResourceState(TArrayView<const FRHITransition*> T
 	}
 }
 
+void FRHICommandListBase::TrackResourceStateEnd(TArrayView<const FRHITransition*> Transitions)
+{
+	if (!NeedTransitionsPatching() || !NeedsAllPlanes())
+	{
+		return;
+	}
+
+	for (const FRHITransition* Transition : Transitions)
+	{
+		const FRHITransitionCreateInfoData* TransitionCreateInfoData = Transition->CreationInfoData;
+
+		if (TransitionCreateInfoData)
+		{
+			const bool bTrackOnEnd = EnumHasAnyFlags(TransitionCreateInfoData->SrcPipelines, ERHIPipeline::Graphics) && EnumHasAnyFlags(TransitionCreateInfoData->DstPipelines, ERHIPipeline::Graphics);
+			if (bTrackOnEnd)
+			{
+				for (int32 Index = 0; Index < TransitionCreateInfoData->TransitionInfos.Num(); Index++)
+				{
+					const FRHITransitionInfo* TransitionInfo = &TransitionCreateInfoData->TransitionInfos[Index];
+
+					if ((TransitionInfo->Type == FRHITransitionInfo::EType::Buffer) && (GRHIGlobals.SupportImplicitBufferStateTracking))
+					{
+						// This RHI supports Buffers state tracking at platform level, so we don't need to track them
+						continue;
+					}
+
+					FRHIViewableResource* TransitionViewableResource = GetViewableResource(*TransitionInfo);
+					if (TransitionViewableResource == nullptr)
+					{
+						continue;
+					}
+
+					TRHIResourceState<FRHICmdListBaseArrayAllocator>* TrackedResourceState = TrackedResources.Find(TransitionViewableResource);
+
+					const bool bShouldSkipAfterState = ((EnumHasAnyFlags(TransitionInfo->Flags, EResourceTransitionFlags::IgnoreAfterState)) || (TransitionInfo->AccessAfter == ERHIAccess::Unknown));
+					const bool bNeedUpdateAfterState = (bShouldSkipAfterState == false);
+					if (bNeedUpdateAfterState)
+					{
+						// Update the TrackedResource State
+						if (TrackedResourceState)
+						{
+							TrackedResourceState->SetSubresourceState(*TransitionInfo, TransitionInfo->AccessAfter);
+						}
+						else
+						{
+							// Add a new element
+							TrackedResources.Add(TransitionViewableResource, TRHIResourceState<FRHICmdListBaseArrayAllocator>(*TransitionInfo, TransitionInfo->AccessAfter, TransitionViewableResource->GetTrackedSubResourceCount(), true));
+						}
+					}
+				}
+			}
+		}
+	}
+}
+
 void FRHICommandListBase::UpdateTransitionsAndResourceState()
 {
 	TRACE_CPUPROFILER_EVENT_SCOPE(UpdateTransitionsAndResourceState);
diff --git a/Engine/Source/Runtime/RHI/Private/RHITransition.cpp b/Engine/Source/Runtime/RHI/Private/RHITransition.cpp
--- a/Engine/Source/Runtime/RHI/Private/RHITransition.cpp
+++ b/Engine/Source/Runtime/RHI/Private/RHITransition.cpp
@@ -48,14 +48,18 @@ void FRHITransition::CleanupInfoData() const
 	CreationInfoData = nullptr;
 }
 
-void FRHITransition::CreatePlatformTransition() const
+void FRHITransition::CreatePlatformTransition(bool bCleanupInfoData) const
 {
 	checkSlow(GRHIGlobals.EnableHighLevelTransitionPatching);
 	if (IsPlatformLevelCreated() == false)
 	{
 		GDynamicRHI->RHICreateTransition(const_cast<FRHITransition*>(this), CreationInfoData->GetCreationInfo());
 		EnumAddFlags(Flags, ETransitionFlag::PlatformLevelCreated);
-		CleanupInfoData();
+
+		if (bCleanupInfoData)
+		{
+			CleanupInfoData();
+		}
 	}
 }
 
@@ -65,12 +69,8 @@ void FRHITransition::CreatePlatformTransitionThreadSafe() const
 	
 	UE::TUniqueLock Lock(TransitionCreationMutex);
 
-	if (IsPlatformLevelCreated() == false)
-	{
-		GDynamicRHI->RHICreateTransition(const_cast<FRHITransition*>(this), CreationInfoData->GetCreationInfo());
-		EnumAddFlags(Flags, ETransitionFlag::PlatformLevelCreated);
-		CleanupInfoData();
-	}
+	const bool bCleanupInfoData = true;
+	CreatePlatformTransition(bCleanupInfoData);
 }
 
 void FRHITransition::Cleanup() const
diff --git a/Engine/Source/Runtime/RHI/Public/RHICommandList.h b/Engine/Source/Runtime/RHI/Public/RHICommandList.h
--- a/Engine/Source/Runtime/RHI/Public/RHICommandList.h
+++ b/Engine/Source/Runtime/RHI/Public/RHICommandList.h
@@ -1519,14 +1519,16 @@ private:
 	TArray<FRenderPassPatchInfo, FRHICmdListBaseArrayAllocator> RenderPassInfoToPatch;
 	
 	// This will scan the array and keep track of transitions that needs patching later at Dispatch time
-	RHI_API void TrackResourceState(TArrayView<const FRHITransition*> Transitions);
-	
+	RHI_API void TrackResourceStateBegin(TArrayView<const FRHITransition*> Transitions);
+	RHI_API void TrackResourceStateEnd(TArrayView<const FRHITransition*> Transitions);
+
 	// Called at Dispatch time, in order of submission, it's when the Unknown and Discard transitions can be patched with the real state
 	void UpdateTransitionsAndResourceState();
 	
 	// Bypass needs to patch Transitions on the spot
-	inline void PatchTransitionsInPlace(FRHITransitionInfo* TransitionInfo, const ERHIPipeline SrcPipeline, const ERHIPipeline DestPipeline);
+	inline void PatchTransitionsInPlace(FRHITransitionInfo* TransitionInfo, const ERHIPipeline SrcPipeline, const ERHIPipeline DestPipeline, bool bForceTrackingUpdate=false);
 	inline void PatchTransitionsInPlace(TArrayView<const FRHITransition*> Transitions);
+	inline void UpdateTrackedAccessInPlace(TArrayView<const FRHITransition*> Transitions);
 
 protected:
 	inline void ResolveNopAccessInPlace(FRHIRenderPassInfo& InInfo);
@@ -2960,7 +2962,8 @@ public:
 			{
 				for (const FRHITransition* Transition : Transitions)
 				{
-					Transition->CreatePlatformTransition();
+					const bool bCleanupInfoData = !NeedsAllPlanes();
+					Transition->CreatePlatformTransition(bCleanupInfoData);
 				}
 			}
 
@@ -2973,7 +2976,8 @@ public:
 		}
 		else
 		{
-			TrackResourceState(Transitions);
+			TrackResourceStateBegin(Transitions);
+
 			// Copy the transition array into the command list
 			FRHITransition** DstTransitionArray = (FRHITransition**)Alloc(sizeof(FRHITransition*) * Transitions.Num(), alignof(FRHITransition*));
 			FMemory::Memcpy(DstTransitionArray, Transitions.GetData(), sizeof(FRHITransition*) * Transitions.Num());
@@ -2996,6 +3000,8 @@ public:
 				checkSlow(Transition->IsPlatformLevelCreated());
 			}
 #endif
+			UpdateTrackedAccessInPlace(Transitions);
+
 			GetComputeContext().RHIEndTransitions(Transitions);
 
 			for (const FRHITransition* Transition : Transitions)
@@ -3005,6 +3011,8 @@ public:
 		}
 		else
 		{
+			TrackResourceStateEnd(Transitions);
+
 			// Copy the transition array into the command list
 			FRHITransition** DstTransitionArray = (FRHITransition**)Alloc(sizeof(FRHITransition*) * Transitions.Num(), alignof(FRHITransition*));
 			FMemory::Memcpy(DstTransitionArray, Transitions.GetData(), sizeof(FRHITransition*) * Transitions.Num());
diff --git a/Engine/Source/Runtime/RHI/Public/RHICommandList.inl b/Engine/Source/Runtime/RHI/Public/RHICommandList.inl
--- a/Engine/Source/Runtime/RHI/Public/RHICommandList.inl
+++ b/Engine/Source/Runtime/RHI/Public/RHICommandList.inl
@@ -229,7 +229,58 @@ inline void TRHILambdaCommandMultiPipe<RHICmdListType, LAMBDA>::ExecuteAndDestru
 	Lambda.~LAMBDA();
 }
 
-FORCEINLINE_DEBUGGABLE void FRHICommandListBase::PatchTransitionsInPlace(FRHITransitionInfo* TransitionInfo, const ERHIPipeline SrcPipeline, const ERHIPipeline DestPipeline)
+FORCEINLINE_DEBUGGABLE void FRHICommandListBase::UpdateTrackedAccessInPlace(TArrayView<const FRHITransition*> Transitions)
+{
+	check(Bypass());
+
+	if (!NeedsAllPlanes())
+	{
+		return;
+	}
+
+	for (const FRHITransition* Transition : Transitions)
+	{
+		FRHITransitionCreateInfoData* TransitionCreateInfoData = Transition->CreationInfoData;
+		if (TransitionCreateInfoData)
+		{
+			const ERHIPipeline SrcPipeline = TransitionCreateInfoData->SrcPipelines;
+			const ERHIPipeline DestPipeline = TransitionCreateInfoData->DstPipelines;
+
+			if (SrcPipeline != DestPipeline)
+			{
+				continue;
+			}
+
+			for (int32 i = 0; i < TransitionCreateInfoData->TransitionInfos.Num(); i++)
+			{
+				const FRHITransitionInfo& TransitionInfo = TransitionCreateInfoData->TransitionInfos[i];
+
+				if ((TransitionInfo.Type == FRHITransitionInfo::EType::Buffer) && (GRHIGlobals.SupportImplicitBufferStateTracking))
+				{
+					// This RHI supports Buffers state tracking at platform level, so we don't need to track them
+					continue;
+				}
+
+				FRHIViewableResource* TransitionViewableResource = GetViewableResource(TransitionInfo);
+				if (TransitionViewableResource)
+				{
+					const bool bShouldSkipAfterState = ((EnumHasAnyFlags(TransitionInfo.Flags, EResourceTransitionFlags::IgnoreAfterState)) || (TransitionInfo.AccessAfter == ERHIAccess::Unknown));
+					const bool bNeedUpdateAfterState = (bShouldSkipAfterState == false);
+
+					if (bNeedUpdateAfterState)
+					{
+						TransitionViewableResource->TrackedAccess.SetSubresourceState(TransitionInfo, TransitionInfo.AccessAfter);
+						TransitionViewableResource->TrackedAccess.CheckAllSubresourceSame();
+					}
+				}
+			}
+
+			Transition->CleanupInfoData();
+		}
+	}
+}
+
+FORCEINLINE_DEBUGGABLE void FRHICommandListBase::PatchTransitionsInPlace(FRHITransitionInfo* TransitionInfo, const ERHIPipeline SrcPipeline, const ERHIPipeline DestPipeline, bool bForceTrackingUpdate)
 {
 	if ((TransitionInfo->Type == FRHITransitionInfo::EType::Buffer) && (GRHIGlobals.SupportImplicitBufferStateTracking))
 	{
@@ -277,7 +328,11 @@ FORCEINLINE_DEBUGGABLE void FRHICommandListBase::PatchTransitionsInPlace(FRHITra
 			TransitionInfo->AccessAfter = TransitionViewableResource->TrackedAccess.GetSubresourceState(SubResourceIndex);
 		}
 
-		if (bNeedUpdateAfterState)
+		// When NeedsAllPlanes is active, we might have a RenderPass querying TrackedAccess between the Begin/End if we remain on the Graphics pipe
+		const bool bCanUpdate = (bForceTrackingUpdate || !NeedsAllPlanes() || 
+			!EnumHasAnyFlags(SrcPipeline, ERHIPipeline::Graphics) || 
+			!EnumHasAnyFlags(DestPipeline, ERHIPipeline::Graphics));
+		if (bNeedUpdateAfterState && bCanUpdate)
 		{
 			TransitionViewableResource->TrackedAccess.SetSubresourceState(SubResourceIndex, TransitionInfo->AccessAfter);
 			TransitionViewableResource->TrackedAccess.CheckAllSubresourceSame();
diff --git a/Engine/Source/Runtime/RHI/Public/RHITransition.h b/Engine/Source/Runtime/RHI/Public/RHITransition.h
--- a/Engine/Source/Runtime/RHI/Public/RHITransition.h
+++ b/Engine/Source/Runtime/RHI/Public/RHITransition.h
@@ -584,7 +584,9 @@ public:
 		if ((SubresourceIndex.IsWholeResource()) || (SubResourceCount == 1))
 		{
 			if (SubresourceIndex.IsWholeResource() == false)
+			{
 				check(SubresourceIndex.IsZeroIndexResource());
+			}
 
 			bAllSubresourcesSame = 1;
 			ResourceState = State;
@@ -859,7 +861,9 @@ private:
 	{}
 
 	~FRHITransition()
-	{}
+	{
+		checkSlow(CreationInfoData == nullptr);
+	}
 
 	// Give private access to specific functions/RHI commands that need to allocate or control transitions.
 	friend RHI_API FRHITransition* RHICreateTransition(const FRHITransitionCreateInfo& CreateInfo);
@@ -880,7 +884,7 @@ private:
 		return FMath::Max((uint64)alignof(FRHITransition), GRHITransitionPrivateData_AlignInBytes);
 	}
 	
-	RHI_API void CreatePlatformTransition() const;
+	RHI_API void CreatePlatformTransition(bool bCleanupInfoData) const;
 	RHI_API void CreatePlatformTransitionThreadSafe() const;
 
 	inline void MarkBegin(ERHIPipeline Pipeline) const
