Improvements to stat gpu
- Changed the structure of FSignalFence and FWaitFence events. We now assume each GPU queue has exactly one fence, which is signalled sequentially. This is necessary to simplify the tracking of idle time on the GPU.
- Also changed the event stream iterator to allow pausing / resuming. This is required whenever we see an FWaitFence event, where we've not yet see the corresponding FSignalFence event.
- Added "Wait" and "Idle" stat groups to show time the GPU is waiting on other queues, or idle because it has completed all work and is waiting for more from the CPU.

Modified how fence completion values are calculated in D3D12 RHI, so that submitted fences are always signalled sequentially.
- Necessary for the new GPU profiler's tracking.

diff --git a/Engine/Source/Runtime/D3D12RHI/Private/D3D12Device.cpp b/Engine/Source/Runtime/D3D12RHI/Private/D3D12Device.cpp
--- a/Engine/Source/Runtime/D3D12RHI/Private/D3D12Device.cpp
+++ b/Engine/Source/Runtime/D3D12RHI/Private/D3D12Device.cpp
@@ -36,8 +36,10 @@ static uint32 GetQueryHeapPoolIndex(D3D12_QUERY_HEAP_TYPE HeapType)
 FD3D12Queue::FD3D12Queue(FD3D12Device* Device, ED3D12QueueType QueueType)
 	: Device(Device)
 	, QueueType(QueueType)
+	, Fence(this)
 	, BarrierTimestamps(Device, QueueType, D3D12_QUERY_TYPE_TIMESTAMP)
 	, bSupportsTileMapping(FD3D12DynamicRHI::GetD3DRHI()->QueueSupportsTileMapping(QueueType))
+	, ExecuteCommandListsFence(this)
 {
 	FD3D12Adapter* Adapter = Device->GetParentAdapter();
 
@@ -96,6 +98,7 @@ UE::RHI::GPUProfiler::FQueue FD3D12Queue::GetProfilerQueue() const
 FD3D12Device::FD3D12Device(FRHIGPUMask InGPUMask, FD3D12Adapter* InAdapter)
 	: FD3D12SingleNodeGPUObject(InGPUMask)
 	, FD3D12AdapterChild       (InAdapter)
+	, TileMappingFence         (nullptr)
 #if (RHI_NEW_GPU_PROFILER == 0)
 	, GPUProfilingData         (this)
 #endif
diff --git a/Engine/Source/Runtime/D3D12RHI/Private/D3D12RHI.cpp b/Engine/Source/Runtime/D3D12RHI/Private/D3D12RHI.cpp
--- a/Engine/Source/Runtime/D3D12RHI/Private/D3D12RHI.cpp
+++ b/Engine/Source/Runtime/D3D12RHI/Private/D3D12RHI.cpp
@@ -609,8 +609,8 @@ void FD3D12DynamicRHI::RHIEndFrame(const FRHIEndFrameArgs& Args)
 			break;
 		}
 
-		Payload.EventStream.Emplace<UE::RHI::GPUProfiler::FEvent::FFrameBoundary>(
-			Args.FrameNumber
+		// CPU timestamp for the frame boundary event is filled in by the submission thread
+		Payload.EndFrameEvent = UE::RHI::GPUProfiler::FEvent::FFrameBoundary(0, Args.FrameNumber
 		#if WITH_RHI_BREADCRUMBS
 			, (Pipeline != ERHIPipeline::None) ? Args.GPUBreadcrumbs[Pipeline] : nullptr
 		#endif
diff --git a/Engine/Source/Runtime/D3D12RHI/Private/D3D12Submission.cpp b/Engine/Source/Runtime/D3D12RHI/Private/D3D12Submission.cpp
--- a/Engine/Source/Runtime/D3D12RHI/Private/D3D12Submission.cpp
+++ b/Engine/Source/Runtime/D3D12RHI/Private/D3D12Submission.cpp
@@ -669,20 +669,27 @@ uint64 FD3D12Queue::FinalizePayload(bool bRequiresSignal, FPayloadArray& Payload
 
 	// Keep the latest fence value in the submitted payload.
 	// The interrupt thread uses this to determine when work has completed.
-	PayloadToSubmit->CompletionFenceValue = ++Fence.NextCompletionValue;
-	PayloadToSubmit->bAlwaysSignal |= bRequiresSignal;
+	uint64 NextCompletionValue = Fence.NextCompletionValue;
 
 	// Set the fence/value pair into any sync points we need to signal.
 	for (FD3D12SyncPointRef& SyncPoint : PayloadToSubmit->SyncPointsToSignal)
 	{
 		check(!SyncPoint->ResolvedFence.IsSet());
-		SyncPoint->ResolvedFence.Emplace(Fence, PayloadToSubmit->CompletionFenceValue);
+		SyncPoint->ResolvedFence.Emplace(Fence, NextCompletionValue);
+	}
+
+	PayloadToSubmit->CompletionFenceValue = NextCompletionValue;
+	PayloadToSubmit->bAlwaysSignal |= bRequiresSignal;
+
+	if (PayloadToSubmit->bAlwaysSignal || PayloadToSubmit->SyncPointsToSignal.Num())
+	{
+		++Fence.NextCompletionValue;
 	}
 
 	PayloadsToHandDown.Add(PayloadToSubmit);
 	PayloadToSubmit = nullptr;
 
-	return Fence.NextCompletionValue;
+	return NextCompletionValue;
 }
 
 void FD3D12DynamicRHI::UpdateReservedResources(FD3D12Payload* Payload)
@@ -734,11 +741,10 @@ void FD3D12DynamicRHI::FlushBatchedPayloads(FD3D12Queue::FPayloadArray& Payloads
 		for (auto& [LocalFence, Value] : Payload->QueueFencesToWait)
 		{
 		#if RHI_NEW_GPU_PROFILER
-			// Use the raw ID3D12Fence pointer as the fence's unique ID, since we never release these.
 			Payload->EventStream.Emplace<UE::RHI::GPUProfiler::FEvent::FWaitFence>(
 				  FPlatformTime::Cycles64()
-				, reinterpret_cast<uint64>(LocalFence.D3DFence.GetReference())
 				, Value
+				, LocalFence.OwnerQueue->GetProfilerQueue()
 			);
 		#endif
 
@@ -849,10 +855,8 @@ void FD3D12DynamicRHI::FlushBatchedPayloads(FD3D12Queue::FPayloadArray& Payloads
 			check(Queue.Fence.LastSignaledValue < Payload->CompletionFenceValue);
 
 		#if RHI_NEW_GPU_PROFILER
-			// Use the raw ID3D12Fence pointer as the fence's unique ID, since we never release these.
 			Payload->EventStream.Emplace<UE::RHI::GPUProfiler::FEvent::FSignalFence>(
 				  FPlatformTime::Cycles64()
-				, reinterpret_cast<uint64>(Queue.Fence.D3DFence.GetReference())
 				, Payload->CompletionFenceValue
 			);
 		#endif
@@ -861,6 +865,14 @@ void FD3D12DynamicRHI::FlushBatchedPayloads(FD3D12Queue::FPayloadArray& Payloads
 			Queue.Fence.LastSignaledValue.store(Payload->CompletionFenceValue, std::memory_order_release);
 		}
 
+	#if RHI_NEW_GPU_PROFILER
+		if (Payload->EndFrameEvent.IsSet())
+		{
+			Payload->EndFrameEvent->CPUTimestamp = FPlatformTime::Cycles64();
+			Payload->EventStream.Emplace<UE::RHI::GPUProfiler::FEvent::FFrameBoundary>(*Payload->EndFrameEvent);
+		}
+	#endif
+
 		// Submission of this payload is completed. Signal the submission event if one was provided.
 		if (Payload->SubmissionEvent)
 		{
diff --git a/Engine/Source/Runtime/D3D12RHI/Private/D3D12Submission.h b/Engine/Source/Runtime/D3D12RHI/Private/D3D12Submission.h
--- a/Engine/Source/Runtime/D3D12RHI/Private/D3D12Submission.h
+++ b/Engine/Source/Runtime/D3D12RHI/Private/D3D12Submission.h
@@ -40,10 +40,15 @@ enum class ED3D12SyncPointType
 // Fence type used by the device queues to manage GPU completion
 struct FD3D12Fence
 {
+	FD3D12Queue* const OwnerQueue;
 	TRefCountPtr<ID3D12Fence> D3DFence;
-	uint64 NextCompletionValue = 0;
+	uint64 NextCompletionValue = 1;
 	std::atomic<uint64> LastSignaledValue = 0;
 	bool bInterruptAwaited = false;
+
+	FD3D12Fence(FD3D12Queue* OwnerQueue)
+		: OwnerQueue(OwnerQueue)
+	{}
 };
 
 // Used by FD3D12SyncPoint and the submission thread to fix up signaled fence values at the end-of-pipe
@@ -241,6 +246,7 @@ struct FD3D12PayloadBase : public FD3D12PayloadBaseFixLayout
 
 #if RHI_NEW_GPU_PROFILER
 	UE::RHI::GPUProfiler::FEventStream EventStream;
+	TOptional<UE::RHI::GPUProfiler::FEvent::FFrameBoundary> EndFrameEvent;
 #endif
 
 	virtual ~FD3D12PayloadBase();
@@ -269,7 +275,11 @@ struct FD3D12PayloadBase : public FD3D12PayloadBaseFixLayout
 
 	virtual bool HasSignalWork() const
 	{
-		return RequiresQueueFenceSignal() || ManualFencesToSignal.Num() > 0 || SubmissionEvent != nullptr;
+		return RequiresQueueFenceSignal() || ManualFencesToSignal.Num() > 0 || SubmissionEvent != nullptr
+		#if RHI_NEW_GPU_PROFILER
+			|| EndFrameEvent.IsSet()
+		#endif
+		;
 	}
 
 protected:
diff --git a/Engine/Source/Runtime/RHI/Private/GPUProfiler.cpp b/Engine/Source/Runtime/RHI/Private/GPUProfiler.cpp
--- a/Engine/Source/Runtime/RHI/Private/GPUProfiler.cpp
+++ b/Engine/Source/Runtime/RHI/Private/GPUProfiler.cpp
@@ -722,50 +722,91 @@ bool FGPUTiming::GAreGlobalsInitialized = false;
 namespace UE::RHI::GPUProfiler
 {
 #if HAS_GPU_STATS
-	// Total GPU time per queue.
-	static FGPUStat GPUStat_Total(TEXT("[TOTAL]"));
+	// Per queue GPU stats
+
+	// Total busy time on the current queue.
+	static FGPUStat GPUStat_Total(TEXT("Queue Total"));
+
+	// Total time idle, waiting for more commands to arrive from the CPU.
+	static FGPUStat GPUStat_Idle(TEXT("Queue Idle"));
+
+	// Total time blocked waiting on another GPU queue / fence.
+	static FGPUStat GPUStat_Waiting(TEXT("Queue Waiting"));
 #endif
 
-	TMap<FQueue, TUniquePtr<FGPUStat::FStatCategory>> FGPUStat::Categories;
+	TMap<FGPUStat::FStatIdentifier, TUniquePtr<FGPUStat::FStatCategory>> FGPUStat::Categories;
 
-	FGPUStat::FStatCategory& FGPUStat::GetCategory(FQueue Queue)
+	FGPUStat::FStatCategory& FGPUStat::GetCategory(FStatIdentifier ID)
 	{
-		TUniquePtr<FStatCategory>& Category = Categories.FindOrAdd(Queue);
+		TUniquePtr<FStatCategory>& Category = Categories.FindOrAdd(ID);
 		if (!Category)
 		{
-			Category = MakeUnique<FStatCategory>(Queue);
+			Category = MakeUnique<FStatCategory>(ID);
 		}
 		return *Category;
 	}
 
-	FGPUStat::FStatCategory_Stats::FStatCategory_Stats(FQueue Queue)
+	FString FGPUStat::FStatIdentifier::GetIDString(bool bFriendly) const
+	{
+		TCHAR const* TypeStr;
+		switch (Type)
+		{
+		default: checkNoEntry(); [[fallthrough]];
+		case EType::Busy: TypeStr = TEXT("Busy"); break;
+		case EType::Idle: TypeStr = TEXT("Idle"); break;
+		case EType::Wait: TypeStr = TEXT("Wait"); break;
+		}
+
+		if (bFriendly)
+		{
+			return FString::Printf(TEXT("GPU %d %s Queue %d (%s)")
+				, Queue.GPU
+				, Queue.GetTypeString()
+				, Queue.Index
+				, TypeStr
+			);
+		}
+		else
+		{
+			return FString::Printf(TEXT("GPU%d_%s%d_%s")
+				, Queue.GPU
+				, Queue.GetTypeString()
+				, Queue.Index
+				, TypeStr
+			);
+		}
+	}
+
+	FGPUStat::FStatCategory_Stats::FStatCategory_Stats(FStatIdentifier const& ID)
 	#if STATS
-		: GroupName(FString::Printf(TEXT("STATGROUP_GPU%d%s%d"), Queue.GPU, Queue.GetTypeString(), Queue.Index))
-		, GroupDesc(FString::Printf(TEXT("GPU %d %s Queue %d Timing"), Queue.GPU, Queue.GetTypeString(), Queue.Index))
+		: GroupName(FString::Printf(TEXT("STATGROUP_%s"), *ID.GetIDString(false)))
+		, GroupDesc(FString::Printf(TEXT("%s Timing"), *ID.GetIDString(true)))
 	#endif
 	{}
 
-	FGPUStat::FStatCategory_Csv::FStatCategory_Csv(FQueue Queue)
+	FGPUStat::FStatCategory_Csv::FStatCategory_Csv(FStatIdentifier const& ID)
 	#if CSV_PROFILER_STATS
-		: Csv(*FString::Printf(TEXT("GPU%d%s%d"), Queue.GPU, Queue.GetTypeString(), Queue.Index), true)
+		: Csv(*ID.GetIDString(false), true)
 	#endif
 	{}
 
-	FGPUStat::FStatCategory::FStatCategory(FQueue Queue)
-		: FStatCategory_Stats(Queue)
-		, FStatCategory_Csv(Queue)
+	FGPUStat::FStatCategory::FStatCategory(FStatIdentifier const& ID)
+		: FStatCategory_Stats(ID)
+		, FStatCategory_Csv(ID)
 	{}
 
 #if STATS
-	TStatId FGPUStat::GetStatId(FQueue Queue)
+	TStatId FGPUStat::GetStatId(FQueue Queue, EType Type)
 	{
-		FStatInstance& Instance = Instances.FindOrAdd(Queue);
+		FStatIdentifier ID = { Queue, Type };
+		FStatInstance& Instance = Instances.FindOrAdd(ID);
 
 		if (!Instance.Stat)
 		{
-			FStatCategory& Category = GetCategory(Queue);
+			FStatCategory& Category = GetCategory(ID);
+
+			Instance.StatName = MakeUnique<FAnsiString>(FString::Printf(TEXT("STAT_%s_%s"), *ID.GetIDString(false), Name));
 
-			Instance.StatName = MakeUnique<FAnsiString>(FString::Printf(TEXT("STAT_GPU%d%s%d_%s"), Queue.GPU, Queue.GetTypeString(), Queue.Index, Name));
 			Instance.Stat = MakeUnique<FDynamicStat>(
 				**Instance.StatName,
 				Name,
@@ -788,11 +829,12 @@ namespace UE::RHI::GPUProfiler
 #if CSV_PROFILER_STATS
 	FCsvDeclaredStat& FGPUStat::GetCsvStat(FQueue Queue)
 	{
-		FStatInstance& Instance = Instances.FindOrAdd(Queue);
+		FStatIdentifier ID = { Queue, EType::Busy };
+		FStatInstance& Instance = Instances.FindOrAdd(ID);
 
 		if (!Instance.CsvStat.IsSet())
 		{
-			FStatCategory& Category = GetCategory(Queue);
+			FStatCategory& Category = GetCategory(ID);
 			Instance.CsvStat.Emplace(Name, Category.Csv.Index);
 		}
 
@@ -820,22 +862,168 @@ namespace UE::RHI::GPUProfiler
 
 	void ProcessEvents(FQueue Queue, FEventStream EventStream)
 	{
+		if (EventStream.IsEmpty())
+		{
+			return;
+		}
+
 		TRACE_CPUPROFILER_EVENT_SCOPE(UE::RHI::GPUProfiler::ProcessEvents);
+		TSharedRef<FEventStream> SharedStream = MakeShared<FEventStream>(MoveTemp(EventStream));
 
-		if (!EventStream.IsEmpty())
+		for (FEventSink* Sink : GetSinks())
+		{
+			Sink->AppendStream(Queue, SharedStream);
+			Sink->ProcessStreams();
+		}
+	}
+
+	void InitializeQueues(TConstArrayView<FQueue> Queues)
+	{
+		for (FQueue Queue : Queues)
 		{
 			for (FEventSink* Sink : GetSinks())
 			{
-				Sink->ProcessEvents(Queue, EventStream);
+				Sink->QueueStates.Add(Queue, Sink->CreateQueueState(Queue));
 			}
 		}
 	}
 
-	void InitializeQueues(TConstArrayView<FQueue> Queues)
+	TUniquePtr<FEventSink::FQueueState> FEventSink::CreateQueueState(FQueue Queue) const
 	{
-		for (FEventSink* Sink : GetSinks())
+		return MakeUnique<FQueueState>(Queue);
+	}
+
+	void FEventSink::AppendStream(FQueue Queue, TSharedRef<FEventStream> const& NewEventStream)
+	{
+		FQueueState& QueueState = *QueueStates.FindChecked(Queue);
+		QueueState.PendingStreams.Enqueue(FIterator(NewEventStream));
+	}
+
+	void FEventSink::ProcessStreams()
+	{
+		// Process the queue as far as possible
+		bool bProgress;
+		do
 		{
-			Sink->InitializeQueues(Queues);
+			bProgress = false;
+			for (auto& [Queue, QueueState] : QueueStates)
+			{
+				while (FIterator* Iterator = QueueState->PendingStreams.Peek())
+				{
+					FEvent const* Start = Iterator->Peek();
+
+					bool bPaused = !ProcessStream(*QueueState.Get(), *Iterator);
+
+					FEvent const* End = Iterator->Peek();
+					bProgress |= End != Start;
+
+					if (bPaused)
+					{
+						// The queue was paused by a Wait event
+						check(End);
+						break;
+					}
+
+					if (!End)
+					{
+						// This stream has been fully processed.
+						QueueState->PendingStreams.Dequeue();
+					}
+				}
+			}
+		} while (bProgress);
+	}
+
+	void FEventSink::FQueueState::ResolveSignal(FEvent::FSignalFence const& SignalFenceEvent)
+	{
+		FResolvedSignal& Result = Signals.Emplace_GetRef();
+		
+		//
+		// Take the max between the previous GPU EndWork event and the CPU timestamp. The signal cannot have happened on the GPU until the CPU has submitted the command to the driver.
+		// 
+		// An example would be a GPU queue that completes work and goes idle at time T. Later, the CPU issues a Signal without other prior work at time T + 100ms.
+		// The fence signal cannot have happened until time T + 100ms because the CPU hadn't instructed the GPU to do so until then.
+		// LastGPUCycles would still be set to time T, since that was the time of the preceeding EndWork event.
+		//
+		Result.GPUTimestampBOP = FMath::Max(LastGPUCycles, SignalFenceEvent.CPUTimestamp);
+		Result.Value = SignalFenceEvent.Value;
+
+		//
+		// Fences signals *MUST* be sequential, to remove ambiguity caused by trimming the Signals array.
+		// 
+		// To explain why, assume non-sequential signals are allowed, and consider the following example events on an arbitrary queue:
+		// 
+		//		     [Signal 2]
+		//		-- Frame Boundary --
+		//		     [Signal 4]
+		//
+		// Assume, after trimming events earlier than the frame boundary, that only [Signal 4] remains in the Signals array.
+		// Then, some other queue attempts to [Wait 3]. We need to compute when [Wait 3] is resolved with only the information about [Signal 4].
+		// 
+		// Given that fences resolve waits as soon as the signalled value is >= the wait value, we could assume the fence was resolved at [Signal 4].
+		// However, we don't know if the fence was already signalled to value 3 before the frame boundary and the trimming.
+		// 
+		// Without this information, it is ambiguous whether [Wait 3] is already resolved by a [Signal 3] before the frame boundary that is no longer
+		// in the Signals array, or won't be resolved until [Signal 4]. We could have had this sequence of events:
+		// 
+		//		     [Signal 2]
+		//		     [Signal 3]
+		//		-- Frame Boundary --
+		//		     [Signal 4]
+		// 
+		// Requiring that fences are always signalled in sequential order solves this.
+		// If the awaited value is less than the first Signal, the fence has already been signalled before the frame boundary.
+		//
+		check(Result.Value == MaxSignal.Value + 1);
+
+		// Signals should always advance in time
+		check(Result.GPUTimestampBOP >= MaxSignal.GPUTimestampBOP);
+
+		MaxSignal = Result;
+	}
+
+	TOptional<FEventSink::FResolvedWait> FEventSink::ResolveWait(FQueueState& LocalQueue, FEvent::FWaitFence const& WaitFenceEvent)
+	{
+		FQueueState const& RemoteQueue = *QueueStates.FindChecked(WaitFenceEvent.Queue);
+
+		if (RemoteQueue.MaxSignal.Value < WaitFenceEvent.Value)
+		{
+			// Fence has not yet been signalled on the remote queue
+			return {};
+		}
+		else
+		{
+			// Fence has been signalled, but it may be in the future.
+
+			FResolvedWait Result;
+			Result.CPUTimestamp = WaitFenceEvent.CPUTimestamp;
+
+			//
+			// The wait cannot be resolved any earlier than:
+			//
+			//		1) The wait command was issued to the driver (WaitFenceEvent.CPUTimestamp)
+			//		2) The GPU completed prior work on this queue (LocalQueue.LastGPUCycles)
+			//
+			Result.GPUTimestampTOP = FMath::Max(WaitFenceEvent.CPUTimestamp, LocalQueue.LastGPUCycles);
+			//
+			//      3) The wait maybe be further delayed by the remote queue the GPU is awaiting.
+			//
+			int32 Index = Algo::LowerBoundBy(RemoteQueue.Signals, WaitFenceEvent.Value, [](FResolvedSignal const& Signal) { return Signal.Value; });
+			if (RemoteQueue.Signals.IsValidIndex(Index))
+			{
+				FResolvedSignal const& Signal = RemoteQueue.Signals[Index];
+
+				//
+				// Only consider this signal's timestamp if the fence was not already signalled at the previous frame boundary.
+				// See comment in ResolveSignal() for details.
+				//
+				if (!(Index == 0 && WaitFenceEvent.Value < Signal.Value))
+				{
+					Result.GPUTimestampTOP = FMath::Max(Result.GPUTimestampTOP, Signal.GPUTimestampBOP);
+				}
+			}
+
+			return Result;
 		}
 	}
 
@@ -952,7 +1140,8 @@ namespace UE::RHI::GPUProfiler
 		struct FStatState
 		{
 			uint64 BusyCycles = 0;
-			uint64 LastCycles = 0;
+			uint64 IdleCycles = 0;
+			uint64 WaitCycles = 0;
 
 			FStatState() = default;
 			FStatState(FStatState const&) = default;
@@ -961,22 +1150,39 @@ namespace UE::RHI::GPUProfiler
 				: FStatState(Other)
 			{
 				Other.BusyCycles = 0;
-				Other.LastCycles = 0;
+				Other.IdleCycles = 0;
+				Other.WaitCycles = 0;
 			}
 
-			void BeginWork(uint64 Timestamp)
+			void Accumulate(uint64 Busy, uint64 Wait, uint64 Idle)
 			{
-				LastCycles = FMath::Max(LastCycles, Timestamp);
+				BusyCycles += Busy;
+				IdleCycles += Idle;
+				WaitCycles += Wait;
 			}
 
-			void EndWork(uint64 Timestamp)
+			void EmitResults(FQueue Queue
+			#if HAS_GPU_STATS
+				, FGPUStat& GPUStat
+			#endif
+			#if CSV_PROFILER_STATS
+				, FCsvProfiler* CsvProfiler
+			#endif
+			) const
 			{
-				uint64 Delta = Timestamp > LastCycles
-					? Timestamp - LastCycles
-					: 0;
+			#if HAS_GPU_STATS
+				SET_FLOAT_STAT_FName(GPUStat.GetStatId(Queue, FGPUStat::EType::Busy).GetName(), FPlatformTime::ToMilliseconds64(BusyCycles));
+				SET_FLOAT_STAT_FName(GPUStat.GetStatId(Queue, FGPUStat::EType::Idle).GetName(), FPlatformTime::ToMilliseconds64(IdleCycles));
+				SET_FLOAT_STAT_FName(GPUStat.GetStatId(Queue, FGPUStat::EType::Wait).GetName(), FPlatformTime::ToMilliseconds64(WaitCycles));
+			#endif
 
-				LastCycles = FMath::Max(LastCycles, Timestamp);
-				BusyCycles += Delta;
+			#if CSV_PROFILER_STATS
+				if (CsvProfiler)
+				{
+					FCsvDeclaredStat& CsvStat = GPUStat.GetCsvStat(Queue);
+					CsvProfiler->RecordEndOfPipeCustomStat(CsvStat.Name, CsvStat.CategoryIndex, FPlatformTime::ToMilliseconds64(BusyCycles), ECsvCustomStatOp::Set);
+				}
+			#endif
 			}
 		};
 
@@ -985,47 +1191,151 @@ namespace UE::RHI::GPUProfiler
 			FTimestampStream Queue;
 			FStatState WholeQueueStat;
 
+			uint64 CPUFrameBoundary = 0;
+
 		#if WITH_RHI_BREADCRUMBS
 			TMap<FRHIBreadcrumbData_Stats, FStatState> Stats;
 		#endif
 		};
 
-		struct FQueueState
+		struct FQueueState : public FEventSink::FQueueState
 		{
-			FQueue::EType Type;
-
-			bool bBusy = false;
 			FQueueTimestamps Timestamps;
+			uint64 CPUBeginWorkTime = 0;
 
 		#if WITH_RHI_BREADCRUMBS
 			TMap<FRHIBreadcrumbData_Stats, int32> ActiveStats;
 			FRHIBreadcrumbNode* Breadcrumb = nullptr;
 		#endif
 
+			bool bBusy = false;
+
 			FQueueState(FQueue const& Queue)
-				: Type(Queue.Type)
+				: FEventSink::FQueueState(Queue)
 			{}
-		};
 
-		using FFrameState = TMap<FQueue, FQueueTimestamps>;
+			void AccumulateTime(uint64 Busy, uint64 Wait, uint64 Idle)
+			{
+				Timestamps.WholeQueueStat.Accumulate(Busy, Wait, Idle);
 
-		TMap<FQueue, FQueueState> QueueStates;
-		TMap<uint32, FFrameState> Frames;
+			#if WITH_RHI_BREADCRUMBS
+				// Apply the timings to all active stats
+				for (auto const& [Stat, RefCount] : ActiveStats)
+				{
+					Timestamps.Stats.FindChecked(Stat).Accumulate(Busy, Wait, Idle);
+				}
+			#endif
+			}
 
-		void InitializeQueues(TConstArrayView<FQueue> Queues) override
-		{
-			for (FQueue const& Queue : Queues)
+			void BeginWork(uint64 GPUTimestampTOP, uint64 CPUTimestamp)
 			{
-				check(QueueStates.Find(Queue) == nullptr);
-				QueueStates.Add(Queue, Queue);
+				Timestamps.Queue.AddTimestamp(GPUTimestampTOP, true);
+
+				uint64 Idle = CPUTimestamp > LastGPUCycles
+					? CPUTimestamp - LastGPUCycles
+					: 0;
+
+				AccumulateTime(0, 0, Idle);
+				LastGPUCycles = FMath::Max(LastGPUCycles, GPUTimestampTOP);
+			}
+
+			void EndWork(uint64 GPUTimestampBOP)
+			{
+				Timestamps.Queue.AddTimestamp(GPUTimestampBOP, false);
+
+				uint64 Busy = GPUTimestampBOP > LastGPUCycles
+					? GPUTimestampBOP - LastGPUCycles
+					: 0;
+
+				AccumulateTime(Busy, 0, 0);
+				LastGPUCycles = FMath::Max(LastGPUCycles, GPUTimestampBOP);
 			}
+
+			void BeginStat(FRHIBreadcrumbData_Stats const& Stat, uint64 GPUTimestampTOP)
+			{
+				uint64 Busy = GPUTimestampTOP > LastGPUCycles
+					? GPUTimestampTOP - LastGPUCycles
+					: 0;
+
+				AccumulateTime(Busy, 0, 0);
+				LastGPUCycles = FMath::Max(LastGPUCycles, GPUTimestampTOP);
+
+				if (Stat.ShouldComputeStat())
+				{
+					// Disregard the stat if it is nested within itself (i.e. its already in the ActiveStats map with a non-zero ref count).
+					// Only the outermost stat will count the busy time, otherwise we'd be double-counting the nested time.
+					int32 RefCount = ActiveStats.FindOrAdd(Stat)++;
+					if (RefCount == 0)
+					{
+						Timestamps.Stats.FindOrAdd(Stat);
+					}
+				}
+			}
+
+			void EndStat(FRHIBreadcrumbData_Stats const& Stat, uint64 GPUTimestampBOP)
+			{
+				uint64 Busy = GPUTimestampBOP > LastGPUCycles
+					? GPUTimestampBOP - LastGPUCycles
+					: 0;
+
+				AccumulateTime(Busy, 0, 0);
+				LastGPUCycles = FMath::Max(LastGPUCycles, GPUTimestampBOP);
+
+				if (Stat.ShouldComputeStat())
+				{
+					// Pop the stat when the refcount hits zero.
+					int32 RefCount = --ActiveStats.FindChecked(Stat);
+					if (RefCount == 0)
+					{
+						ActiveStats.FindAndRemoveChecked(Stat);
+					}
+				}
+			}
+
+			void Wait(FResolvedWait const& ResolvedWait)
+			{
+				// Time the queue was idle between the last EndWork event, and the Wait command being submitted to the GPU driver.
+				uint64 Idle = ResolvedWait.CPUTimestamp > LastGPUCycles
+					? ResolvedWait.CPUTimestamp - LastGPUCycles
+					: 0;
+
+				uint64 WaitStart = FMath::Max(ResolvedWait.CPUTimestamp, LastGPUCycles);
+
+				// Time the queue spent waiting for the fence to signal on another queue.
+				uint64 Wait = ResolvedWait.GPUTimestampTOP > WaitStart
+					? ResolvedWait.GPUTimestampTOP - WaitStart
+					: 0;
+
+				// Bring the last GPU busy end time forwards to where the wait is resolved.
+				LastGPUCycles = ResolvedWait.GPUTimestampTOP;
+
+				AccumulateTime(0, Wait, Idle);
+			}
+
+			void TrimSignals(uint64 CPUTimestamp)
+			{
+				// Remove all signals that occured on the GPU timeline before this frame boundary on the CPU.
+				int32 Index = Algo::LowerBoundBy(Signals, CPUTimestamp, [](FResolvedSignal const& Signal) { return Signal.GPUTimestampBOP; });
+				if (Index >= 0)
+				{
+					Signals.RemoveAt(0, Index, EAllowShrinking::No);
+				}
+			}
+		};
+
+		virtual TUniquePtr<FEventSink::FQueueState> CreateQueueState(FQueue Queue) const override
+		{
+			return MakeUnique<FQueueState>(Queue);
 		}
 
-		void ProcessEvents(FQueue Queue, FEventStream const& EventStream) override
+		using FFrameState = TMap<FQueue, FQueueTimestamps>;
+		TMap<uint32, FFrameState> Frames;
+
+		bool ProcessStream(FEventSink::FQueueState& BaseQueueState, FIterator& Iterator) override
 		{
-			FQueueState& QueueState = QueueStates.FindChecked(Queue);
+			FQueueState& QueueState = static_cast<FQueueState&>(BaseQueueState);
 
-			for (FEvent const* Event : EventStream)
+			while (FEvent const* Event = Iterator.Peek())
 			{
 				switch (Event->GetType())
 				{
@@ -1035,16 +1345,11 @@ namespace UE::RHI::GPUProfiler
 						QueueState.bBusy = true;
 
 						FEvent::FBeginWork const& BeginWork = Event->Value.Get<FEvent::FBeginWork>();
-						QueueState.Timestamps.Queue.AddTimestamp(BeginWork.GPUTimestampTOP, true);
-						QueueState.Timestamps.WholeQueueStat.BeginWork(BeginWork.GPUTimestampTOP);
 
-					#if WITH_RHI_BREADCRUMBS
-						// Apply the timestamp to all active stats
-						for (auto const& [Stat, RefCount] : QueueState.ActiveStats)
-						{
-							QueueState.Timestamps.Stats.FindChecked(Stat).BeginWork(BeginWork.GPUTimestampTOP);
-						}
-					#endif
+						check(QueueState.LastCPUCycles <= BeginWork.CPUTimestamp);
+						QueueState.LastCPUCycles = BeginWork.CPUTimestamp;
+
+						QueueState.BeginWork(BeginWork.GPUTimestampTOP, BeginWork.CPUTimestamp);
 					}
 					break;
 
@@ -1054,16 +1359,7 @@ namespace UE::RHI::GPUProfiler
 						QueueState.bBusy = false;
 
 						FEvent::FEndWork const& EndWork = Event->Value.Get<FEvent::FEndWork>();
-						QueueState.Timestamps.Queue.AddTimestamp(EndWork.GPUTimestampBOP, false);
-						QueueState.Timestamps.WholeQueueStat.EndWork(EndWork.GPUTimestampBOP);
-
-					#if WITH_RHI_BREADCRUMBS
-						// Apply the timestamp to all active stats
-						for (auto const& [Stat, RefCount] : QueueState.ActiveStats)
-						{
-							QueueState.Timestamps.Stats.FindChecked(Stat).EndWork(EndWork.GPUTimestampBOP);
-						}
-					#endif
+						QueueState.EndWork(EndWork.GPUTimestampBOP);
 					}
 					break;
 
@@ -1075,17 +1371,7 @@ namespace UE::RHI::GPUProfiler
 						FEvent::FBeginBreadcrumb const& BeginBreadcrumb = Event->Value.Get<FEvent::FBeginBreadcrumb>();
 						FRHIBreadcrumbData_Stats const& Stat = BeginBreadcrumb.Breadcrumb->Name.Data;
 
-						if (Stat.ShouldComputeStat())
-						{
-							// Disregard the stat if it is nested within itself (i.e. its already in the ActiveStats map with a non-zero ref count).
-							// Only the outermost stat will count the busy time, otherwise we'd be double-counting the nested time.
-							int32 RefCount = QueueState.ActiveStats.FindOrAdd(Stat)++;
-							if (RefCount == 0)
-							{
-								QueueState.Timestamps.Stats.FindOrAdd(Stat).BeginWork(BeginBreadcrumb.GPUTimestampTOP);
-							}
-						}
-
+						QueueState.BeginStat(Stat, BeginBreadcrumb.GPUTimestampTOP);
 						QueueState.Breadcrumb = BeginBreadcrumb.Breadcrumb;
 					}
 					break;
@@ -1097,17 +1383,7 @@ namespace UE::RHI::GPUProfiler
 						FEvent::FEndBreadcrumb const& EndBreadcrumb = Event->Value.Get<FEvent::FEndBreadcrumb>();
 						FRHIBreadcrumbData_Stats const& Stat = EndBreadcrumb.Breadcrumb->Name.Data;
 
-						if (Stat.ShouldComputeStat())
-						{
-							// Pop the stat when the refcount hits zero.
-							int32 RefCount = --QueueState.ActiveStats.FindChecked(Stat);
-							if (RefCount == 0)
-							{
-								QueueState.Timestamps.Stats.FindChecked(Stat).EndWork(EndBreadcrumb.GPUTimestampBOP);
-								QueueState.ActiveStats.FindAndRemoveChecked(Stat);
-							}							
-						}
-
+						QueueState.EndStat(Stat, EndBreadcrumb.GPUTimestampBOP);
 						QueueState.Breadcrumb = EndBreadcrumb.Breadcrumb->GetParent();
 					}
 					break;
@@ -1117,6 +1393,11 @@ namespace UE::RHI::GPUProfiler
 					{
 						check(!QueueState.bBusy);
 						FEvent::FSignalFence const& SignalFence = Event->Value.Get<FEvent::FSignalFence>();
+
+						check(QueueState.LastCPUCycles <= SignalFence.CPUTimestamp);
+						QueueState.LastCPUCycles = SignalFence.CPUTimestamp;
+
+						QueueState.ResolveSignal(SignalFence);
 					}
 					break;
 
@@ -1124,6 +1405,19 @@ namespace UE::RHI::GPUProfiler
 					{
 						check(!QueueState.bBusy);
 						FEvent::FWaitFence const& WaitFence = Event->Value.Get<FEvent::FWaitFence>();
+
+						check(QueueState.LastCPUCycles <= WaitFence.CPUTimestamp);
+						QueueState.LastCPUCycles = WaitFence.CPUTimestamp;
+
+						TOptional<FResolvedWait> ResolvedWait = ResolveWait(QueueState, WaitFence);
+
+						if (!ResolvedWait.IsSet())
+						{
+							// Unresolved fence, pause processing
+							return false;
+						}
+
+						QueueState.Wait(*ResolvedWait);
 					}
 					break;
 
@@ -1132,8 +1426,13 @@ namespace UE::RHI::GPUProfiler
 						check(!QueueState.bBusy);
 						FEvent::FFrameBoundary const& FrameBoundary = Event->Value.Get<FEvent::FFrameBoundary>();
 
+						check(QueueState.LastCPUCycles <= FrameBoundary.CPUTimestamp);
+						QueueState.LastCPUCycles = FrameBoundary.CPUTimestamp;
+
+						QueueState.Timestamps.CPUFrameBoundary = FrameBoundary.CPUTimestamp;
+
 						FFrameState& FrameState = Frames.FindOrAdd(FrameBoundary.FrameNumber);
-						FrameState.Emplace(Queue, MoveTemp(QueueState.Timestamps));
+						FrameState.Emplace(QueueState.Queue, MoveTemp(QueueState.Timestamps));
 
 					#if WITH_RHI_BREADCRUMBS
 						// Reinsert timestamp streams for the current active stats on 
@@ -1146,6 +1445,22 @@ namespace UE::RHI::GPUProfiler
 
 						if (FrameState.Num() == QueueStates.Num())
 						{
+							// Trim the Signals array in each queue, up to the lowest frame boundary CPU timestamp.
+							uint64 MinFrameBoundary = TNumericLimits<uint64>::Max();
+							{
+								for (auto& [Queue, LocalQueueState] : FrameState)
+								{
+									MinFrameBoundary = FMath::Min(MinFrameBoundary, LocalQueueState.CPUFrameBoundary);
+								}
+							}
+
+							{
+								for (auto& [Queue, LocalQueueState] : QueueStates)
+								{
+									static_cast<FQueueState&>(*LocalQueueState.Get()).TrimSignals(MinFrameBoundary);
+								}
+							}
+
 							// All registered queues have reported their frame boundary event.
 							// We have a full set of data to compute the total frame GPU stats.
 							ProcessFrame(FrameState);
@@ -1155,14 +1470,19 @@ namespace UE::RHI::GPUProfiler
 					}
 					break;
 				}
+
+				Iterator.Pop();
 			}
+
+			return true;
 		}
 
 		void ProcessFrame(FFrameState& FrameState)
 		{
 		#if CSV_PROFILER_STATS
 			const bool bCsvStatsEnabled = !!CVarGPUCsvStatsEnabled.GetValueOnAnyThread();
-			FCsvProfiler* CsvProfiler = bCsvStatsEnabled ? FCsvProfiler::Get() : nullptr;
+			FCsvProfiler* CsvProfiler = FCsvProfiler::Get();
+			CsvProfiler->BeginFrameEOP();
 		#else
 			const bool bCsvStatsEnabled = false;
 		#endif
@@ -1173,30 +1493,26 @@ namespace UE::RHI::GPUProfiler
 				// Compute the individual GPU stats
 				for (auto const& [Stat, StatState] : QueueTimestamps.Stats)
 				{
-					double StatMilliseconds = FPlatformTime::ToMilliseconds64(StatState.BusyCycles);
-					SET_FLOAT_STAT_FName(Stat.GPUStat->GetStatId(Queue).GetName(), StatMilliseconds);
-
-				#if CSV_PROFILER_STATS
-					if (CsvProfiler)
-					{
-						FCsvDeclaredStat& CsvStat = Stat.GPUStat->GetCsvStat(Queue);
-						CsvProfiler->RecordCustomStat(CsvStat.Name, CsvStat.CategoryIndex, StatMilliseconds, ECsvCustomStatOp::Set);
-					}
-				#endif
+					StatState.EmitResults(Queue
+					#if HAS_GPU_STATS
+						, *Stat.GPUStat
+					#endif
+					#if CSV_PROFILER_STATS
+						, bCsvStatsEnabled ? CsvProfiler : nullptr
+					#endif
+					);
 				}
 			#endif // WITH_RHI_BREADCRUMBS
 
 				// Set the whole-frame per queue stat
-				double QueueMilliseconds = FPlatformTime::ToMilliseconds64(QueueTimestamps.WholeQueueStat.BusyCycles);
-				SET_FLOAT_STAT_FName(GPUStat_Total.GetStatId(Queue).GetName(), QueueMilliseconds);
-
-			#if CSV_PROFILER_STATS
-				if (CsvProfiler)
-				{
-					FCsvDeclaredStat& CsvStat = GPUStat_Total.GetCsvStat(Queue);
-					CsvProfiler->RecordCustomStat(CsvStat.Name, CsvStat.CategoryIndex, QueueMilliseconds, ECsvCustomStatOp::Set);
-				}
-			#endif
+				QueueTimestamps.WholeQueueStat.EmitResults(Queue
+				#if HAS_GPU_STATS
+					, GPUStat_Total
+				#endif
+				#if CSV_PROFILER_STATS
+					, bCsvStatsEnabled ? CsvProfiler : nullptr
+				#endif
+				);
 			}
 
 			// Compute the whole-frame total GPU time.
@@ -1280,10 +1596,8 @@ namespace UE::RHI::GPUProfiler
 			{}
 		};
 
-		struct FQueueState
+		struct FQueueState : public FEventSink::FQueueState
 		{
-			FQueue const Queue;
-
 			enum class EMode
 			{
 				Idle,
@@ -1337,62 +1651,62 @@ namespace UE::RHI::GPUProfiler
 			std::atomic<bool> bTriggerProfile { false };
 
 			FQueueState(FQueue Queue)
-				: Queue(Queue)
+				: FEventSink::FQueueState(Queue)
 			{}
 
-			void ProcessEvents(FEventStream const& EventStream);
+			void ProcessEvents(FIterator& Iterator);
 
 			void LogTree(uint32 FrameNumber);
 		};
 
-		TMap<FQueue, TUniquePtr<FQueueState>> QueueStates;
-
-		void InitializeQueues(TConstArrayView<FQueue> Queues) override
+		virtual TUniquePtr<FEventSink::FQueueState> CreateQueueState(FQueue Queue) const override
 		{
-			for (FQueue const& Queue : Queues)
-			{
-				check(QueueStates.Find(Queue) == nullptr);
-				QueueStates.Add(Queue, MakeUnique<FQueueState>(Queue));
-			}
+			return MakeUnique<FQueueState>(Queue);
 		}
 
-		void ProcessEvents(FQueue Queue, FEventStream const& EventStream) override
+		bool ProcessStream(FEventSink::FQueueState& QueueState, FIterator& Iterator) override
 		{
-			QueueStates.FindChecked(Queue)->ProcessEvents(EventStream);
+			static_cast<FQueueState&>(QueueState).ProcessEvents(Iterator);
+			return true;
 		}
 
 		void ProfileNextFrame()
 		{
 			for (auto& [Queue, State] : QueueStates)
 			{
-				State->bTriggerProfile = true;
+				static_cast<FQueueState&>(*State).bTriggerProfile = true;
 			}
 		}
 
 	} GGPUProfilerSink_ProfileGPU;
 
-	void FGPUProfilerSink_ProfileGPU::FQueueState::ProcessEvents(FEventStream const& EventStream)
+	void FGPUProfilerSink_ProfileGPU::FQueueState::ProcessEvents(FIterator& Iterator)
 	{
-		auto Iterator = std::begin(EventStream);
-		auto End = std::end(EventStream);
-
 	Restart:
-		if (State.Mode == EMode::Idle && bTriggerProfile.exchange(false))
+		if (State.Mode == EMode::Idle)
 		{
-			State.Mode = EMode::WaitingFrame;
+			if (bTriggerProfile.exchange(false))
+			{
+				State.Mode = EMode::WaitingFrame;
+			}
+			else
+			{
+				Iterator.PopAll();
+				return;
+			}
 		}
 
 		if (State.Mode == EMode::WaitingFrame)
 		{
 			// Discard all received events until we reach a FrameBoundary event
-			for (; Iterator != End; ++Iterator)
+			while (FEvent const* Event = Iterator.Pop())
 			{
-				if ((*Iterator)->GetType() == FEvent::EType::FrameBoundary)
+				if (Event->GetType() == FEvent::EType::FrameBoundary)
 				{
 					// Start profiling until we receive another FrameBoundary event
 					State.Mode = EMode::Active;
 
-					FEvent::FFrameBoundary const& FrameBoundary = (*Iterator)->Value.Get<FEvent::FFrameBoundary>();
+					FEvent::FFrameBoundary const& FrameBoundary = Event->Value.Get<FEvent::FFrameBoundary>();
 
 					// Build the node tree 
 					State.PushNode(TEXT("<root>"));
@@ -1413,7 +1727,6 @@ namespace UE::RHI::GPUProfiler
 					Recurse(Recurse, FrameBoundary.Breadcrumb);
 				#endif // WITH_RHI_BREADCRUMBS
 
-					++Iterator;
 					break;
 				}
 			}
@@ -1447,10 +1760,8 @@ namespace UE::RHI::GPUProfiler
 				}
 			};
 
-			for (; Iterator != End; ++Iterator)
+			while (FEvent const* Event = Iterator.Pop())
 			{
-				FEvent const* Event = *Iterator;
-
 				switch (Event->GetType())
 				{
 				case FEvent::EType::BeginWork:
diff --git a/Engine/Source/Runtime/RHI/Private/GpuProfilerTrace.cpp b/Engine/Source/Runtime/RHI/Private/GpuProfilerTrace.cpp
--- a/Engine/Source/Runtime/RHI/Private/GpuProfilerTrace.cpp
+++ b/Engine/Source/Runtime/RHI/Private/GpuProfilerTrace.cpp
@@ -328,7 +328,7 @@ UE_TRACE_EVENT_END()
 struct FGPUProfilerSink_Trace final : public FEventSink
 {
 	bool bInitialized = false;
-	TArray<FQueue> CachedQueues;
+	TSet<FQueue> CachedQueues;
 
 	void Initialize()
 	{
@@ -342,7 +342,7 @@ struct FGPUProfilerSink_Trace final : public FEventSink
 		bInitialized = true;
 	}
 
-	void TraceQueues(TConstArrayView<FQueue> Queues)
+	void TraceQueues(TSet<FQueue>& Queues)
 	{
 		for (FQueue const& Queue : Queues)
 		{
@@ -353,27 +353,12 @@ struct FGPUProfilerSink_Trace final : public FEventSink
 		}
 	}
 
-	void InitializeQueues(TConstArrayView<FQueue> Queues) override
+	bool ProcessStream(FQueueState& QueueState, FIterator& Iterator) override
 	{
 		if (!UE_TRACE_CHANNELEXPR_IS_ENABLED(GpuChannel))
 		{
-			CachedQueues.Append(Queues);
-			return;
-		}
-
-		if (!bInitialized)
-		{
-			Initialize();
-		}
-
-		TraceQueues(Queues);
-	}
-
-	void ProcessEvents(FQueue Queue, FEventStream const& EventStream) override
-	{
-		if (!UE_TRACE_CHANNELEXPR_IS_ENABLED(GpuChannel))
-		{
-			return;
+			CachedQueues.Add(QueueState.Queue);
+			return true;
 		}
 
 		if (!bInitialized)
@@ -387,7 +372,9 @@ struct FGPUProfilerSink_Trace final : public FEventSink
 			CachedQueues.Empty();
 		}
 
-		for (FEvent const* Event : EventStream)
+		FQueue Queue = QueueState.Queue;
+
+		while (FEvent const* Event = Iterator.Pop())
 		{
 			switch (Event->GetType())
 			{
@@ -475,6 +462,8 @@ struct FGPUProfilerSink_Trace final : public FEventSink
 				}
 			}
 		}
+
+		return true;
 	}
 
 } GGPUProfilerSink_Trace;
diff --git a/Engine/Source/Runtime/RHI/Public/GPUProfiler.h b/Engine/Source/Runtime/RHI/Public/GPUProfiler.h
--- a/Engine/Source/Runtime/RHI/Public/GPUProfiler.h
+++ b/Engine/Source/Runtime/RHI/Public/GPUProfiler.h
@@ -14,6 +14,7 @@
 #include "RHIStats.h"
 
 #include "Containers/AnsiString.h"
+#include "Containers/SpscQueue.h"
 
 #if RHI_NEW_GPU_PROFILER
 
@@ -89,6 +90,9 @@ namespace UE::RHI::GPUProfiler
 		// Inserted on each call to RHIEndFrame. Marks the end of a profiler frame.
 		struct FFrameBoundary
 		{
+			// CPU timestamp from the platform RHI's submission thread where the frame boundary occured.
+			uint64 CPUTimestamp;
+
 			// The index of the frame that just ended.
 			// Very first frame of the engine is frame 0 (from boot to first call to RHIEndFrame).
 			uint32 FrameNumber;
@@ -98,12 +102,13 @@ namespace UE::RHI::GPUProfiler
 			FRHIBreadcrumbNode* Breadcrumb;
 		#endif
 
-			FFrameBoundary(uint32 FrameNumber
+			FFrameBoundary(uint64 CPUTimestamp, uint32 FrameNumber
 			#if WITH_RHI_BREADCRUMBS
 				, FRHIBreadcrumbNode* Breadcrumb
 			#endif
 				)
-				: FrameNumber(FrameNumber)
+				: CPUTimestamp(CPUTimestamp)
+				, FrameNumber(FrameNumber)
 			#if WITH_RHI_BREADCRUMBS
 				, Breadcrumb(Breadcrumb)
 			#endif
@@ -186,15 +191,11 @@ namespace UE::RHI::GPUProfiler
 			//
 			uint64 CPUTimestamp;
 
-			// Unique ID of the fence signaled.
-			uint64 ID;
-
 			// The fence value signaled.
 			uint64 Value;
 
-			FSignalFence(uint64 CPUTimestamp, uint64 ID, uint64 Value)
+			FSignalFence(uint64 CPUTimestamp, uint64 Value)
 				: CPUTimestamp(CPUTimestamp)
-				, ID(ID)
 				, Value(Value)
 			{}
 		};
@@ -205,16 +206,16 @@ namespace UE::RHI::GPUProfiler
 			// Timestamp when the fence wait was enqueued to the GPU/driver.
 			uint64 CPUTimestamp;
 
-			// Unique ID of the fence awaited.
-			uint64 ID;
-
 			// The fence value awaited.
 			uint64 Value;
 
-			FWaitFence(uint64 CPUTimestamp, uint64 ID, uint64 Value)
+			// The queue the GPU is waiting for a fence signal from.
+			FQueue Queue;
+
+			FWaitFence(uint64 CPUTimestamp, uint64 Value, FQueue Queue)
 				: CPUTimestamp(CPUTimestamp)
-				, ID(ID)
 				, Value(Value)
+				, Queue(Queue)
 			{}
 		};
 
@@ -277,6 +278,8 @@ namespace UE::RHI::GPUProfiler
 
 	class FEventStream
 	{
+		friend struct FEventSink;
+
 	private:
 		struct FChunk
 		{
@@ -410,22 +413,38 @@ namespace UE::RHI::GPUProfiler
 			Other.Current = nullptr;
 			Other.First = nullptr;
 		}
+	};
 
-		auto begin() const
+	struct FEventSink
+	{
+		friend RHI_API void ProcessEvents(FQueue Queue, FEventStream EventStream);
+		friend RHI_API void InitializeQueues(TConstArrayView<FQueue> Queues);
+
+		struct FIterator
 		{
-			class FIterator
-			{
-				friend FEventStream;
+			friend FEventSink;
+
+		private:
+			TSharedRef<FEventStream> Stream;
 
-				FChunk* Current;
-				uint32 Index = 0;
+			FEventStream::FChunk* Current;
+			uint32 Index = 0;
 
-				FIterator(FChunk* Current)
-					: Current(Current)
-				{}
+			FIterator(TSharedRef<FEventStream> const& Stream)
+				: Stream(Stream)
+				, Current(Stream->First)
+			{}
+
+		public:
+			FEvent const* Peek() const
+			{
+				return Current ? Current->GetElement(Index) : nullptr;
+			}
 
-			public:
-				FIterator& operator++()
+			FEvent const* Pop()
+			{
+				FEvent const* Result = Peek();
+				if (Result)
 				{
 					++Index;
 
@@ -434,42 +453,73 @@ namespace UE::RHI::GPUProfiler
 						Current = Current->Header.Next;
 						Index = 0;
 					}
-
-					return *this;
 				}
 
-				bool operator != (std::nullptr_t) const
-				{
-					return Current != nullptr;
-				}
+				return Result;
+			}
 
-				FEvent const* operator*() const
-				{
-					return Current->GetElement(Index);
-				}
-			};
+			void PopAll()
+			{
+				Current = nullptr;
+				Index = 0;
+			}
+		};
 
-			return FIterator(First);
-		}
+	protected:
+		struct FResolvedWait
+		{
+			uint64 GPUTimestampTOP = 0;
+			uint64 CPUTimestamp    = 0;
+		};
 
-		std::nullptr_t end() const
+		struct FResolvedSignal
 		{
-			return nullptr;
-		}
-	};
+			uint64 GPUTimestampBOP = 0;
+			uint64 Value           = 0;
+		};
+
+		struct FQueueState
+		{
+			FQueue const Queue;
+
+			// Array of fence signal history. Events are kept until all queues have processed events
+			// later than the CPU timestamps of these signals. The old events are then trimmed.
+			TArray<FResolvedSignal> Signals;
+
+			// The value of the latest signaled fence on this queue.
+			FResolvedSignal MaxSignal;
+
+			// The GPU timestamp of the last event processed.
+			uint64 LastGPUCycles = 0;
+			uint64 LastCPUCycles = 0;
+
+			TSpscQueue<FIterator> PendingStreams;
+
+			FQueueState(FQueue const Queue)
+				: Queue(Queue)
+			{}
+
+			void ResolveSignal(FEvent::FSignalFence const& SignalFenceEvent);
+		};
+
+		TMap<FQueue, TUniquePtr<FQueueState>> QueueStates;
+
+		void AppendStream(FQueue Queue, TSharedRef<FEventStream> const& NewEventStream);
+		void ProcessStreams();
 
-	struct FEventSink
-	{
-	protected:
 		RHI_API FEventSink();
 		RHI_API ~FEventSink();
 
 		FEventSink(FEventSink const&) = delete;
 		FEventSink(FEventSink&&) = delete;
 
+		// Attempts to retrieve the CPU and GPU timestamps of when a fence wait is resolved by a signal on another queue.
+		TOptional<FResolvedWait> ResolveWait(FQueueState& QueueState, FEvent::FWaitFence const& WaitFenceEvent);
+
+		virtual TUniquePtr<FQueueState> CreateQueueState(FQueue Queue) const;
+
 	public:
-		virtual void ProcessEvents(FQueue Queue, FEventStream const& EventStream) = 0;
-		virtual void InitializeQueues(TConstArrayView<FQueue> Queues) = 0;
+		virtual bool ProcessStream(FQueueState& QueueState, FIterator& Iterator) = 0;
 	};
 
 	RHI_API void ProcessEvents(FQueue Queue, FEventStream EventStream);
@@ -477,6 +527,29 @@ namespace UE::RHI::GPUProfiler
 
 	struct FGPUStat
 	{
+		enum class EType
+		{
+			Busy, Idle, Wait
+		};
+
+		struct FStatIdentifier
+		{
+			FQueue Queue;
+			EType Type;
+
+			FString GetIDString(bool bFriendly) const;
+
+			bool operator == (FStatIdentifier const& RHS) const
+			{
+				return Queue == RHS.Queue && Type == RHS.Type;
+			}
+
+			friend uint32 GetTypeHash(FStatIdentifier const& ID)
+			{
+				return HashCombineFast(GetTypeHash(ID.Queue), GetTypeHash(ID.Type));
+			}
+		};
+
 	private:
 		struct FStatCategory_Stats
 		{
@@ -485,7 +558,7 @@ namespace UE::RHI::GPUProfiler
 			FString     const GroupDesc;
 		#endif
 
-			FStatCategory_Stats(FQueue Queue);
+			FStatCategory_Stats(FStatIdentifier const& ID);
 		};
 
 		struct FStatCategory_Csv
@@ -494,14 +567,14 @@ namespace UE::RHI::GPUProfiler
 			FCsvCategory const Csv;
 		#endif
 
-			FStatCategory_Csv(FQueue Queue);
+			FStatCategory_Csv(FStatIdentifier const& ID);
 		};
 
 		struct FStatCategory
 			: public FStatCategory_Stats
 			, public FStatCategory_Csv
 		{
-			FStatCategory(FQueue Queue);
+			FStatCategory(FStatIdentifier const& ID);
 		};
 
 		struct FStatInstance
@@ -515,10 +588,10 @@ namespace UE::RHI::GPUProfiler
 		#endif
 		};
 
-		static TMap<FQueue, TUniquePtr<FStatCategory>> Categories;
-		static FStatCategory& GetCategory(FQueue Queue);
+		static TMap<FStatIdentifier, TUniquePtr<FStatCategory>> Categories;
+		static FStatCategory& GetCategory(FStatIdentifier ID);
 
-		TMap<FQueue, FStatInstance> Instances;
+		TMap<FStatIdentifier, FStatInstance> Instances;
 
 	public:
 		TCHAR const* const Name;
@@ -528,7 +601,7 @@ namespace UE::RHI::GPUProfiler
 		{}
 
 	#if STATS
-		TStatId GetStatId(FQueue Queue);
+		TStatId GetStatId(FQueue Queue, EType Type);
 	#endif
 
 	#if CSV_PROFILER_STATS
diff --git a/Engine/Source/Runtime/Windows/D3D11RHI/Private/D3D11RHI.cpp b/Engine/Source/Runtime/Windows/D3D11RHI/Private/D3D11RHI.cpp
--- a/Engine/Source/Runtime/Windows/D3D11RHI/Private/D3D11RHI.cpp
+++ b/Engine/Source/Runtime/Windows/D3D11RHI/Private/D3D11RHI.cpp
@@ -38,8 +38,10 @@ void FD3D11DynamicRHI::RHIEndFrame(const FRHIEndFrameArgs& Args)
 		auto& EndWork = EmplaceProfilerEvent<UE::RHI::GPUProfiler::FEvent::FEndWork>();
 		InsertProfilerTimestamp(&EndWork.GPUTimestampBOP);
 
+		uint64 Timestamp = FPlatformTime::Cycles64();
+
 		// Insert frame boundary
-		EmplaceProfilerEvent<UE::RHI::GPUProfiler::FEvent::FFrameBoundary>(Args.FrameNumber
+		EmplaceProfilerEvent<UE::RHI::GPUProfiler::FEvent::FFrameBoundary>(Timestamp, Args.FrameNumber
 	#if WITH_RHI_BREADCRUMBS
 			, Args.GPUBreadcrumbs[ERHIPipeline::Graphics]
 	#endif
@@ -86,7 +88,7 @@ void FD3D11DynamicRHI::RHIEndFrame(const FRHIEndFrameArgs& Args)
 		}
 
 		// Start the next frame's GPU work
-		auto& BeginWork = EmplaceProfilerEvent<UE::RHI::GPUProfiler::FEvent::FBeginWork>(FPlatformTime::Cycles64());
+		auto& BeginWork = EmplaceProfilerEvent<UE::RHI::GPUProfiler::FEvent::FBeginWork>(Timestamp);
 		InsertProfilerTimestamp(&BeginWork.GPUTimestampTOP);
 	}
 
