VulkanRHI: Only apply the NeedsAllPlanes mode to textures.

#rnx

diff --git a/Engine/Source/Runtime/RHI/Private/RHICommandList.cpp b/Engine/Source/Runtime/RHI/Private/RHICommandList.cpp
--- a/Engine/Source/Runtime/RHI/Private/RHICommandList.cpp
+++ b/Engine/Source/Runtime/RHI/Private/RHICommandList.cpp
@@ -2304,7 +2304,10 @@ void FRHICommandListBase::TrackResourceStateBegin(TArrayView<const FRHITransitio
 			const bool bShouldSkipAfterState = ((EnumHasAnyFlags(TransitionInfo->Flags, EResourceTransitionFlags::IgnoreAfterState)) || (TransitionInfo->AccessAfter == ERHIAccess::Unknown));
 			const bool bNeedUpdateAfterState = (bShouldSkipAfterState == false) && (bDiscardOnAsync == false);
 
-			const bool bTrackOnEnd = NeedsAllPlanes() && EnumHasAnyFlags(TransitionCreateInfoData->SrcPipelines, ERHIPipeline::Graphics) && EnumHasAnyFlags(TransitionCreateInfoData->DstPipelines, ERHIPipeline::Graphics);
+			const bool bTrackOnEnd = NeedsAllPlanes() && 
+				(TransitionInfo->Type == FRHITransitionInfo::EType::Texture) &&
+				EnumHasAnyFlags(TransitionCreateInfoData->SrcPipelines, ERHIPipeline::Graphics) && 
+				EnumHasAnyFlags(TransitionCreateInfoData->DstPipelines, ERHIPipeline::Graphics);
 			if (bNeedUpdateAfterState && !bTrackOnEnd)
 			{
 				// Update the TrackedResource State
@@ -2362,9 +2365,9 @@ void FRHICommandListBase::TrackResourceStateEnd(TArrayView<const FRHITransition*
 				{
 					const FRHITransitionInfo* TransitionInfo = &TransitionCreateInfoData->TransitionInfos[Index];
 
-					if ((TransitionInfo->Type == FRHITransitionInfo::EType::Buffer) && (GRHIGlobals.SupportImplicitBufferStateTracking))
+					if (TransitionInfo->Type != FRHITransitionInfo::EType::Texture)
 					{
-						// This RHI supports Buffers state tracking at platform level, so we don't need to track them
+						// NeedsAllPlanes only applied to textures
 						continue;
 					}
 
