Fix assert firing in FCompilePipelineStateTask on Android.

diff --git a/Engine/Source/Runtime/RHI/Private/PipelineStateCache.cpp b/Engine/Source/Runtime/RHI/Private/PipelineStateCache.cpp
--- a/Engine/Source/Runtime/RHI/Private/PipelineStateCache.cpp
+++ b/Engine/Source/Runtime/RHI/Private/PipelineStateCache.cpp
@@ -2902,7 +2902,7 @@ public:
 	void CompilePSO(const FGraphicsPipelineStateInitializer::EPSOPrecacheCompileType* OptionalPriorityOverride = nullptr)
 	{
 		LLM_SCOPE(ELLMTag::PSO);
-		FTaskTagScope Scope(ETaskTag::EParallelRenderingThread);
+		FTaskTagScope Scope(GetDesiredThread() == ENamedThreads::RHIThread ? ETaskTag::EParallelRhiThread : ETaskTag::EParallelRenderingThread);
 
 #if WITH_RHI_BREADCRUMBS
 		if (PSOCompilationDebugData.BreadcrumbNode)
