Modified GPU profiler to make use of the "end of pipe" stats system timeline added in CL 39205780
- GPU stats are computed and published by an end-of-pipe task graph task, after GPU rendering has completed. The stats are broken unless the correct timeline in the stats system is used, to prevent stats getting attributed to the wrong frame.
- Since the stats system frame ticks in a different place than RHIEndFrame, we have to pipe the new frame number through from the render thread down to the end-of-pipe.

Combined "profilegpu" and "stat GPU" profiler sinks
- They were computing almost the same data, so made more sense from a code reuse point of view to merge them.
- Also modified how event streams are published by platform RHIs to reduce the number of loops over the data we need to do inside the profiler.

diff --git a/Engine/Source/Runtime/D3D12RHI/Private/D3D12CommandList.cpp b/Engine/Source/Runtime/D3D12RHI/Private/D3D12CommandList.cpp
--- a/Engine/Source/Runtime/D3D12RHI/Private/D3D12CommandList.cpp
+++ b/Engine/Source/Runtime/D3D12RHI/Private/D3D12CommandList.cpp
@@ -430,6 +430,9 @@ void FD3D12CommandList::WriteTimestamp(FD3D12QueryLocation const& Location, ED3D
 
 FD3D12CommandList::FState::FState(FD3D12CommandAllocator* CommandAllocator, FD3D12QueryAllocator* TimestampAllocator, FD3D12QueryAllocator* PipelineStatsAllocator)
 	: CommandAllocator(CommandAllocator)
+#if RHI_NEW_GPU_PROFILER
+	, EventStream(CommandAllocator->Device->GetQueue(CommandAllocator->QueueType).GetProfilerQueue())
+#endif
 {
 	if (TimestampAllocator)
 	{
diff --git a/Engine/Source/Runtime/D3D12RHI/Private/D3D12Device.cpp b/Engine/Source/Runtime/D3D12RHI/Private/D3D12Device.cpp
--- a/Engine/Source/Runtime/D3D12RHI/Private/D3D12Device.cpp
+++ b/Engine/Source/Runtime/D3D12RHI/Private/D3D12Device.cpp
@@ -33,6 +33,13 @@ static uint32 GetQueryHeapPoolIndex(D3D12_QUERY_HEAP_TYPE HeapType)
 	}
 }
 
+FD3D12Timing::FD3D12Timing(FD3D12Queue& Queue)
+	: Queue(Queue)
+#if RHI_NEW_GPU_PROFILER
+	, EventStream(Queue.GetProfilerQueue())
+#endif
+{}
+
 FD3D12Queue::FD3D12Queue(FD3D12Device* Device, ED3D12QueueType QueueType)
 	: Device(Device)
 	, QueueType(QueueType)
diff --git a/Engine/Source/Runtime/D3D12RHI/Private/D3D12Device.h b/Engine/Source/Runtime/D3D12RHI/Private/D3D12Device.h
--- a/Engine/Source/Runtime/D3D12RHI/Private/D3D12Device.h
+++ b/Engine/Source/Runtime/D3D12RHI/Private/D3D12Device.h
@@ -93,9 +93,7 @@ public:
 
 #endif
 
-	FD3D12Timing(FD3D12Queue& Queue)
-		: Queue(Queue)
-	{}
+	FD3D12Timing(FD3D12Queue& Queue);
 };
 
 // Encapsulates a single D3D command queue, and maintains the 
diff --git a/Engine/Source/Runtime/D3D12RHI/Private/D3D12RHI.cpp b/Engine/Source/Runtime/D3D12RHI/Private/D3D12RHI.cpp
--- a/Engine/Source/Runtime/D3D12RHI/Private/D3D12RHI.cpp
+++ b/Engine/Source/Runtime/D3D12RHI/Private/D3D12RHI.cpp
@@ -614,6 +614,9 @@ void FD3D12DynamicRHI::RHIEndFrame(const FRHIEndFrameArgs& Args)
 		#if WITH_RHI_BREADCRUMBS
 			, (Pipeline != ERHIPipeline::None) ? Args.GPUBreadcrumbs[Pipeline] : nullptr
 		#endif
+		#if STATS
+			, Args.StatsFrame
+		#endif
 		);
 	#endif
 	});
diff --git a/Engine/Source/Runtime/D3D12RHI/Private/D3D12Submission.cpp b/Engine/Source/Runtime/D3D12RHI/Private/D3D12Submission.cpp
--- a/Engine/Source/Runtime/D3D12RHI/Private/D3D12Submission.cpp
+++ b/Engine/Source/Runtime/D3D12RHI/Private/D3D12Submission.cpp
@@ -1313,6 +1313,9 @@ FD3D12DynamicRHI::FProcessResult FD3D12DynamicRHI::ProcessInterruptQueue()
 
 FD3D12PayloadBase::FD3D12PayloadBase(FD3D12Queue& Queue)
 	: Queue(Queue)
+#if RHI_NEW_GPU_PROFILER
+	, EventStream(Queue.GetProfilerQueue())
+#endif
 {}
 
 FD3D12PayloadBase::~FD3D12PayloadBase()
@@ -1344,12 +1347,14 @@ static TAutoConsoleVariable<int32> CVarGPUTimeFromTimestamps(
 void FD3D12DynamicRHI::ProcessTimestamps(FD3D12TimingArray const& TimingPerQueue)
 {
 #if RHI_NEW_GPU_PROFILER
-
-	for (auto const& Timing : TimingPerQueue)
 	{
-		UE::RHI::GPUProfiler::ProcessEvents(Timing->Queue.GetProfilerQueue(), MoveTemp(Timing->EventStream));
+		TArray<UE::RHI::GPUProfiler::FEventStream, TInlineAllocator<GD3D12MaxNumQueues>> Streams;
+		for (auto const& Timing : TimingPerQueue)
+		{
+			Streams.Add(MoveTemp(Timing->EventStream));
+		}
+		UE::RHI::GPUProfiler::ProcessEvents(Streams);
 	}
-
 #else
 
 	// The total number of cycles where at least one GPU pipe was busy during the frame.
diff --git a/Engine/Source/Runtime/RHI/Private/GPUProfiler.cpp b/Engine/Source/Runtime/RHI/Private/GPUProfiler.cpp
--- a/Engine/Source/Runtime/RHI/Private/GPUProfiler.cpp
+++ b/Engine/Source/Runtime/RHI/Private/GPUProfiler.cpp
@@ -11,6 +11,7 @@
 #include "RHI.h"
 #include "GpuProfilerTrace.h"
 #include "Containers/AnsiString.h"
+#include "Stats/StatsData.h"
 
 #if !UE_BUILD_SHIPPING
 #include "VisualizerEvents.h"
@@ -721,322 +722,728 @@ bool FGPUTiming::GAreGlobalsInitialized = false;
 
 namespace UE::RHI::GPUProfiler
 {
-#if HAS_GPU_STATS
-	// Per queue GPU stats
-
-	// Total busy time on the current queue.
-	static FGPUStat GPUStat_Total(TEXT("Queue Total"));
-
-	// Total time idle, waiting for more commands to arrive from the CPU.
-	static FGPUStat GPUStat_Idle(TEXT("Queue Idle"));
+	TLockFreePointerListUnordered<void, PLATFORM_CACHE_LINE_SIZE> FEventStream::FChunk::MemoryPool;
 
-	// Total time blocked waiting on another GPU queue / fence.
-	static FGPUStat GPUStat_Waiting(TEXT("Queue Waiting"));
-#endif
+	static TArray<FEventSink*>& GetSinks()
+	{
+		static TArray<FEventSink*> Sinks;
+		return Sinks;
+	}
 
-	TMap<FGPUStat::FStatIdentifier, TUniquePtr<FGPUStat::FStatCategory>> FGPUStat::Categories;
+	FEventSink::FEventSink()
+	{
+		GetSinks().Add(this);
+	}
 
-	FGPUStat::FStatCategory& FGPUStat::GetCategory(FStatIdentifier ID)
+	FEventSink::~FEventSink()
 	{
-		TUniquePtr<FStatCategory>& Category = Categories.FindOrAdd(ID);
-		if (!Category)
-		{
-			Category = MakeUnique<FStatCategory>(ID);
-		}
-		return *Category;
+		GetSinks().RemoveSingle(this);
 	}
 
-	FString FGPUStat::FStatIdentifier::GetIDString(bool bFriendly) const
+	void ProcessEvents(TArrayView<FEventStream> EventStreams)
 	{
-		TCHAR const* TypeStr;
-		switch (Type)
+		TRACE_CPUPROFILER_EVENT_SCOPE(UE::RHI::GPUProfiler::ProcessEvents);
+
+		TArray<TSharedRef<FEventStream>> SharedStreams;
+		SharedStreams.Reserve(EventStreams.Num());
+
+		for (FEventStream& Stream : EventStreams)
 		{
-		default: checkNoEntry(); [[fallthrough]];
-		case EType::Busy: TypeStr = TEXT("Busy"); break;
-		case EType::Idle: TypeStr = TEXT("Idle"); break;
-		case EType::Wait: TypeStr = TEXT("Wait"); break;
+			if (!Stream.IsEmpty())
+			{
+				SharedStreams.Emplace(MakeShared<FEventStream>(MoveTemp(Stream)));
+			}
 		}
 
-		if (bFriendly)
+		if (SharedStreams.Num())
 		{
-			return FString::Printf(TEXT("GPU %d %s Queue %d (%s)")
-				, Queue.GPU
-				, Queue.GetTypeString()
-				, Queue.Index
-				, TypeStr
-			);
+			for (FEventSink* Sink : GetSinks())
+			{
+				Sink->ProcessStreams(SharedStreams);
+			}
 		}
-		else
+	}
+
+	void InitializeQueues(TConstArrayView<FQueue> Queues)
+	{
+		for (FEventSink* Sink : GetSinks())
 		{
-			return FString::Printf(TEXT("GPU%d_%s%d_%s")
-				, Queue.GPU
-				, Queue.GetTypeString()
-				, Queue.Index
-				, TypeStr
-			);
+			Sink->InitializeQueues(Queues);
 		}
 	}
 
-	FGPUStat::FStatCategory_Stats::FStatCategory_Stats(FStatIdentifier const& ID)
-	#if STATS
-		: GroupName(FString::Printf(TEXT("STATGROUP_%s"), *ID.GetIDString(false)))
-		, GroupDesc(FString::Printf(TEXT("%s Timing"), *ID.GetIDString(true)))
-	#endif
-	{}
+#if WITH_PROFILEGPU
+	template <uint32 Width>
+	struct TUnicodeHorizontalBar
+	{
+		TCHAR Text[Width + 1];
 
-	FGPUStat::FStatCategory_Csv::FStatCategory_Csv(FStatIdentifier const& ID)
-	#if CSV_PROFILER_STATS
-		: Csv(*ID.GetIDString(false), true)
-	#endif
-	{}
+		// 0 <= Value <= 1
+		TUnicodeHorizontalBar(double Value)
+		{
+			TCHAR* Output = Text;
+			int32 Solid, Partial, Blank;
+			{
+				double Integer;
+				double Remainder = FMath::Modf(FMath::Clamp(Value, 0.0, 1.0) * Width, &Integer);
 
-	FGPUStat::FStatCategory::FStatCategory(FStatIdentifier const& ID)
-		: FStatCategory_Stats(ID)
-		, FStatCategory_Csv(ID)
-	{}
+				Solid = (int32)Integer;
+				Partial = (int32)FMath::Floor(Remainder * 8);
+				Blank = (Width - Solid - (Partial > 0 ? 1 : 0));
+			}
 
-#if STATS
-	TStatId FGPUStat::GetStatId(FQueue Queue, EType Type)
-	{
-		FStatIdentifier ID = { Queue, Type };
-		FStatInstance& Instance = Instances.FindOrAdd(ID);
+			// Solid characters
+			for (int32 Index = 0; Index < Solid; ++Index)
+			{
+				*Output++ = TEXT('█');
+			}
 
-		if (!Instance.Stat)
-		{
-			FStatCategory& Category = GetCategory(ID);
+			// Partially filled character
+			if (Partial > 0)
+			{
+				static constexpr TCHAR const Data[] = TEXT("▏▎▍▌▋▊▉");
+				*Output++ = Data[Partial - 1];
+			}
 
-			Instance.StatName = MakeUnique<FAnsiString>(FString::Printf(TEXT("STAT_%s_%s"), *ID.GetIDString(false), Name));
+			// Blank Characters to pad out the width
+			for (int32 Index = 0; Index < Blank; ++Index)
+			{
+				*Output++ = TEXT(' ');
+			}
 
-			Instance.Stat = MakeUnique<FDynamicStat>(
-				**Instance.StatName,
-				Name,
-				*Category.GroupName,
-				"STATCAT_GPU",
-				*Category.GroupDesc,
-				true, // IsDefaultEnabled
-				true, // IsClearEveryFrame
-				EStatDataType::ST_double,
-				false, // IsCycleStat
-				false, // SortByName
-				FPlatformMemory::MCR_Invalid
-			);
+			*Output++ = 0;
+			check(uintptr_t(Output) == (uintptr_t(Text) + sizeof(Text)));
 		}
+	};
 
-		return Instance.Stat->GetStatId();
-	}
-#endif
-
-#if CSV_PROFILER_STATS
-	FCsvDeclaredStat& FGPUStat::GetCsvStat(FQueue Queue)
+	struct FNode
 	{
-		FStatIdentifier ID = { Queue, EType::Busy };
-		FStatInstance& Instance = Instances.FindOrAdd(ID);
+		FString Name;
 
-		if (!Instance.CsvStat.IsSet())
+		FNode* Parent = nullptr;
+		FNode* Next = nullptr;
+
+		TArray<FNode*> Children;
+
+		struct FStats
 		{
-			FStatCategory& Category = GetCategory(ID);
-			Instance.CsvStat.Emplace(Name, Category.Csv.Index);
-		}
+			uint32 NumDraws      = 0;
+			uint32 NumDispatches = 0;
+			uint32 NumPrimitives = 0;
+			uint32 NumVertices   = 0;
 
-		return *Instance.CsvStat;
-	}
-#endif
+			uint64 BusyCycles    = 0;
+			uint64 IdleCycles    = 0;
+			uint64 WaitCycles    = 0;
 
-	TLockFreePointerListUnordered<void, PLATFORM_CACHE_LINE_SIZE> FEventStream::FChunk::MemoryPool;
+			double GetBusyMilliseconds() const
+			{
+				return FPlatformTime::ToMilliseconds64(BusyCycles);
+			}
 
-	static TArray<FEventSink*>& GetSinks()
-	{
-		static TArray<FEventSink*> Sinks;
-		return Sinks;
-	}
+			bool HasWork() const
+			{
+				return NumDraws > 0 || NumDispatches > 0;
+			}
 
-	FEventSink::FEventSink()
-	{
-		GetSinks().Add(this);
-	}
+			FStats& operator += (FStats const& Stats)
+			{
+				NumDraws      += Stats.NumDraws;
+				NumDispatches += Stats.NumDispatches;
+				NumPrimitives += Stats.NumPrimitives;
+				NumVertices   += Stats.NumVertices;
+				BusyCycles    += Stats.BusyCycles;
+				IdleCycles    += Stats.IdleCycles;
+				WaitCycles    += Stats.WaitCycles;
+				return *this;
+			}
 
-	FEventSink::~FEventSink()
-	{
-		GetSinks().RemoveSingle(this);
-	}
+			FStats& operator += (FEvent::FStats const& Stats)
+			{
+				NumDraws      += Stats.NumDraws;
+				NumDispatches += Stats.NumDispatches;
+				NumPrimitives += Stats.NumPrimitives;
+				NumVertices   += Stats.NumVertices;
+
+				return *this;
+			}
+
+			void Accumulate(uint64 Busy, uint64 Wait, uint64 Idle)
+			{
+				BusyCycles += Busy;
+				IdleCycles += Idle;
+				WaitCycles += Wait;
+			}
+		};
+
+		// Exclusive stats for this node
+		FStats Exclusive;
+
+		// Sum of stats including all children
+		FStats Inclusive;
 
-	void ProcessEvents(FQueue Queue, FEventStream EventStream)
+		FNode(FString&& Name)
+			: Name(MoveTemp(Name))
+		{}
+	};
+
+	struct FTable
 	{
-		if (EventStream.IsEmpty())
-		{
-			return;
-		}
+		bool const bUnicodeOutput;
 
-		TRACE_CPUPROFILER_EVENT_SCOPE(UE::RHI::GPUProfiler::ProcessEvents);
-		TSharedRef<FEventStream> SharedStream = MakeShared<FEventStream>(MoveTemp(EventStream));
+		FTable()
+			: bUnicodeOutput(GCVarProfileGPU_UnicodeOutput.GetValueOnAnyThread())
+		{}
 
-		for (FEventSink* Sink : GetSinks())
+		enum class EColumn : uint32
 		{
-			Sink->AppendStream(Queue, SharedStream);
-			Sink->ProcessStreams();
-		}
-	}
+			Exclusive_NumDraws,
+			Exclusive_NumDispatches,
+			Exclusive_NumPrimitives,
+			Exclusive_NumVertices,
+			Exclusive_Percent,
+			Exclusive_Time,
+
+			Inclusive_NumDraws,
+			Inclusive_NumDispatches,
+			Inclusive_NumPrimitives,
+			Inclusive_NumVertices,
+			Inclusive_Percent,
+			Inclusive_Time,
+
+			Events,
+
+			Num
+		};
 
-	void InitializeQueues(TConstArrayView<FQueue> Queues)
-	{
-		for (FQueue Queue : Queues)
+		uint32 GetColumnMinimumWidth(EColumn Column) const
 		{
-			for (FEventSink* Sink : GetSinks())
+			switch (Column)
 			{
-				Sink->QueueStates.Add(Queue, Sink->CreateQueueState(Queue));
+			case EColumn::Events:
+				return 6;
 			}
-		}
-	}
-
-	TUniquePtr<FEventSink::FQueueState> FEventSink::CreateQueueState(FQueue Queue) const
-	{
-		return MakeUnique<FQueueState>(Queue);
-	}
 
-	void FEventSink::AppendStream(FQueue Queue, TSharedRef<FEventStream> const& NewEventStream)
-	{
-		FQueueState& QueueState = *QueueStates.FindChecked(Queue);
-		QueueState.PendingStreams.Enqueue(FIterator(NewEventStream));
-	}
+			return 0;
+		}
 
-	void FEventSink::ProcessStreams()
-	{
-		// Process the queue as far as possible
-		bool bProgress;
-		do
+		TCHAR const* GetColumnHeader(EColumn Column) const
 		{
-			bProgress = false;
-			for (auto& [Queue, QueueState] : QueueStates)
+			switch (Column)
 			{
-				while (FIterator* Iterator = QueueState->PendingStreams.Peek())
-				{
-					FEvent const* Start = Iterator->Peek();
+			case EColumn::Exclusive_NumDraws:
+			case EColumn::Inclusive_NumDraws:
+				return TEXT("Draws");
 
-					bool bPaused = !ProcessStream(*QueueState.Get(), *Iterator);
+			case EColumn::Exclusive_NumDispatches:
+			case EColumn::Inclusive_NumDispatches:
+				return TEXT("Dsptch");
+				
+			case EColumn::Exclusive_NumPrimitives:
+			case EColumn::Inclusive_NumPrimitives:
+				return TEXT("Prim");
 
-					FEvent const* End = Iterator->Peek();
-					bProgress |= End != Start;
+			case EColumn::Exclusive_NumVertices:
+			case EColumn::Inclusive_NumVertices:
+				return TEXT("Vert");
 
-					if (bPaused)
-					{
-						// The queue was paused by a Wait event
-						check(End);
-						break;
-					}
+			case EColumn::Exclusive_Percent:
+			case EColumn::Inclusive_Percent:
+				return TEXT("Percent");
+				
+			case EColumn::Exclusive_Time:
+			case EColumn::Inclusive_Time:
+				return TEXT("Time");
+			}
 
-					if (!End)
-					{
-						// This stream has been fully processed.
-						QueueState->PendingStreams.Dequeue();
-					}
-				}
+			return TEXT("");
+		}
+
+		uint32 GetColumnGroup(EColumn Column) const
+		{
+		switch (Column)
+			{
+			case EColumn::Exclusive_NumDraws:
+			case EColumn::Exclusive_NumDispatches:
+			case EColumn::Exclusive_NumPrimitives:
+			case EColumn::Exclusive_NumVertices:
+			case EColumn::Exclusive_Percent:
+			case EColumn::Exclusive_Time:
+				return 0;
+
+			case EColumn::Inclusive_NumDraws:
+			case EColumn::Inclusive_NumDispatches:
+			case EColumn::Inclusive_NumPrimitives:
+			case EColumn::Inclusive_NumVertices:
+			case EColumn::Inclusive_Percent:
+			case EColumn::Inclusive_Time:
+				return 1;
+
+			default:
+			case EColumn::Events:
+				return 2;
 			}
-		} while (bProgress);
-	}
+		}
 
-	void FEventSink::FQueueState::ResolveSignal(FEvent::FSignalFence const& SignalFenceEvent)
-	{
-		FResolvedSignal& Result = Signals.Emplace_GetRef();
-		
-		//
-		// Take the max between the previous GPU EndWork event and the CPU timestamp. The signal cannot have happened on the GPU until the CPU has submitted the command to the driver.
-		// 
-		// An example would be a GPU queue that completes work and goes idle at time T. Later, the CPU issues a Signal without other prior work at time T + 100ms.
-		// The fence signal cannot have happened until time T + 100ms because the CPU hadn't instructed the GPU to do so until then.
-		// LastGPUCycles would still be set to time T, since that was the time of the preceeding EndWork event.
-		//
-		Result.GPUTimestampBOP = FMath::Max(LastGPUCycles, SignalFenceEvent.CPUTimestamp);
-		Result.Value = SignalFenceEvent.Value;
-
-		//
-		// Fences signals *MUST* be sequential, to remove ambiguity caused by trimming the Signals array.
-		// 
-		// To explain why, assume non-sequential signals are allowed, and consider the following example events on an arbitrary queue:
-		// 
-		//		     [Signal 2]
-		//		-- Frame Boundary --
-		//		     [Signal 4]
-		//
-		// Assume, after trimming events earlier than the frame boundary, that only [Signal 4] remains in the Signals array.
-		// Then, some other queue attempts to [Wait 3]. We need to compute when [Wait 3] is resolved with only the information about [Signal 4].
-		// 
-		// Given that fences resolve waits as soon as the signalled value is >= the wait value, we could assume the fence was resolved at [Signal 4].
-		// However, we don't know if the fence was already signalled to value 3 before the frame boundary and the trimming.
-		// 
-		// Without this information, it is ambiguous whether [Wait 3] is already resolved by a [Signal 3] before the frame boundary that is no longer
-		// in the Signals array, or won't be resolved until [Signal 4]. We could have had this sequence of events:
-		// 
-		//		     [Signal 2]
-		//		     [Signal 3]
-		//		-- Frame Boundary --
-		//		     [Signal 4]
-		// 
-		// Requiring that fences are always signalled in sequential order solves this.
-		// If the awaited value is less than the first Signal, the fence has already been signalled before the frame boundary.
-		//
-		check(Result.Value == MaxSignal.Value + 1);
-
-		// Signals should always advance in time
-		check(Result.GPUTimestampBOP >= MaxSignal.GPUTimestampBOP);
-
-		MaxSignal = Result;
-	}
+		TCHAR const* GetGroupName(uint32 GroupIndex) const
+		{
+			switch (GroupIndex)
+			{
+			case 0: return TEXT("Exclusive");
+			case 1: return TEXT("Inclusive");
+			case 2: return TEXT("Events");
+			}
+			return TEXT("");
+		}
 
-	TOptional<FEventSink::FResolvedWait> FEventSink::ResolveWait(FQueueState& LocalQueue, FEvent::FWaitFence const& WaitFenceEvent)
-	{
-		FQueueState const& RemoteQueue = *QueueStates.FindChecked(WaitFenceEvent.Queue);
+		uint32 NumRows = 0;
+		TStaticArray<TArray<FString>, uint32(EColumn::Num)> Columns { InPlace };
+		TArray<bool> RowBreaks;
 
-		if (RemoteQueue.MaxSignal.Value < WaitFenceEvent.Value)
+		FString& Col(EColumn Column)
 		{
-			// Fence has not yet been signalled on the remote queue
-			return {};
+			return Columns[uint32(Column)].Emplace_GetRef();
 		}
-		else
+
+		bool HasRows() const
 		{
-			// Fence has been signalled, but it may be in the future.
-
-			FResolvedWait Result;
-			Result.CPUTimestamp = WaitFenceEvent.CPUTimestamp;
-
-			//
-			// The wait cannot be resolved any earlier than:
-			//
-			//		1) The wait command was issued to the driver (WaitFenceEvent.CPUTimestamp)
-			//		2) The GPU completed prior work on this queue (LocalQueue.LastGPUCycles)
-			//
-			Result.GPUTimestampTOP = FMath::Max(WaitFenceEvent.CPUTimestamp, LocalQueue.LastGPUCycles);
-			//
-			//      3) The wait maybe be further delayed by the remote queue the GPU is awaiting.
-			//
-			int32 Index = Algo::LowerBoundBy(RemoteQueue.Signals, WaitFenceEvent.Value, [](FResolvedSignal const& Signal) { return Signal.Value; });
-			if (RemoteQueue.Signals.IsValidIndex(Index))
-			{
-				FResolvedSignal const& Signal = RemoteQueue.Signals[Index];
+			return NumRows > 0;
+		}
 
-				//
-				// Only consider this signal's timestamp if the fence was not already signalled at the previous frame boundary.
-				// See comment in ResolveSignal() for details.
-				//
-				if (!(Index == 0 && WaitFenceEvent.Value < Signal.Value))
-				{
-					Result.GPUTimestampTOP = FMath::Max(Result.GPUTimestampTOP, Signal.GPUTimestampBOP);
-				}
-			}
+		void AddRow(FNode* Root, FNode::FStats const& Inclusive, FNode::FStats const& Exclusive, FString const& Name, uint32 Level)
+		{
+			double ExclusivePercent = double(Exclusive.BusyCycles) / Root->Inclusive.BusyCycles;
+			double InclusivePercent = double(Inclusive.BusyCycles) / Root->Inclusive.BusyCycles;
+
+			static constexpr uint32 BarWidth = 8;
+			TUnicodeHorizontalBar<BarWidth> ExclusiveBar = ExclusivePercent;
+			TUnicodeHorizontalBar<BarWidth> InclusiveBar = InclusivePercent;
+
+			static constexpr TCHAR const BarSeparator[] = TEXT(" ┊ ");
+
+			Col(EColumn::Exclusive_NumDraws     ) = FString::Printf(TEXT("%d"), Exclusive.NumDraws);
+			Col(EColumn::Exclusive_NumDispatches) = FString::Printf(TEXT("%d"), Exclusive.NumDispatches);
+			Col(EColumn::Exclusive_NumPrimitives) = FString::Printf(TEXT("%d"), Exclusive.NumPrimitives);
+			Col(EColumn::Exclusive_NumVertices  ) = FString::Printf(TEXT("%d"), Exclusive.NumVertices);
+			Col(EColumn::Exclusive_Percent      ) = FString::Printf(TEXT("%.1f%%%s%s"), ExclusivePercent * 100.0, bUnicodeOutput ? BarSeparator : TEXT(""), bUnicodeOutput ? ExclusiveBar.Text : TEXT(""));
+			Col(EColumn::Exclusive_Time         ) = FString::Printf(TEXT("%.3f ms"), FPlatformTime::ToMilliseconds64(Exclusive.BusyCycles));
+
+			Col(EColumn::Inclusive_NumDraws     ) = FString::Printf(TEXT("%d"), Inclusive.NumDraws);
+			Col(EColumn::Inclusive_NumDispatches) = FString::Printf(TEXT("%d"), Inclusive.NumDispatches);
+			Col(EColumn::Inclusive_NumPrimitives) = FString::Printf(TEXT("%d"), Inclusive.NumPrimitives);
+			Col(EColumn::Inclusive_NumVertices  ) = FString::Printf(TEXT("%d"), Inclusive.NumVertices);
+			Col(EColumn::Inclusive_Percent      ) = FString::Printf(TEXT("%.1f%%%s%s"), InclusivePercent * 100.0, bUnicodeOutput ? BarSeparator : TEXT(""), bUnicodeOutput ? InclusiveBar.Text : TEXT(""));
+			Col(EColumn::Inclusive_Time         ) = FString::Printf(TEXT("%.3f ms"), FPlatformTime::ToMilliseconds64(Inclusive.BusyCycles));
+				
+			static constexpr uint32 SpacesPerIndent = 3;
+			Col(EColumn::Events) = FString::Printf(TEXT("%*s"), Name.Len() + (Level * SpacesPerIndent), *Name);
 
-			return Result;
+			// Insert a horizontal rule before each root level row.
+			RowBreaks.Add(Level == 0);
+
+			NumRows++;
 		}
-	}
 
-	// Handles computing the "stat unit" GPU time, and "stat gpu" stats.
-	struct FGPUProfilerSink_StatSystem final : public FEventSink
-	{
-		class FTimestampStream
+		struct FChars
 		{
-		private:
-			TArray<uint64> Values;
+			TCHAR const* Left;
+			TCHAR const* GroupSeparator;
+			TCHAR const* LastGroupSeparator;
+			TCHAR const* Right;
+			TCHAR const* CellSeparator;
+		};
 
-		public:
-			struct FState
+		struct FFormat
+		{
+			TCHAR const* LineMajor;
+			TCHAR const* LineMinor;
+
+			FChars const TopRow;
+			FChars const GroupNameRow;
+			FChars const GroupBorderRow;
+			FChars const ValueRow;
+			FChars const DividorRow;
+			FChars const BottomRow;
+		};
+
+		FString ToString() const
+		{
+			if (bUnicodeOutput)
+			{
+				static constexpr FFormat Unicode = 
+				{
+					.LineMajor  = TEXT("━"),
+					.LineMinor  = TEXT("─"),
+
+					//                     Left     GrpSep     LastGrp     Right     CellSep
+					.TopRow        { TEXT("    ┏"), TEXT("┳"), TEXT("┳"), TEXT("┓"), TEXT(" ") },
+					.GroupNameRow  { TEXT("    ┃"), TEXT("┃"), TEXT("┃"), TEXT("┃"), TEXT(" ") },
+					.GroupBorderRow{ TEXT("    ┠"), TEXT("╂"), TEXT("┨"), TEXT("┃"), TEXT("┬") },
+					.ValueRow      { TEXT("    ┃"), TEXT("┃"), TEXT("┃"), TEXT("┃"), TEXT("│") },
+					.DividorRow    { TEXT("    ┠"), TEXT("╂"), TEXT("╂"), TEXT("┨"), TEXT("┼") },
+					.BottomRow     { TEXT("    ┗"), TEXT("┻"), TEXT("┻"), TEXT("┛"), TEXT("┷") },
+				};
+
+				return ToStringInner(Unicode);
+			}
+			else
+			{
+				static constexpr FFormat Ascii = 
+				{
+					.LineMajor  = TEXT("-"),
+					.LineMinor  = TEXT("-"),
+
+					//                     Left     GrpSep     LastGrp     Right     CellSep
+					.TopRow        { TEXT("    +"), TEXT("+"), TEXT("+"), TEXT("+"), TEXT(" ") },
+					.GroupNameRow  { TEXT("    |"), TEXT("|"), TEXT("|"), TEXT("|"), TEXT(" ") },
+					.GroupBorderRow{ TEXT("    +"), TEXT("+"), TEXT("+"), TEXT("|"), TEXT("+") },
+					.ValueRow      { TEXT("    |"), TEXT("|"), TEXT("|"), TEXT("|"), TEXT("|") },
+					.DividorRow    { TEXT("    +"), TEXT("+"), TEXT("+"), TEXT("+"), TEXT("+") },
+					.BottomRow     { TEXT("    +"), TEXT("+"), TEXT("+"), TEXT("+"), TEXT("+") },
+				};
+
+				return ToStringInner(Ascii);
+			}
+		}
+
+		FString ToStringInner(FFormat const& Format) const
+		{
+			struct FGroup  { uint32 Index, Width; };
+			struct FColumn { uint32 Index, Width; };
+
+			static constexpr uint32 NumGroups = 3;
+			static constexpr uint32 CellPadding = 1;
+
+			// Auto-size column widths to their contents
+			TStaticArray<int32, uint32(EColumn::Num)> ColumnWidths{ InPlace, 0 };
+			for (uint32 ColumnIndex = 0; ColumnIndex < uint32(EColumn::Num); ++ColumnIndex)
+			{
+				if (Columns[ColumnIndex].Num() == 0)
+					continue;
+
+				check(Columns[ColumnIndex].Num() == NumRows);
+
+				int32& Width = ColumnWidths[ColumnIndex];
+
+				// Auto-size column width
+				Width = GetColumnMinimumWidth(EColumn(ColumnIndex));
+				Width = FMath::Max(Width, FCString::Strlen(GetColumnHeader(EColumn(ColumnIndex))));
+
+				for (FString const& Cell : Columns[ColumnIndex])
+				{
+					Width = FMath::Max(Width, Cell.Len());
+				}
+			}
+
+			FString Result;
+
+			auto EmitGroupRow = [&](FChars const& Chars, TUniqueFunction<void(FGroup)> GroupCallback)
+			{
+				uint32 const CellSeparatorLength = FCString::Strlen(Chars.CellSeparator);
+
+				Result += Chars.Left;
+
+				uint32 GroupWidth = 0;
+				uint32 GroupIndex = 0;
+
+				for (uint32 ColumnIndex = 0; ColumnIndex < uint32(EColumn::Num); ++ColumnIndex)
+				{
+					if (Columns[ColumnIndex].Num() == 0)
+						continue;
+
+					GroupWidth += ColumnWidths[ColumnIndex] + CellPadding * 2;
+					GroupIndex = GetColumnGroup(EColumn(ColumnIndex));
+
+					if (GroupIndex != GetColumnGroup(EColumn(ColumnIndex + 1)))
+					{ 
+						// Group Change
+						GroupCallback({ GroupIndex, GroupWidth });
+
+						// Add the group separator character
+						Result += GroupIndex < NumGroups - 2
+							? Chars.GroupSeparator
+							: Chars.LastGroupSeparator;
+
+						GroupWidth = 0;
+					}
+					else if (ColumnIndex < uint32(EColumn::Num) - 1)
+					{
+						// Same group. Count the (missing) cell division
+						GroupWidth += CellSeparatorLength;
+					}
+				}
+
+				// Emit final group
+				GroupCallback({ GroupIndex, GroupWidth });
+
+				// Close the row
+				Result += Chars.Right;
+				Result += TEXT("\n");
+			};
+
+			auto EmitValueRow = [&](FChars const& Chars, TUniqueFunction<void(FColumn)> CellCallback)
+			{
+				Result += Chars.Left;
+
+				for (uint32 ColumnIndex = 0; ColumnIndex < uint32(EColumn::Num); ++ColumnIndex)
+				{
+					if (Columns[ColumnIndex].Num() == 0)
+						continue;
+
+					CellCallback({ ColumnIndex, ColumnWidths[ColumnIndex] + (CellPadding * 2) });
+
+					if (ColumnIndex < uint32(EColumn::Num) - 1)
+					{
+						uint32 GroupIndex = GetColumnGroup(EColumn(ColumnIndex));
+						if (GroupIndex != GetColumnGroup(EColumn(ColumnIndex + 1)))
+						{
+							// Group change, add the group separator
+							Result += GroupIndex < NumGroups - 2
+								? Chars.GroupSeparator
+								: Chars.LastGroupSeparator;
+						}
+						else
+						{
+							// Same group, add the cell separator
+							Result += Chars.CellSeparator;
+						}
+					}
+				}
+
+				// Close the row
+				Result += Chars.Right;
+				Result += TEXT("\n");
+			};
+
+			auto AlignCenter = [&](TCHAR const* Str, uint32 Width)
+			{
+				int32 PaddingLeft = FMath::Max(0, int32(Width) - FCString::Strlen(Str));
+				int32 PaddingRight = (PaddingLeft / 2) + (PaddingLeft & 1);
+				PaddingLeft /= 2;
+
+				Result += FString::Printf(TEXT("%*s%s%*s"), PaddingLeft, TEXT(""), Str, PaddingRight, TEXT(""));
+			};
+
+			// Top Border
+			EmitGroupRow(Format.TopRow, [&](FGroup Group)
+			{
+				while (Group.Width--)
+				{
+					Result += Format.LineMajor;
+				}
+			});
+
+			// Exclusive / Inclusive Group Row
+			EmitGroupRow(Format.GroupNameRow, [&](FGroup Group)
+			{
+				TCHAR const* Str = Group.Index != GetColumnGroup(EColumn::Events)
+					? GetGroupName(Group.Index)
+					: TEXT("");
+
+				AlignCenter(Str, Group.Width);
+			});
+
+			// Events Group Row
+			EmitValueRow(Format.GroupBorderRow, [&](FColumn Column)
+			{
+				if (Column.Index == uint32(EColumn::Events))
+				{
+					AlignCenter(GetGroupName(GetColumnGroup(EColumn::Events)), Column.Width);
+				}
+				else
+				{
+					while (Column.Width--)
+					{
+						Result += Format.LineMinor;
+					}
+				}
+			});
+
+			// Header Row
+			EmitValueRow(Format.ValueRow, [&](FColumn Column)
+			{
+				AlignCenter(GetColumnHeader(EColumn(Column.Index)), Column.Width);
+			});
+
+			// Header Border Row
+			EmitValueRow(Format.DividorRow, [&](FColumn Column)
+			{
+				while (Column.Width--)
+				{
+					Result += Format.LineMinor;
+				}
+			});
+
+			// Value rows
+			for (uint32 RowIndex = 0; RowIndex < NumRows; ++RowIndex)
+			{
+				if (RowIndex > 0 && RowBreaks[RowIndex])
+				{
+					// Add a horizontal rule
+					EmitValueRow(Format.DividorRow, [&](FColumn Column)
+					{
+						while (Column.Width--)
+						{
+							Result += Format.LineMinor;
+						}
+					});
+				}
+
+				EmitValueRow(Format.ValueRow, [&](FColumn Column)
+				{
+					int32 Width = Column.Width - (CellPadding * 2);
+					if (EColumn(Column.Index) == EColumn::Events)
+					{
+						Width = -Width; // Align left
+					}
+
+					FString const& Cell = Columns[Column.Index][RowIndex];
+					Result += FString::Printf(TEXT("%*s%*s%*s")
+						, CellPadding, TEXT("")
+						, Width, *Cell
+						, CellPadding, TEXT(""));
+				});
+			}
+
+			// Bottom Border
+			EmitValueRow(Format.BottomRow, [&](FColumn Column)
+			{
+				while (Column.Width--)
+				{
+					Result += Format.LineMajor;
+				}
+			});
+
+			return Result;
+		}
+	};
+#endif
+
+#if HAS_GPU_STATS
+	// Per queue GPU stats
+
+	// Total busy time on the current queue.
+	static FGPUStat GPUStat_Total(TEXT("Queue Total"));
+#endif
+
+	TMap<FGPUStat::FStatIdentifier, TUniquePtr<FGPUStat::FStatCategory>> FGPUStat::Categories;
+
+	FGPUStat::FStatCategory& FGPUStat::GetCategory(FStatIdentifier ID)
+	{
+		TUniquePtr<FStatCategory>& Category = Categories.FindOrAdd(ID);
+		if (!Category)
+		{
+			Category = MakeUnique<FStatCategory>(ID);
+		}
+		return *Category;
+	}
+
+	FString FGPUStat::FStatIdentifier::GetIDString(bool bFriendly) const
+	{
+		TCHAR const* TypeStr;
+		switch (Type)
+		{
+		default: checkNoEntry(); [[fallthrough]];
+		case EType::Busy: TypeStr = TEXT("Busy"); break;
+		case EType::Idle: TypeStr = TEXT("Idle"); break;
+		case EType::Wait: TypeStr = TEXT("Wait"); break;
+		}
+
+		if (bFriendly)
+		{
+			return FString::Printf(TEXT("GPU %d %s Queue %d (%s)")
+				, Queue.GPU
+				, Queue.GetTypeString()
+				, Queue.Index
+				, TypeStr
+			);
+		}
+		else
+		{
+			return FString::Printf(TEXT("GPU%d_%s%d_%s")
+				, Queue.GPU
+				, Queue.GetTypeString()
+				, Queue.Index
+				, TypeStr
+			);
+		}
+	}
+
+	FGPUStat::FStatCategory_Stats::FStatCategory_Stats(FStatIdentifier const& ID)
+	#if STATS
+		: GroupName(FString::Printf(TEXT("STATGROUP_%s"), *ID.GetIDString(false)))
+		, GroupDesc(FString::Printf(TEXT("%s Timing"), *ID.GetIDString(true)))
+	#endif
+	{}
+
+	FGPUStat::FStatCategory_Csv::FStatCategory_Csv(FStatIdentifier const& ID)
+	#if CSV_PROFILER_STATS
+		: Csv(*ID.GetIDString(false), true)
+	#endif
+	{}
+
+	FGPUStat::FStatCategory::FStatCategory(FStatIdentifier const& ID)
+		: FStatCategory_Stats(ID)
+		, FStatCategory_Csv(ID)
+	{}
+
+#if STATS
+	TStatId FGPUStat::GetStatId(FQueue Queue, EType Type)
+	{
+		FStatIdentifier ID = { Queue, Type };
+		FStatInstance& Instance = Instances.FindOrAdd(ID);
+
+		if (!Instance.Stat)
+		{
+			FStatCategory& Category = GetCategory(ID);
+
+			Instance.StatName = MakeUnique<FAnsiString>(FString::Printf(TEXT("STAT_%s_%s"), *ID.GetIDString(false), Name));
+
+			Instance.Stat = MakeUnique<FDynamicStat>(
+				**Instance.StatName,
+				Name,
+				*Category.GroupName,
+				"STATCAT_GPU",
+				*Category.GroupDesc,
+				true, // IsDefaultEnabled
+				true, // IsClearEveryFrame
+				EStatDataType::ST_double,
+				false, // IsCycleStat
+				false, // SortByName
+				FPlatformMemory::MCR_Invalid
+			);
+		}
+
+		return Instance.Stat->GetStatId();
+	}
+#endif
+
+#if CSV_PROFILER_STATS
+	FCsvDeclaredStat& FGPUStat::GetCsvStat(FQueue Queue)
+	{
+		FStatIdentifier ID = { Queue, EType::Busy };
+		FStatInstance& Instance = Instances.FindOrAdd(ID);
+
+		if (!Instance.CsvStat.IsSet())
+		{
+			FStatCategory& Category = GetCategory(ID);
+			Instance.CsvStat.Emplace(Name, Category.Csv.Index);
+		}
+
+		return *Instance.CsvStat;
+	}
+#endif
+
+	// Handles computing the "stat unit" GPU time, "stat gpu" stats, and "profilegpu".
+	struct FGPUProfilerSink_StatSystem final : public FEventSink
+	{
+		class FTimestampStream
+		{
+		private:
+			TArray<uint64> Values;
+
+		public:
+			struct FState
 			{
 				FTimestampStream const& Stream;
 				int32 TimestampIndex = 0;
@@ -1163,6 +1570,7 @@ namespace UE::RHI::GPUProfiler
 
 			void EmitResults(FQueue Queue
 			#if HAS_GPU_STATS
+				, FEndOfPipeStats* Stats
 				, FGPUStat& GPUStat
 			#endif
 			#if CSV_PROFILER_STATS
@@ -1171,9 +1579,9 @@ namespace UE::RHI::GPUProfiler
 			) const
 			{
 			#if HAS_GPU_STATS
-				SET_FLOAT_STAT_FName(GPUStat.GetStatId(Queue, FGPUStat::EType::Busy).GetName(), FPlatformTime::ToMilliseconds64(BusyCycles));
-				SET_FLOAT_STAT_FName(GPUStat.GetStatId(Queue, FGPUStat::EType::Idle).GetName(), FPlatformTime::ToMilliseconds64(IdleCycles));
-				SET_FLOAT_STAT_FName(GPUStat.GetStatId(Queue, FGPUStat::EType::Wait).GetName(), FPlatformTime::ToMilliseconds64(WaitCycles));
+				Stats->AddMessage(GPUStat.GetStatId(Queue, FGPUStat::EType::Busy).GetName(), EStatOperation::Set, FPlatformTime::ToMilliseconds64(BusyCycles));
+				Stats->AddMessage(GPUStat.GetStatId(Queue, FGPUStat::EType::Idle).GetName(), EStatOperation::Set, FPlatformTime::ToMilliseconds64(IdleCycles));
+				Stats->AddMessage(GPUStat.GetStatId(Queue, FGPUStat::EType::Wait).GetName(), EStatOperation::Set, FPlatformTime::ToMilliseconds64(WaitCycles));
 			#endif
 
 			#if CSV_PROFILER_STATS
@@ -1198,1242 +1606,793 @@ namespace UE::RHI::GPUProfiler
 		#endif
 		};
 
-		struct FQueueState : public FEventSink::FQueueState
+		struct FResolvedWait
 		{
-			FQueueTimestamps Timestamps;
-			uint64 CPUBeginWorkTime = 0;
+			uint64 GPUTimestampTOP = 0;
+			uint64 CPUTimestamp = 0;
+		};
 
-		#if WITH_RHI_BREADCRUMBS
-			TMap<FRHIBreadcrumbData_Stats, int32> ActiveStats;
-			FRHIBreadcrumbNode* Breadcrumb = nullptr;
-		#endif
+		struct FResolvedSignal
+		{
+			uint64 GPUTimestampBOP = 0;
+			uint64 Value = 0;
+		};
 
-			bool bBusy = false;
-
-			FQueueState(FQueue const& Queue)
-				: FEventSink::FQueueState(Queue)
-			{}
-
-			void AccumulateTime(uint64 Busy, uint64 Wait, uint64 Idle)
-			{
-				Timestamps.WholeQueueStat.Accumulate(Busy, Wait, Idle);
-
-			#if WITH_RHI_BREADCRUMBS
-				// Apply the timings to all active stats
-				for (auto const& [Stat, RefCount] : ActiveStats)
-				{
-					Timestamps.Stats.FindChecked(Stat).Accumulate(Busy, Wait, Idle);
-				}
-			#endif
-			}
-
-			void BeginWork(uint64 GPUTimestampTOP, uint64 CPUTimestamp)
-			{
-				Timestamps.Queue.AddTimestamp(GPUTimestampTOP, true);
-
-				uint64 Idle = CPUTimestamp > LastGPUCycles
-					? CPUTimestamp - LastGPUCycles
-					: 0;
-
-				AccumulateTime(0, 0, Idle);
-				LastGPUCycles = FMath::Max(LastGPUCycles, GPUTimestampTOP);
-			}
-
-			void EndWork(uint64 GPUTimestampBOP)
-			{
-				Timestamps.Queue.AddTimestamp(GPUTimestampBOP, false);
-
-				uint64 Busy = GPUTimestampBOP > LastGPUCycles
-					? GPUTimestampBOP - LastGPUCycles
-					: 0;
-
-				AccumulateTime(Busy, 0, 0);
-				LastGPUCycles = FMath::Max(LastGPUCycles, GPUTimestampBOP);
-			}
-
-			void BeginStat(FRHIBreadcrumbData_Stats const& Stat, uint64 GPUTimestampTOP)
-			{
-				uint64 Busy = GPUTimestampTOP > LastGPUCycles
-					? GPUTimestampTOP - LastGPUCycles
-					: 0;
-
-				AccumulateTime(Busy, 0, 0);
-				LastGPUCycles = FMath::Max(LastGPUCycles, GPUTimestampTOP);
-
-				if (Stat.ShouldComputeStat())
-				{
-					// Disregard the stat if it is nested within itself (i.e. its already in the ActiveStats map with a non-zero ref count).
-					// Only the outermost stat will count the busy time, otherwise we'd be double-counting the nested time.
-					int32 RefCount = ActiveStats.FindOrAdd(Stat)++;
-					if (RefCount == 0)
-					{
-						Timestamps.Stats.FindOrAdd(Stat);
-					}
-				}
-			}
-
-			void EndStat(FRHIBreadcrumbData_Stats const& Stat, uint64 GPUTimestampBOP)
-			{
-				uint64 Busy = GPUTimestampBOP > LastGPUCycles
-					? GPUTimestampBOP - LastGPUCycles
-					: 0;
-
-				AccumulateTime(Busy, 0, 0);
-				LastGPUCycles = FMath::Max(LastGPUCycles, GPUTimestampBOP);
-
-				if (Stat.ShouldComputeStat())
-				{
-					// Pop the stat when the refcount hits zero.
-					int32 RefCount = --ActiveStats.FindChecked(Stat);
-					if (RefCount == 0)
-					{
-						ActiveStats.FindAndRemoveChecked(Stat);
-					}
-				}
-			}
-
-			void Wait(FResolvedWait const& ResolvedWait)
-			{
-				// Time the queue was idle between the last EndWork event, and the Wait command being submitted to the GPU driver.
-				uint64 Idle = ResolvedWait.CPUTimestamp > LastGPUCycles
-					? ResolvedWait.CPUTimestamp - LastGPUCycles
-					: 0;
-
-				uint64 WaitStart = FMath::Max(ResolvedWait.CPUTimestamp, LastGPUCycles);
-
-				// Time the queue spent waiting for the fence to signal on another queue.
-				uint64 Wait = ResolvedWait.GPUTimestampTOP > WaitStart
-					? ResolvedWait.GPUTimestampTOP - WaitStart
-					: 0;
-
-				// Bring the last GPU busy end time forwards to where the wait is resolved.
-				LastGPUCycles = ResolvedWait.GPUTimestampTOP;
-
-				AccumulateTime(0, Wait, Idle);
-			}
-
-			void TrimSignals(uint64 CPUTimestamp)
-			{
-				// Remove all signals that occured on the GPU timeline before this frame boundary on the CPU.
-				int32 Index = Algo::LowerBoundBy(Signals, CPUTimestamp, [](FResolvedSignal const& Signal) { return Signal.GPUTimestampBOP; });
-				if (Index >= 0)
-				{
-					Signals.RemoveAt(0, Index, EAllowShrinking::No);
-				}
-			}
-		};
-
-		virtual TUniquePtr<FEventSink::FQueueState> CreateQueueState(FQueue Queue) const override
-		{
-			return MakeUnique<FQueueState>(Queue);
-		}
-
-		using FFrameState = TMap<FQueue, FQueueTimestamps>;
-		TMap<uint32, FFrameState> Frames;
-
-		bool ProcessStream(FEventSink::FQueueState& BaseQueueState, FIterator& Iterator) override
+		struct FFrameState : TMap<FQueue, FQueueTimestamps>
 		{
-			FQueueState& QueueState = static_cast<FQueueState&>(BaseQueueState);
-
-			while (FEvent const* Event = Iterator.Peek())
-			{
-				switch (Event->GetType())
-				{
-				case FEvent::EType::BeginWork:
-					{
-						check(!QueueState.bBusy);
-						QueueState.bBusy = true;
-
-						FEvent::FBeginWork const& BeginWork = Event->Value.Get<FEvent::FBeginWork>();
-
-						check(QueueState.LastCPUCycles <= BeginWork.CPUTimestamp);
-						QueueState.LastCPUCycles = BeginWork.CPUTimestamp;
-
-						QueueState.BeginWork(BeginWork.GPUTimestampTOP, BeginWork.CPUTimestamp);
-					}
-					break;
-
-				case FEvent::EType::EndWork:
-					{
-						check(QueueState.bBusy);
-						QueueState.bBusy = false;
-
-						FEvent::FEndWork const& EndWork = Event->Value.Get<FEvent::FEndWork>();
-						QueueState.EndWork(EndWork.GPUTimestampBOP);
-					}
-					break;
-
-			#if WITH_RHI_BREADCRUMBS
-				case FEvent::EType::BeginBreadcrumb:
-					{
-						check(QueueState.bBusy);
-
-						FEvent::FBeginBreadcrumb const& BeginBreadcrumb = Event->Value.Get<FEvent::FBeginBreadcrumb>();
-						FRHIBreadcrumbData_Stats const& Stat = BeginBreadcrumb.Breadcrumb->Name.Data;
-
-						QueueState.BeginStat(Stat, BeginBreadcrumb.GPUTimestampTOP);
-						QueueState.Breadcrumb = BeginBreadcrumb.Breadcrumb;
-					}
-					break;
-
-				case FEvent::EType::EndBreadcrumb:
-					{
-						check(QueueState.bBusy);
-
-						FEvent::FEndBreadcrumb const& EndBreadcrumb = Event->Value.Get<FEvent::FEndBreadcrumb>();
-						FRHIBreadcrumbData_Stats const& Stat = EndBreadcrumb.Breadcrumb->Name.Data;
-
-						QueueState.EndStat(Stat, EndBreadcrumb.GPUTimestampBOP);
-						QueueState.Breadcrumb = EndBreadcrumb.Breadcrumb->GetParent();
-					}
-					break;
-			#endif // WITH_RHI_BREADCRUMBS
-
-				case FEvent::EType::SignalFence:
-					{
-						check(!QueueState.bBusy);
-						FEvent::FSignalFence const& SignalFence = Event->Value.Get<FEvent::FSignalFence>();
-
-						check(QueueState.LastCPUCycles <= SignalFence.CPUTimestamp);
-						QueueState.LastCPUCycles = SignalFence.CPUTimestamp;
-
-						QueueState.ResolveSignal(SignalFence);
-					}
-					break;
-
-				case FEvent::EType::WaitFence:
-					{
-						check(!QueueState.bBusy);
-						FEvent::FWaitFence const& WaitFence = Event->Value.Get<FEvent::FWaitFence>();
-
-						check(QueueState.LastCPUCycles <= WaitFence.CPUTimestamp);
-						QueueState.LastCPUCycles = WaitFence.CPUTimestamp;
-
-						TOptional<FResolvedWait> ResolvedWait = ResolveWait(QueueState, WaitFence);
-
-						if (!ResolvedWait.IsSet())
-						{
-							// Unresolved fence, pause processing
-							return false;
-						}
-
-						QueueState.Wait(*ResolvedWait);
-					}
-					break;
-
-				case FEvent::EType::FrameBoundary:
-					{
-						check(!QueueState.bBusy);
-						FEvent::FFrameBoundary const& FrameBoundary = Event->Value.Get<FEvent::FFrameBoundary>();
-
-						check(QueueState.LastCPUCycles <= FrameBoundary.CPUTimestamp);
-						QueueState.LastCPUCycles = FrameBoundary.CPUTimestamp;
-
-						QueueState.Timestamps.CPUFrameBoundary = FrameBoundary.CPUTimestamp;
-
-						FFrameState& FrameState = Frames.FindOrAdd(FrameBoundary.FrameNumber);
-						FrameState.Emplace(QueueState.Queue, MoveTemp(QueueState.Timestamps));
-
-					#if WITH_RHI_BREADCRUMBS
-						// Reinsert timestamp streams for the current active stats on 
-						// this queue, since these got moved into the frame state.
-						for (auto& [Stat, RefCount] : QueueState.ActiveStats)
-						{
-							QueueState.Timestamps.Stats.FindOrAdd(Stat);
-						}
-					#endif
-
-						if (FrameState.Num() == QueueStates.Num())
-						{
-							// Trim the Signals array in each queue, up to the lowest frame boundary CPU timestamp.
-							uint64 MinFrameBoundary = TNumericLimits<uint64>::Max();
-							{
-								for (auto& [Queue, LocalQueueState] : FrameState)
-								{
-									MinFrameBoundary = FMath::Min(MinFrameBoundary, LocalQueueState.CPUFrameBoundary);
-								}
-							}
-
-							{
-								for (auto& [Queue, LocalQueueState] : QueueStates)
-								{
-									static_cast<FQueueState&>(*LocalQueueState.Get()).TrimSignals(MinFrameBoundary);
-								}
-							}
-
-							// All registered queues have reported their frame boundary event.
-							// We have a full set of data to compute the total frame GPU stats.
-							ProcessFrame(FrameState);
-
-							Frames.Remove(FrameBoundary.FrameNumber);
-						}
-					}
-					break;
-				}
-
-				Iterator.Pop();
-			}
-
-			return true;
-		}
-
-		void ProcessFrame(FFrameState& FrameState)
-		{
-		#if CSV_PROFILER_STATS
-			const bool bCsvStatsEnabled = !!CVarGPUCsvStatsEnabled.GetValueOnAnyThread();
-			FCsvProfiler* CsvProfiler = FCsvProfiler::Get();
-			CsvProfiler->BeginFrameEOP();
-		#else
-			const bool bCsvStatsEnabled = false;
+		#if STATS
+			TOptional<int64> StatsFrame;
 		#endif
-
-			for (auto const& [Queue, QueueTimestamps] : FrameState)
-			{
-			#if WITH_RHI_BREADCRUMBS
-				// Compute the individual GPU stats
-				for (auto const& [Stat, StatState] : QueueTimestamps.Stats)
-				{
-					StatState.EmitResults(Queue
-					#if HAS_GPU_STATS
-						, *Stat.GPUStat
-					#endif
-					#if CSV_PROFILER_STATS
-						, bCsvStatsEnabled ? CsvProfiler : nullptr
-					#endif
-					);
-				}
-			#endif // WITH_RHI_BREADCRUMBS
-
-				// Set the whole-frame per queue stat
-				QueueTimestamps.WholeQueueStat.EmitResults(Queue
-				#if HAS_GPU_STATS
-					, GPUStat_Total
-				#endif
-				#if CSV_PROFILER_STATS
-					, bCsvStatsEnabled ? CsvProfiler : nullptr
-				#endif
-				);
-			}
-
-			// Compute the whole-frame total GPU time.
-			TArray<FTimestampStream::FState, TInlineAllocator<GetRHIPipelineCount() * MAX_NUM_GPUS>> StreamPointers;
-			for (auto const& [Queue, State] : FrameState)
-			{
-				StreamPointers.Emplace(State.Queue);
-			}
-			uint64 WholeFrameUnion = FTimestampStream::ComputeUnion(StreamPointers);
-
-			// Update the global GPU frame time stats - need to convert to Cycles32 rather than Cycles64.
-			GGPUFrameTime = FPlatformMath::TruncToInt(FPlatformTime::ToSeconds64(WholeFrameUnion) / FPlatformTime::GetSecondsPerCycle());
-
-			// @todo set global csv GPU time
-			//RHISetGPUStatTotals(bCsvStatsEnabled, FPlatformTime::ToMilliseconds64(WholeFrameUnion));
-		}
-
-	} GGPUProfilerSink_StatSystem;
-
-#if WITH_PROFILEGPU
-	struct FGPUProfilerSink_ProfileGPU final : public FEventSink
-	{
-		struct FNode
-		{
-			FString Name;
-
-			FNode* Parent = nullptr;
-			FNode* Next = nullptr;
-
-			TArray<FNode*> Children;
-
-			uint64 LastCycles = 0;
-
-			struct FStats
-			{
-				uint32 NumDraws      = 0;
-				uint32 NumDispatches = 0;
-				uint32 NumPrimitives = 0;
-				uint32 NumVertices   = 0;
-				uint64 BusyCycles    = 0;
-
-				double GetBusyMilliseconds() const
-				{
-					return FPlatformTime::ToMilliseconds64(BusyCycles);
-				}
-
-				bool HasWork() const
-				{
-					return NumDraws > 0 || NumDispatches > 0;
-				}
-
-				FStats& operator += (FStats const& Stats)
-				{
-					NumDraws      += Stats.NumDraws;
-					NumDispatches += Stats.NumDispatches;
-					NumPrimitives += Stats.NumPrimitives;
-					NumVertices   += Stats.NumVertices;
-					BusyCycles    += Stats.BusyCycles;
-					return *this;
-				}
-
-				FStats& operator += (FEvent::FStats const& Stats)
-				{
-					NumDraws      += Stats.NumDraws;
-					NumDispatches += Stats.NumDispatches;
-					NumPrimitives += Stats.NumPrimitives;
-					NumVertices   += Stats.NumVertices;
-
-					return *this;
-				}
-			};
-
-			// Exclusive stats for this node
-			FStats Exclusive;
-
-			// Sum of stats including all children
-			FStats Inclusive;
-
-			FNode(FString&& Name)
-				: Name(MoveTemp(Name))
-			{}
 		};
 
-		struct FQueueState : public FEventSink::FQueueState
+		struct FQueueState
 		{
-			enum class EMode
-			{
-				Idle,
-				WaitingFrame,
-				Active
-			};
-
-			struct
-			{
-				TArray<TUniquePtr<FNode>> Nodes;
-				FNode* Current = nullptr;
-				FNode* Prev    = nullptr;
-				FNode* First   = nullptr;
-
-				uint64 LastBeginCycles = 0;
-				uint64 LastEndCycles   = 0;
-
-				EMode Mode = EMode::Idle;
-
-				void PushNode(FString&& Name)
-				{
-					FNode* Parent = Current;
-					Current = Nodes.Emplace_GetRef(MakeUnique<FNode>(MoveTemp(Name))).Get();
-					Current->Parent = Parent;
-
-					if (!First)
-					{
-						First = Current;
-					}
-
-					if (Parent)
-					{
-						Parent->Children.Add(Current);
-					}
-
-					if (Prev)
-					{
-						Prev->Next = Current;
-					}
-					Prev = Current;
-				}
-
-				void PopNode()
-				{
-					check(Current && Current->Parent);
-					Current = Current->Parent;
-				}
-
-			} State {};
-
-			std::atomic<bool> bTriggerProfile { false };
-
-			FQueueState(FQueue Queue)
-				: FEventSink::FQueueState(Queue)
-			{}
-
-			void ProcessEvents(FIterator& Iterator);
-
-			void LogTree(uint32 FrameNumber);
-		};
-
-		virtual TUniquePtr<FEventSink::FQueueState> CreateQueueState(FQueue Queue) const override
-		{
-			return MakeUnique<FQueueState>(Queue);
-		}
-
-		bool ProcessStream(FEventSink::FQueueState& QueueState, FIterator& Iterator) override
-		{
-			static_cast<FQueueState&>(QueueState).ProcessEvents(Iterator);
-			return true;
-		}
-
-		void ProfileNextFrame()
-		{
-			for (auto& [Queue, State] : QueueStates)
-			{
-				static_cast<FQueueState&>(*State).bTriggerProfile = true;
-			}
-		}
-
-	} GGPUProfilerSink_ProfileGPU;
-
-	void FGPUProfilerSink_ProfileGPU::FQueueState::ProcessEvents(FIterator& Iterator)
-	{
-	Restart:
-		if (State.Mode == EMode::Idle)
-		{
-			if (bTriggerProfile.exchange(false))
-			{
-				State.Mode = EMode::WaitingFrame;
-			}
-			else
-			{
-				Iterator.PopAll();
-				return;
-			}
-		}
-
-		if (State.Mode == EMode::WaitingFrame)
-		{
-			// Discard all received events until we reach a FrameBoundary event
-			while (FEvent const* Event = Iterator.Pop())
-			{
-				if (Event->GetType() == FEvent::EType::FrameBoundary)
-				{
-					// Start profiling until we receive another FrameBoundary event
-					State.Mode = EMode::Active;
+			FQueue const Queue;
+			TSpscQueue<FEventSink::FIterator> PendingStreams;
 
-					FEvent::FFrameBoundary const& FrameBoundary = Event->Value.Get<FEvent::FFrameBoundary>();
+			// Array of fence signal history. Events are kept until all queues have processed events
+			// later than the CPU timestamps of these signals. The old events are then trimmed.
+			TArray<FResolvedSignal> Signals;
 
-					// Build the node tree 
-					State.PushNode(TEXT("<root>"));
+			// The value of the latest signaled fence on this queue.
+			FResolvedSignal MaxSignal;
 
-				#if WITH_RHI_BREADCRUMBS
-					auto Recurse = [&](auto& Recurse, FRHIBreadcrumbNode* Breadcrumb) -> void
-					{
-						if (!Breadcrumb)
-						{
-							return;
-						}
+			// The GPU timestamp of the last event processed.
+			uint64 LastGPUCycles = 0;
 
-						Recurse(Recurse, Breadcrumb->GetParent());
+			FQueueTimestamps Timestamps;
 
-						FRHIBreadcrumb::FBuffer Buffer;
-						State.PushNode(Breadcrumb->Name.GetTCHAR(Buffer));
-					};
-					Recurse(Recurse, FrameBoundary.Breadcrumb);
-				#endif // WITH_RHI_BREADCRUMBS
+			bool bBusy = false;
 
-					break;
-				}
-			}
-		}
+		#if WITH_RHI_BREADCRUMBS
+			TMap<FRHIBreadcrumbData_Stats, int32> ActiveStats;
+			FRHIBreadcrumbNode* Breadcrumb = nullptr;
+		#endif
 
-		if (State.Mode == EMode::Active)
-		{
-			auto BeginWork = [&](uint64 Timestamp)
+		#if WITH_PROFILEGPU
+			struct
 			{
-				for (FNode* Node = State.Current; Node; Node = Node->Parent)
-				{
-					Node->LastCycles = FMath::Max(Node->LastCycles, Timestamp);
-				}
-			};
+				TArray<TUniquePtr<FNode>> Nodes;
+				FNode* Current = nullptr;
+				FNode* Prev = nullptr;
+				FNode* First = nullptr;
+				bool bProfileFrame = false;
 
-			auto EndWork = [&](uint64 Timestamp)
-			{
-				for (FNode* Node = State.Current; Node; Node = Node->Parent)
+				void PushNode(FString&& Name)
 				{
-					uint64 Delta = Timestamp > Node->LastCycles
-						? Timestamp - Node->LastCycles
-						: 0;
+					FNode* Parent = Current;
+					Current = Nodes.Emplace_GetRef(MakeUnique<FNode>(MoveTemp(Name))).Get();
+					Current->Parent = Parent;
+
+					if (!First)
+					{
+						First = Current;
+					}
 
-					Node->LastCycles = FMath::Max(Node->LastCycles, Timestamp);
-					Node->Inclusive.BusyCycles += Delta;
+					if (Parent)
+					{
+						Parent->Children.Add(Current);
+					}
 
-					if (Node == State.Current)
+					if (Prev)
 					{
-						Node->Exclusive.BusyCycles += Delta;
+						Prev->Next = Current;
 					}
+					Prev = Current;
 				}
-			};
 
-			while (FEvent const* Event = Iterator.Pop())
-			{
-				switch (Event->GetType())
+				void PopNode()
 				{
-				case FEvent::EType::BeginWork:
-					BeginWork(Event->Value.Get<FEvent::FBeginWork>().GPUTimestampTOP);
-					break;
-
-				case FEvent::EType::EndWork:
-					EndWork(Event->Value.Get<FEvent::FEndWork>().GPUTimestampBOP);
-					break;
-
-			#if WITH_RHI_BREADCRUMBS
-				case FEvent::EType::BeginBreadcrumb:
-					{
-						FEvent::FBeginBreadcrumb const& BeginBreadcrumb = Event->Value.Get<FEvent::FBeginBreadcrumb>();
+					check(Current && Current->Parent);
+					Current = Current->Parent;
+				}
 
-						EndWork(BeginBreadcrumb.GPUTimestampTOP);
+				void LogTree(FQueue Queue, uint32 FrameNumber) const
+				{
+					FTable Table;
 
-						// Push a new node
-						FRHIBreadcrumb::FBuffer Buffer;
-						State.PushNode(BeginBreadcrumb.Breadcrumb->Name.GetTCHAR(Buffer));
+					EGPUProfileSortMode SortMode = (EGPUProfileSortMode)FMath::Clamp(GCVarProfileGPU_Sort.GetValueOnAnyThread(), 0, ((int32)EGPUProfileSortMode::Max - 1));
+					FWildcardString RootWildcard(GCVarProfileGPU_Root.GetValueOnAnyThread());
+					const bool bShowEmptyNodes = GCVarProfileGPU_ShowLeafEvents.GetValueOnAnyThread();
+					const double PercentThreshold = FMath::Clamp(GCVarProfileGPU_ThresholdPercent.GetValueOnAnyThread(), 0.0f, 100.0f);
 
-						BeginWork(BeginBreadcrumb.GPUTimestampTOP);
+					if (SortMode != EGPUProfileSortMode::Chronological)
+					{
+						for (FNode* Node = First; Node; Node = Node->Next)
+						{
+							Node->Children.Sort([SortMode](FNode const& A, FNode const& B)
+							{
+								switch (SortMode)
+								{
+								default:
+								case EGPUProfileSortMode::TimeElapsed: return B.Inclusive.BusyCycles    < A.Inclusive.BusyCycles;
+								case EGPUProfileSortMode::NumPrims   : return B.Inclusive.NumPrimitives < A.Inclusive.NumPrimitives;
+								case EGPUProfileSortMode::NumVerts   : return B.Inclusive.NumVertices   < A.Inclusive.NumVertices;
+								}
+							});
+						}
 					}
-					break;
 
-				case FEvent::EType::EndBreadcrumb:
+					auto Recurse = [&](auto& Recurse, FNode* Root, FNode* Current, bool bParentMatchedFilter, int32 Level) -> bool
 					{
-						FEvent::FEndBreadcrumb const& EndBreadcrumb = Event->Value.Get<FEvent::FEndBreadcrumb>();
-						EndWork(EndBreadcrumb.GPUTimestampBOP);
+						// Percent that this node was of the total frame time
+						const double Percent = Root ? (Current->Inclusive.GetBusyMilliseconds() / Root->Inclusive.GetBusyMilliseconds()) * 100.0 : 100.0;
 
-						State.PopNode();
+						// Filter nodes according to cvar settings
+						const bool bAboveThreshold = Percent >= PercentThreshold;
+						const bool bNameMatches = bParentMatchedFilter || RootWildcard.IsMatch(Current->Name);
+						const bool bHasWork = bShowEmptyNodes || Current->Inclusive.HasWork();
 
-						BeginWork(EndBreadcrumb.GPUTimestampBOP);
-					}
-					break;
-			#endif // WITH_RHI_BREADCRUMBS
+						const bool bDisplayEvent = bNameMatches && bHasWork && bAboveThreshold;
+			
+						if (bDisplayEvent)
+						{
+							if (Root == nullptr)
+							{
+								Root = Current;
+							}
 
-				case FEvent::EType::Stats:
-					{
-						FEvent::FStats const& Stats = Event->Value.Get<FEvent::FStats>();
+							Table.AddRow(
+								Root,
+								Current->Inclusive,
+								Current->Exclusive,
+								Current->Name,
+								Level
+							);
+						}
 
-						check(State.Current);
-						State.Current->Exclusive += Stats;
+						FNode::FStats OtherChildrenInclusive;
+						FNode::FStats OtherChildrenExclusive;
+						uint32 NumHiddenChildren = 0;
 
-						for (FNode* Node = State.Current; Node; Node = Node->Parent)
+						for (FNode* Child : Current->Children)
 						{
-							Node->Inclusive += Stats;
+							bool bChildShown = Recurse(Recurse, Root, Child, bDisplayEvent, bDisplayEvent ? Level + 1 : Level);
+							if (!bChildShown)
+							{
+								OtherChildrenInclusive += Child->Inclusive;
+								OtherChildrenExclusive += Child->Exclusive;
+
+								NumHiddenChildren++;
+							}
 						}
-					}
-					break;
 
-				case FEvent::EType::FrameBoundary:
+						if (bDisplayEvent && NumHiddenChildren > 0)
+						{
+							// Don't show the "other children" node if their total inclusive time is still below the percent threshold
+							if ((double(OtherChildrenInclusive.BusyCycles) / Root->Inclusive.BusyCycles) >= PercentThreshold)
+							{
+								Table.AddRow(
+									Root,
+									OtherChildrenInclusive,
+									OtherChildrenExclusive,
+									FString::Printf(TEXT("%d Other %s"), NumHiddenChildren, NumHiddenChildren >= 2 ? TEXT("Children") : TEXT("Child")),
+									Level + 1
+								);
+							}
+						}
+			
+						return bDisplayEvent;
+					};
+
+					// Skip building the table if there was no useful work
+					if (First && First->Inclusive.BusyCycles > 0)
 					{
-						FEvent::FFrameBoundary const& FrameBoundary = Event->Value.Get<FEvent::FFrameBoundary>();
-						LogTree(FrameBoundary.FrameNumber);
+						Recurse(Recurse, nullptr, First, false, 0);
+					}
+
+					FString Final = FString::Printf(
+						TEXT("\n")
+						TEXT("GPU Profile for Frame %d - GPU %d - %s %d\n")
+						TEXT("\n")
+						TEXT("    - %-30s: %.2fms\n")
+						TEXT("    - %-30s: \"%s\"\n")
+						TEXT("    - %-30s: %.2f%%\n")
+						TEXT("    - %-30s: %s\n")
+						TEXT("\n")
+						TEXT("%s")
+						, FrameNumber
+						, Queue.GPU
+						, Queue.GetTypeString()
+						, Queue.Index
+						, TEXT("Frame Time")
+						, First ? First->Inclusive.GetBusyMilliseconds() : 0.0
+						, *IConsoleManager::Get().FindConsoleObjectName(GCVarProfileGPU_Root.AsVariable())
+						, *RootWildcard
+						, *IConsoleManager::Get().FindConsoleObjectName(GCVarProfileGPU_ThresholdPercent.AsVariable())
+						, PercentThreshold
+						, *IConsoleManager::Get().FindConsoleObjectName(GCVarProfileGPU_ShowLeafEvents.AsVariable())
+						, bShowEmptyNodes ? TEXT("true") : TEXT("false")
+						, Table.HasRows() ? *Table.ToString() : TEXT("    No recorded work for this queue.\n")
+					);
+
+					TArray<FString> Lines;
+					Final.ParseIntoArrayLines(Lines, false);
 
-						// Reset tracking
-						State = {};
-						goto Restart;
+					for (FString const& Line : Lines)
+					{
+						UE_LOG(LogRHI, Display, TEXT("%s"), *Line);
 					}
-					break;
 				}
-			}
-		}
-	}
+			} Profile;
+		#endif
 
-	template <uint32 Width>
-	struct TUnicodeHorizontalBar
-	{
-		TCHAR Text[Width + 1];
+			FQueueState(FQueue const& Queue)
+				: Queue(Queue)
+			{}
 
-		// 0 <= Value <= 1
-		TUnicodeHorizontalBar(double Value)
-		{
-			TCHAR* Output = Text;
-			int32 Solid, Partial, Blank;
+			void ResolveSignal(FEvent::FSignalFence const& Event)
 			{
-				double Integer;
-				double Remainder = FMath::Modf(FMath::Clamp(Value, 0.0, 1.0) * Width, &Integer);
+				FResolvedSignal& Result = Signals.Emplace_GetRef();
+		
+				//
+				// Take the max between the previous GPU EndWork event and the CPU timestamp. The signal cannot have happened on the GPU until the CPU has submitted the command to the driver.
+				// 
+				// An example would be a GPU queue that completes work and goes idle at time T. Later, the CPU issues a Signal without other prior work at time T + 100ms.
+				// The fence signal cannot have happened until time T + 100ms because the CPU hadn't instructed the GPU to do so until then.
+				// LastGPUCycles would still be set to time T, since that was the time of the preceeding EndWork event.
+				//
+				Result.GPUTimestampBOP = FMath::Max(LastGPUCycles, Event.CPUTimestamp);
+				Result.Value = Event.Value;
 
-				Solid = (int32)Integer;
-				Partial = (int32)FMath::Floor(Remainder * 8);
-				Blank = (Width - Solid - (Partial > 0 ? 1 : 0));
-			}
+				//
+				// Fences signals *MUST* be sequential, to remove ambiguity caused by trimming the Signals array.
+				// 
+				// To explain why, assume non-sequential signals are allowed, and consider the following example events on an arbitrary queue:
+				// 
+				//		     [Signal 2]
+				//		-- Frame Boundary --
+				//		     [Signal 4]
+				//
+				// Assume, after trimming events earlier than the frame boundary, that only [Signal 4] remains in the Signals array.
+				// Then, some other queue attempts to [Wait 3]. We need to compute when [Wait 3] is resolved with only the information about [Signal 4].
+				// 
+				// Given that fences resolve waits as soon as the signalled value is >= the wait value, we could assume the fence was resolved at [Signal 4].
+				// However, we don't know if the fence was already signalled to value 3 before the frame boundary and the trimming.
+				// 
+				// Without this information, it is ambiguous whether [Wait 3] is already resolved by a [Signal 3] before the frame boundary that is no longer
+				// in the Signals array, or won't be resolved until [Signal 4]. We could have had this sequence of events:
+				// 
+				//		     [Signal 2]
+				//		     [Signal 3]
+				//		-- Frame Boundary --
+				//		     [Signal 4]
+				// 
+				// Requiring that fences are always signalled in sequential order solves this.
+				// If the awaited value is less than the first Signal, the fence has already been signalled before the frame boundary.
+				//
+				check(Result.Value == MaxSignal.Value + 1);
 
-			// Solid characters
-			for (int32 Index = 0; Index < Solid; ++Index)
-			{
-				*Output++ = TEXT('█');
-			}
+				// Signals should always advance in time
+				check(Result.GPUTimestampBOP >= MaxSignal.GPUTimestampBOP);
 
-			// Partially filled character
-			if (Partial > 0)
-			{
-				static constexpr TCHAR const Data[] = TEXT("▏▎▍▌▋▊▉");
-				*Output++ = Data[Partial - 1];
+				MaxSignal = Result;
 			}
 
-			// Blank Characters to pad out the width
-			for (int32 Index = 0; Index < Blank; ++Index)
+			void AccumulateTime(uint64 Busy, uint64 Wait, uint64 Idle)
 			{
-				*Output++ = TEXT(' ');
-			}
-
-			*Output++ = 0;
-			check(uintptr_t(Output) == (uintptr_t(Text) + sizeof(Text)));
-		}
-	};
+				Timestamps.WholeQueueStat.Accumulate(Busy, Wait, Idle);
 
-	void FGPUProfilerSink_ProfileGPU::FQueueState::LogTree(uint32 FrameNumber)
-	{
-		EGPUProfileSortMode SortMode = (EGPUProfileSortMode)FMath::Clamp(GCVarProfileGPU_Sort.GetValueOnAnyThread(), 0, ((int32)EGPUProfileSortMode::Max - 1));
-		FWildcardString RootWildcard(GCVarProfileGPU_Root.GetValueOnAnyThread());
-		const bool bShowEmptyNodes = GCVarProfileGPU_ShowLeafEvents.GetValueOnAnyThread();
-		const double PercentThreshold = FMath::Clamp(GCVarProfileGPU_ThresholdPercent.GetValueOnAnyThread(), 0.0f, 100.0f);
+			#if WITH_RHI_BREADCRUMBS
+				// Apply the timings to all active stats
+				for (auto const& [Stat, RefCount] : ActiveStats)
+				{
+					Timestamps.Stats.FindChecked(Stat).Accumulate(Busy, Wait, Idle);
+				}
+			#endif
 
-		if (SortMode != EGPUProfileSortMode::Chronological)
-		{
-			for (FNode* Node = State.First; Node; Node = Node->Next)
-			{
-				Node->Children.Sort([SortMode](FNode const& A, FNode const& B)
+			#if WITH_PROFILEGPU
+				for (FNode* Node = Profile.Current; Node; Node = Node->Parent)
 				{
-					switch (SortMode)
+					Node->Inclusive.Accumulate(Busy, Wait, Idle);
+
+					if (Node == Profile.Current)
 					{
-					default:
-					case EGPUProfileSortMode::TimeElapsed: return B.Inclusive.BusyCycles    < A.Inclusive.BusyCycles;
-					case EGPUProfileSortMode::NumPrims   : return B.Inclusive.NumPrimitives < A.Inclusive.NumPrimitives;
-					case EGPUProfileSortMode::NumVerts   : return B.Inclusive.NumVertices   < A.Inclusive.NumVertices;
+						Node->Exclusive.Accumulate(Busy, Wait, Idle);
 					}
-				});
+				}
+			#endif
 			}
-		}
-
-		struct FTable
-		{
-			bool const bUnicodeOutput;
-
-			FTable()
-				: bUnicodeOutput(GCVarProfileGPU_UnicodeOutput.GetValueOnAnyThread())
-			{}
 
-			enum class EColumn : uint32
+			void BeginWork(FEvent::FBeginWork const& Event)
 			{
-				Exclusive_NumDraws,
-				Exclusive_NumDispatches,
-				Exclusive_NumPrimitives,
-				Exclusive_NumVertices,
-				Exclusive_Percent,
-				Exclusive_Time,
-
-				Inclusive_NumDraws,
-				Inclusive_NumDispatches,
-				Inclusive_NumPrimitives,
-				Inclusive_NumVertices,
-				Inclusive_Percent,
-				Inclusive_Time,
+				Timestamps.Queue.AddTimestamp(Event.GPUTimestampTOP, true);
 
-				Events,
+				uint64 Idle = Event.CPUTimestamp > LastGPUCycles
+					? Event.CPUTimestamp - LastGPUCycles
+					: 0;
 
-				Num
-			};
+				AccumulateTime(0, 0, Idle);
+				LastGPUCycles = FMath::Max(LastGPUCycles, Event.GPUTimestampTOP);
+			}
 
-			uint32 GetColumnMinimumWidth(EColumn Column) const
+			void EndWork(FEvent::FEndWork const& Event)
 			{
-				switch (Column)
-				{
-				case EColumn::Events:
-					return 6;
-				}
+				Timestamps.Queue.AddTimestamp(Event.GPUTimestampBOP, false);
 
-				return 0;
+				uint64 Busy = Event.GPUTimestampBOP > LastGPUCycles
+					? Event.GPUTimestampBOP - LastGPUCycles
+					: 0;
+
+				AccumulateTime(Busy, 0, 0);
+				LastGPUCycles = FMath::Max(LastGPUCycles, Event.GPUTimestampBOP);
 			}
 
-			TCHAR const* GetColumnHeader(EColumn Column) const
+			void BeginBreadcrumb(FEvent::FBeginBreadcrumb const& Event)
 			{
-				switch (Column)
-				{
-				case EColumn::Exclusive_NumDraws:
-				case EColumn::Inclusive_NumDraws:
-					return TEXT("Draws");
-
-				case EColumn::Exclusive_NumDispatches:
-				case EColumn::Inclusive_NumDispatches:
-					return TEXT("Dsptch");
-				
-				case EColumn::Exclusive_NumPrimitives:
-				case EColumn::Inclusive_NumPrimitives:
-					return TEXT("Prim");
+				uint64 Busy = Event.GPUTimestampTOP > LastGPUCycles
+					? Event.GPUTimestampTOP - LastGPUCycles
+					: 0;
 
-				case EColumn::Exclusive_NumVertices:
-				case EColumn::Inclusive_NumVertices:
-					return TEXT("Vert");
+				AccumulateTime(Busy, 0, 0);
+				LastGPUCycles = FMath::Max(LastGPUCycles, Event.GPUTimestampTOP);
 
-				case EColumn::Exclusive_Percent:
-				case EColumn::Inclusive_Percent:
-					return TEXT("Percent");
-				
-				case EColumn::Exclusive_Time:
-				case EColumn::Inclusive_Time:
-					return TEXT("Time");
+				FRHIBreadcrumbData_Stats const& Stat = Event.Breadcrumb->Name.Data;
+				if (Stat.ShouldComputeStat())
+				{
+					// Disregard the stat if it is nested within itself (i.e. its already in the ActiveStats map with a non-zero ref count).
+					// Only the outermost stat will count the busy time, otherwise we'd be double-counting the nested time.
+					int32 RefCount = ActiveStats.FindOrAdd(Stat)++;
+					if (RefCount == 0)
+					{
+						Timestamps.Stats.FindOrAdd(Stat);
+					}
 				}
 
-				return TEXT("");
-			}
+				Breadcrumb = Event.Breadcrumb;
 
-			uint32 GetColumnGroup(EColumn Column) const
-			{
-			switch (Column)
+			#if WITH_PROFILEGPU
+				if (Profile.bProfileFrame)
 				{
-				case EColumn::Exclusive_NumDraws:
-				case EColumn::Exclusive_NumDispatches:
-				case EColumn::Exclusive_NumPrimitives:
-				case EColumn::Exclusive_NumVertices:
-				case EColumn::Exclusive_Percent:
-				case EColumn::Exclusive_Time:
-					return 0;
-
-				case EColumn::Inclusive_NumDraws:
-				case EColumn::Inclusive_NumDispatches:
-				case EColumn::Inclusive_NumPrimitives:
-				case EColumn::Inclusive_NumVertices:
-				case EColumn::Inclusive_Percent:
-				case EColumn::Inclusive_Time:
-					return 1;
-
-				default:
-				case EColumn::Events:
-					return 2;
+					// Push a new node
+					FRHIBreadcrumb::FBuffer Buffer;
+					Profile.PushNode(Breadcrumb->Name.GetTCHAR(Buffer));
 				}
+			#endif
 			}
 
-			TCHAR const* GetGroupName(uint32 GroupIndex) const
+			void EndBreadcrumb(FEvent::FEndBreadcrumb const& Event)
 			{
-				switch (GroupIndex)
+				uint64 Busy = Event.GPUTimestampBOP > LastGPUCycles
+					? Event.GPUTimestampBOP - LastGPUCycles
+					: 0;
+
+				AccumulateTime(Busy, 0, 0);
+				LastGPUCycles = FMath::Max(LastGPUCycles, Event.GPUTimestampBOP);
+
+				FRHIBreadcrumbData_Stats const& Stat = Event.Breadcrumb->Name.Data;
+				if (Stat.ShouldComputeStat())
 				{
-				case 0: return TEXT("Exclusive");
-				case 1: return TEXT("Inclusive");
-				case 2: return TEXT("Events");
+					// Pop the stat when the refcount hits zero.
+					int32 RefCount = --ActiveStats.FindChecked(Stat);
+					if (RefCount == 0)
+					{
+						ActiveStats.FindAndRemoveChecked(Stat);
+					}
 				}
-				return TEXT("");
-			}
-
-			uint32 NumRows = 0;
-			TStaticArray<TArray<FString>, uint32(EColumn::Num)> Columns { InPlace };
-			TArray<bool> RowBreaks;
 
-			FString& Col(EColumn Column)
-			{
-				return Columns[uint32(Column)].Emplace_GetRef();
-			}
+				Breadcrumb = Event.Breadcrumb->GetParent();
 
-			bool HasRows() const
-			{
-				return NumRows > 0;
+			#if WITH_PROFILEGPU
+				if (Profile.bProfileFrame)
+				{
+					Profile.PopNode();
+				}
+			#endif
 			}
 
-			void AddRow(FNode* Root, FNode::FStats const& Inclusive, FNode::FStats const& Exclusive, FString const& Name, uint32 Level)
+			void Stats(FEvent::FStats const& Event)
 			{
-				double ExclusivePercent = double(Exclusive.BusyCycles) / Root->Inclusive.BusyCycles;
-				double InclusivePercent = double(Inclusive.BusyCycles) / Root->Inclusive.BusyCycles;
-
-				static constexpr uint32 BarWidth = 8;
-				TUnicodeHorizontalBar<BarWidth> ExclusiveBar = ExclusivePercent;
-				TUnicodeHorizontalBar<BarWidth> InclusiveBar = InclusivePercent;
-
-				static constexpr TCHAR const BarSeparator[] = TEXT(" ┊ ");
-
-				Col(EColumn::Exclusive_NumDraws     ) = FString::Printf(TEXT("%d"), Exclusive.NumDraws);
-				Col(EColumn::Exclusive_NumDispatches) = FString::Printf(TEXT("%d"), Exclusive.NumDispatches);
-				Col(EColumn::Exclusive_NumPrimitives) = FString::Printf(TEXT("%d"), Exclusive.NumPrimitives);
-				Col(EColumn::Exclusive_NumVertices  ) = FString::Printf(TEXT("%d"), Exclusive.NumVertices);
-				Col(EColumn::Exclusive_Percent      ) = FString::Printf(TEXT("%.1f%%%s%s"), ExclusivePercent * 100.0, bUnicodeOutput ? BarSeparator : TEXT(""), bUnicodeOutput ? ExclusiveBar.Text : TEXT(""));
-				Col(EColumn::Exclusive_Time         ) = FString::Printf(TEXT("%.3f ms"), FPlatformTime::ToMilliseconds64(Exclusive.BusyCycles));
-
-				Col(EColumn::Inclusive_NumDraws     ) = FString::Printf(TEXT("%d"), Inclusive.NumDraws);
-				Col(EColumn::Inclusive_NumDispatches) = FString::Printf(TEXT("%d"), Inclusive.NumDispatches);
-				Col(EColumn::Inclusive_NumPrimitives) = FString::Printf(TEXT("%d"), Inclusive.NumPrimitives);
-				Col(EColumn::Inclusive_NumVertices  ) = FString::Printf(TEXT("%d"), Inclusive.NumVertices);
-				Col(EColumn::Inclusive_Percent      ) = FString::Printf(TEXT("%.1f%%%s%s"), InclusivePercent * 100.0, bUnicodeOutput ? BarSeparator : TEXT(""), bUnicodeOutput ? InclusiveBar.Text : TEXT(""));
-				Col(EColumn::Inclusive_Time         ) = FString::Printf(TEXT("%.3f ms"), FPlatformTime::ToMilliseconds64(Inclusive.BusyCycles));
-				
-				static constexpr uint32 SpacesPerIndent = 3;
-				Col(EColumn::Events) = FString::Printf(TEXT("%*s"), Name.Len() + (Level * SpacesPerIndent), *Name);
-
-				// Insert a horizontal rule before each root level row.
-				RowBreaks.Add(Level == 0);
+			#if WITH_PROFILEGPU
+				if (Profile.Current)
+				{
+					Profile.Current->Exclusive += Event;
 
-				NumRows++;
+					for (FNode* Node = Profile.Current; Node; Node = Node->Parent)
+					{
+						Node->Inclusive += Event;
+					}
+				}
+			#endif
 			}
 
-			struct FChars
+			void Wait(FResolvedWait const& ResolvedWait)
 			{
-				TCHAR const* Left;
-				TCHAR const* GroupSeparator;
-				TCHAR const* LastGroupSeparator;
-				TCHAR const* Right;
-				TCHAR const* CellSeparator;
-			};
+				// Time the queue was idle between the last EndWork event, and the Wait command being submitted to the GPU driver.
+				uint64 Idle = ResolvedWait.CPUTimestamp > LastGPUCycles
+					? ResolvedWait.CPUTimestamp - LastGPUCycles
+					: 0;
 
-			struct FFormat
-			{
-				TCHAR const* LineMajor;
-				TCHAR const* LineMinor;
+				uint64 WaitStart = FMath::Max(ResolvedWait.CPUTimestamp, LastGPUCycles);
 
-				FChars const TopRow;
-				FChars const GroupNameRow;
-				FChars const GroupBorderRow;
-				FChars const ValueRow;
-				FChars const DividorRow;
-				FChars const BottomRow;
-			};
+				// Time the queue spent waiting for the fence to signal on another queue.
+				uint64 Wait = ResolvedWait.GPUTimestampTOP > WaitStart
+					? ResolvedWait.GPUTimestampTOP - WaitStart
+					: 0;
 
-			FString ToString() const
-			{
-				if (bUnicodeOutput)
-				{
-					static constexpr FFormat Unicode = 
-					{
-						.LineMajor  = TEXT("━"),
-						.LineMinor  = TEXT("─"),
-
-						//                     Left     GrpSep     LastGrp     Right     CellSep
-						.TopRow        { TEXT("    ┏"), TEXT("┳"), TEXT("┳"), TEXT("┓"), TEXT(" ") },
-						.GroupNameRow  { TEXT("    ┃"), TEXT("┃"), TEXT("┃"), TEXT("┃"), TEXT(" ") },
-						.GroupBorderRow{ TEXT("    ┠"), TEXT("╂"), TEXT("┨"), TEXT("┃"), TEXT("┬") },
-						.ValueRow      { TEXT("    ┃"), TEXT("┃"), TEXT("┃"), TEXT("┃"), TEXT("│") },
-						.DividorRow    { TEXT("    ┠"), TEXT("╂"), TEXT("╂"), TEXT("┨"), TEXT("┼") },
-						.BottomRow     { TEXT("    ┗"), TEXT("┻"), TEXT("┻"), TEXT("┛"), TEXT("┷") },
-					};
+				// Bring the last GPU busy end time forwards to where the wait is resolved.
+				LastGPUCycles = ResolvedWait.GPUTimestampTOP;
 
-					return ToStringInner(Unicode);
-				}
-				else
-				{
-					static constexpr FFormat Ascii = 
-					{
-						.LineMajor  = TEXT("-"),
-						.LineMinor  = TEXT("-"),
-
-						//                     Left     GrpSep     LastGrp     Right     CellSep
-						.TopRow        { TEXT("    +"), TEXT("+"), TEXT("+"), TEXT("+"), TEXT(" ") },
-						.GroupNameRow  { TEXT("    |"), TEXT("|"), TEXT("|"), TEXT("|"), TEXT(" ") },
-						.GroupBorderRow{ TEXT("    +"), TEXT("+"), TEXT("+"), TEXT("|"), TEXT("+") },
-						.ValueRow      { TEXT("    |"), TEXT("|"), TEXT("|"), TEXT("|"), TEXT("|") },
-						.DividorRow    { TEXT("    +"), TEXT("+"), TEXT("+"), TEXT("+"), TEXT("+") },
-						.BottomRow     { TEXT("    +"), TEXT("+"), TEXT("+"), TEXT("+"), TEXT("+") },
-					};
+				AccumulateTime(0, Wait, Idle);
+			}
 
-					return ToStringInner(Ascii);
+			void TrimSignals(uint64 CPUTimestamp)
+			{
+				// Remove all signals that occured on the GPU timeline before this frame boundary on the CPU.
+				int32 Index = Algo::LowerBoundBy(Signals, CPUTimestamp, [](FResolvedSignal const& Signal) { return Signal.GPUTimestampBOP; });
+				if (Index >= 0)
+				{
+					Signals.RemoveAt(0, Index, EAllowShrinking::No);
 				}
 			}
 
-			FString ToStringInner(FFormat const& Format) const
+			void FrameBoundary(FEvent::FFrameBoundary const& Event, FFrameState& FrameState, uint32 ProfileFrameNumber)
 			{
-				struct FGroup  { uint32 Index, Width; };
-				struct FColumn { uint32 Index, Width; };
+				check(!bBusy);
+				Timestamps.CPUFrameBoundary = Event.CPUTimestamp;
 
-				static constexpr uint32 NumGroups = 3;
-				static constexpr uint32 CellPadding = 1;
-
-				// Auto-size column widths to their contents
-				TStaticArray<int32, uint32(EColumn::Num)> ColumnWidths{ InPlace, 0 };
-				for (uint32 ColumnIndex = 0; ColumnIndex < uint32(EColumn::Num); ++ColumnIndex)
+			#if WITH_PROFILEGPU
+				if (Profile.bProfileFrame)
 				{
-					if (Columns[ColumnIndex].Num() == 0)
-						continue;
-
-					check(Columns[ColumnIndex].Num() == NumRows);
-
-					int32& Width = ColumnWidths[ColumnIndex];
-
-					// Auto-size column width
-					Width = GetColumnMinimumWidth(EColumn(ColumnIndex));
-					Width = FMath::Max(Width, FCString::Strlen(GetColumnHeader(EColumn(ColumnIndex))));
-
-					for (FString const& Cell : Columns[ColumnIndex])
-					{
-						Width = FMath::Max(Width, Cell.Len());
-					}
+					Profile.LogTree(Queue, Event.FrameNumber);
+					Profile = {};
 				}
+			#endif
 
-				FString Result;
+				FrameState.Emplace(Queue, MoveTemp(Timestamps));
 
-				auto EmitGroupRow = [&](FChars const& Chars, TUniqueFunction<void(FGroup)> GroupCallback)
+			#if WITH_RHI_BREADCRUMBS
+				// Reinsert timestamp streams for the current active stats on 
+				// this queue, since these got moved into the frame state.
+				for (auto& [Stat, RefCount] : ActiveStats)
 				{
-					uint32 const CellSeparatorLength = FCString::Strlen(Chars.CellSeparator);
+					Timestamps.Stats.FindOrAdd(Stat);
+				}
+			#endif
 
-					Result += Chars.Left;
+			#if WITH_PROFILEGPU
+				if (ProfileFrameNumber == Event.FrameNumber + 1)
+				{
+					Profile.bProfileFrame = true;
 
-					uint32 GroupWidth = 0;
-					uint32 GroupIndex = 0;
+					// Build the node tree 
+					Profile.PushNode(TEXT("<root>"));
 
-					for (uint32 ColumnIndex = 0; ColumnIndex < uint32(EColumn::Num); ++ColumnIndex)
+				#if WITH_RHI_BREADCRUMBS
+					auto Recurse = [&](auto& Recurse, FRHIBreadcrumbNode* Breadcrumb) -> void
 					{
-						if (Columns[ColumnIndex].Num() == 0)
-							continue;
+						if (!Breadcrumb)
+						{
+							return;
+						}
 
-						GroupWidth += ColumnWidths[ColumnIndex] + CellPadding * 2;
-						GroupIndex = GetColumnGroup(EColumn(ColumnIndex));
+						Recurse(Recurse, Breadcrumb->GetParent());
 
-						if (GroupIndex != GetColumnGroup(EColumn(ColumnIndex + 1)))
-						{ 
-							// Group Change
-							GroupCallback({ GroupIndex, GroupWidth });
+						FRHIBreadcrumb::FBuffer Buffer;
+						Profile.PushNode(Breadcrumb->Name.GetTCHAR(Buffer));
+					};
+					Recurse(Recurse, Event.Breadcrumb);
+				#endif // WITH_RHI_BREADCRUMBS
+				}
+			#endif
+			}
+		};
 
-							// Add the group separator character
-							Result += GroupIndex < NumGroups - 2
-								? Chars.GroupSeparator
-								: Chars.LastGroupSeparator;
+		std::atomic<bool> bTriggerProfile{ false };
+		uint32 ProfileFrameNumber = 0;
+		uint32 MaxFrameNumber = 0;
 
-							GroupWidth = 0;
-						}
-						else if (ColumnIndex < uint32(EColumn::Num) - 1)
-						{
-							// Same group. Count the (missing) cell division
-							GroupWidth += CellSeparatorLength;
-						}
-					}
+		TMap<uint32, FFrameState> Frames;
+		TMap<FQueue, TUniquePtr<FQueueState>> QueueStates;
 
-					// Emit final group
-					GroupCallback({ GroupIndex, GroupWidth });
+	#if STATS
+		TOptional<int64> LastStatsFrame;
+	#endif
 
-					// Close the row
-					Result += Chars.Right;
-					Result += TEXT("\n");
-				};
+		// Attempts to retrieve the CPU and GPU timestamps of when a fence wait is resolved by a signal on another queue.
+		TOptional<FResolvedWait> ResolveWait(FQueueState& LocalQueue, FEvent::FWaitFence const& WaitFenceEvent)
+		{
+			FQueueState const& RemoteQueue = static_cast<FQueueState const&>(*QueueStates.FindChecked(WaitFenceEvent.Queue));
 
-				auto EmitValueRow = [&](FChars const& Chars, TUniqueFunction<void(FColumn)> CellCallback)
-				{
-					Result += Chars.Left;
+			if (RemoteQueue.MaxSignal.Value < WaitFenceEvent.Value)
+			{
+				// Fence has not yet been signalled on the remote queue
+				return {};
+			}
+			else
+			{
+				// Fence has been signalled, but it may be in the future.
 
-					for (uint32 ColumnIndex = 0; ColumnIndex < uint32(EColumn::Num); ++ColumnIndex)
-					{
-						if (Columns[ColumnIndex].Num() == 0)
-							continue;
+				FResolvedWait Result;
+				Result.CPUTimestamp = WaitFenceEvent.CPUTimestamp;
 
-						CellCallback({ ColumnIndex, ColumnWidths[ColumnIndex] + (CellPadding * 2) });
+				//
+				// The wait cannot be resolved any earlier than:
+				//
+				//		1) The wait command was issued to the driver (WaitFenceEvent.CPUTimestamp)
+				//		2) The GPU completed prior work on this queue (LocalQueue.LastGPUCycles)
+				//
+				Result.GPUTimestampTOP = FMath::Max(WaitFenceEvent.CPUTimestamp, LocalQueue.LastGPUCycles);
+				//
+				//      3) The wait maybe be further delayed by the remote queue the GPU is awaiting.
+				//
+				int32 Index = Algo::LowerBoundBy(RemoteQueue.Signals, WaitFenceEvent.Value, [](FResolvedSignal const& Signal) { return Signal.Value; });
+				if (RemoteQueue.Signals.IsValidIndex(Index))
+				{
+					FResolvedSignal const& Signal = RemoteQueue.Signals[Index];
 
-						if (ColumnIndex < uint32(EColumn::Num) - 1)
-						{
-							uint32 GroupIndex = GetColumnGroup(EColumn(ColumnIndex));
-							if (GroupIndex != GetColumnGroup(EColumn(ColumnIndex + 1)))
-							{
-								// Group change, add the group separator
-								Result += GroupIndex < NumGroups - 2
-									? Chars.GroupSeparator
-									: Chars.LastGroupSeparator;
-							}
-							else
-							{
-								// Same group, add the cell separator
-								Result += Chars.CellSeparator;
-							}
-						}
+					//
+					// Only consider this signal's timestamp if the fence was not already signalled at the previous frame boundary.
+					// See comment in ResolveSignal() for details.
+					//
+					if (!(Index == 0 && WaitFenceEvent.Value < Signal.Value))
+					{
+						Result.GPUTimestampTOP = FMath::Max(Result.GPUTimestampTOP, Signal.GPUTimestampBOP);
 					}
+				}
 
-					// Close the row
-					Result += Chars.Right;
-					Result += TEXT("\n");
-				};
+				return Result;
+			}
+		}
 
-				auto AlignCenter = [&](TCHAR const* Str, uint32 Width)
+		void InitializeQueues(TConstArrayView<FQueue> Queues) override
+		{
+			for (FQueue Queue : Queues)
+			{
+				TUniquePtr<FQueueState>& Ptr = QueueStates.FindOrAdd(Queue);
+				if (!Ptr.IsValid())
 				{
-					int32 PaddingLeft = FMath::Max(0, int32(Width) - FCString::Strlen(Str));
-					int32 PaddingRight = (PaddingLeft / 2) + (PaddingLeft & 1);
-					PaddingLeft /= 2;
-
-					Result += FString::Printf(TEXT("%*s%s%*s"), PaddingLeft, TEXT(""), Str, PaddingRight, TEXT(""));
-				};
+					Ptr = MakeUnique<FQueueState>(Queue);
+				}
+			}
+		}
 
-				// Top Border
-				EmitGroupRow(Format.TopRow, [&](FGroup Group)
+		bool ProcessQueue(FQueueState& QueueState, FIterator& Iterator)
+		{
+			while (FEvent const* Event = Iterator.Peek())
+			{
+				switch (Event->GetType())
 				{
-					while (Group.Width--)
+				case FEvent::EType::BeginWork:
 					{
-						Result += Format.LineMajor;
+						check(!QueueState.bBusy);
+						QueueState.bBusy = true;
+						QueueState.BeginWork(Event->Value.Get<FEvent::FBeginWork>());
 					}
-				});
-
-				// Exclusive / Inclusive Group Row
-				EmitGroupRow(Format.GroupNameRow, [&](FGroup Group)
-				{
-					TCHAR const* Str = Group.Index != GetColumnGroup(EColumn::Events)
-						? GetGroupName(Group.Index)
-						: TEXT("");
+					break;
 
-					AlignCenter(Str, Group.Width);
-				});
+				case FEvent::EType::EndWork:
+					{
+						check(QueueState.bBusy);
+						QueueState.bBusy = false;
+						QueueState.EndWork(Event->Value.Get<FEvent::FEndWork>());
+					}
+					break;
 
-				// Events Group Row
-				EmitValueRow(Format.GroupBorderRow, [&](FColumn Column)
-				{
-					if (Column.Index == uint32(EColumn::Events))
+			#if WITH_RHI_BREADCRUMBS
+				case FEvent::EType::BeginBreadcrumb:
 					{
-						AlignCenter(GetGroupName(GetColumnGroup(EColumn::Events)), Column.Width);
+						check(QueueState.bBusy);
+						QueueState.BeginBreadcrumb(Event->Value.Get<FEvent::FBeginBreadcrumb>());
 					}
-					else
+					break;
+
+				case FEvent::EType::EndBreadcrumb:
 					{
-						while (Column.Width--)
-						{
-							Result += Format.LineMinor;
-						}
+						check(QueueState.bBusy);
+						QueueState.EndBreadcrumb(Event->Value.Get<FEvent::FEndBreadcrumb>());
 					}
-				});
+					break;
+			#endif // WITH_RHI_BREADCRUMBS
 
-				// Header Row
-				EmitValueRow(Format.ValueRow, [&](FColumn Column)
-				{
-					AlignCenter(GetColumnHeader(EColumn(Column.Index)), Column.Width);
-				});
+			#if WITH_PROFILEGPU
+				case FEvent::EType::Stats:
+					{
+						check(QueueState.bBusy);
+						QueueState.Stats(Event->Value.Get<FEvent::FStats>());
+					}
+					break;
+			#endif // WITH_PROFILEGPU
 
-				// Header Border Row
-				EmitValueRow(Format.DividorRow, [&](FColumn Column)
-				{
-					while (Column.Width--)
+				case FEvent::EType::SignalFence:
 					{
-						Result += Format.LineMinor;
+						check(!QueueState.bBusy);
+						QueueState.ResolveSignal(Event->Value.Get<FEvent::FSignalFence>());
 					}
-				});
+					break;
 
-				// Value rows
-				for (uint32 RowIndex = 0; RowIndex < NumRows; ++RowIndex)
-				{
-					if (RowIndex > 0 && RowBreaks[RowIndex])
+				case FEvent::EType::WaitFence:
 					{
-						// Add a horizontal rule
-						EmitValueRow(Format.DividorRow, [&](FColumn Column)
+						check(!QueueState.bBusy);
+						TOptional<FResolvedWait> ResolvedWait = ResolveWait(QueueState, Event->Value.Get<FEvent::FWaitFence>());
+
+						if (!ResolvedWait.IsSet())
 						{
-							while (Column.Width--)
-							{
-								Result += Format.LineMinor;
-							}
-						});
+							// Unresolved fence, pause processing
+							return false;
+						}
+
+						QueueState.Wait(*ResolvedWait);
 					}
+					break;
 
-					EmitValueRow(Format.ValueRow, [&](FColumn Column)
+				case FEvent::EType::FrameBoundary:
 					{
-						int32 Width = Column.Width - (CellPadding * 2);
-						if (EColumn(Column.Index) == EColumn::Events)
+						FEvent::FFrameBoundary const& FrameBoundary = Event->Value.Get<FEvent::FFrameBoundary>();
+						FFrameState& FrameState = Frames.FindOrAdd(FrameBoundary.FrameNumber);
+
+					#if STATS
+						check(!FrameState.StatsFrame.IsSet() || FrameState.StatsFrame.GetValue() == FrameBoundary.StatsFrame);
+						FrameState.StatsFrame = FrameBoundary.StatsFrame;
+					#endif
+
+					#if WITH_PROFILEGPU
+						// Latch the index of the next frame to profile
+						MaxFrameNumber = FMath::Max(FrameBoundary.FrameNumber, MaxFrameNumber);
+						if (bTriggerProfile.exchange(false))
 						{
-							Width = -Width; // Align left
+							ProfileFrameNumber = MaxFrameNumber + 1;
 						}
+					#endif // WITH_PROFILEGPU
 
-						FString const& Cell = Columns[Column.Index][RowIndex];
-						Result += FString::Printf(TEXT("%*s%*s%*s")
-							, CellPadding, TEXT("")
-							, Width, *Cell
-							, CellPadding, TEXT(""));
-					});
-				}
+						QueueState.FrameBoundary(FrameBoundary, FrameState, ProfileFrameNumber);
 
-				// Bottom Border
-				EmitValueRow(Format.BottomRow, [&](FColumn Column)
-				{
-					while (Column.Width--)
-					{
-						Result += Format.LineMajor;
+						if (FrameState.Num() == QueueStates.Num())
+						{
+							// Trim the Signals array in each queue, up to the lowest frame boundary CPU timestamp.
+							{
+								uint64 MinFrameBoundary = TNumericLimits<uint64>::Max();
+								for (auto& [Queue, QueueTimestamps] : FrameState)
+								{
+									MinFrameBoundary = FMath::Min(MinFrameBoundary, QueueTimestamps.CPUFrameBoundary);
+								}
+
+								for (auto& [Queue, LocalQueueState] : QueueStates)
+								{
+									LocalQueueState.Get()->TrimSignals(MinFrameBoundary);
+								}
+							}
+
+							// All registered queues have reported their frame boundary event.
+							// We have a full set of data to compute the total frame GPU stats.
+							ProcessFrame(FrameState);
+
+							Frames.Remove(FrameBoundary.FrameNumber);
+						}
 					}
-				});
+					break;
+				}
 
-				return Result;
+				Iterator.Pop();
 			}
-		} Table;
 
-		auto Recurse = [&](auto& Recurse, FNode* Root, FNode* Current, bool bParentMatchedFilter, int32 Level) -> bool
+			return true;
+		}
+
+		void ProcessFrame(FFrameState& FrameState)
 		{
-			// Percent that this node was of the total frame time
-			const double Percent = Root ? (Current->Inclusive.GetBusyMilliseconds() / Root->Inclusive.GetBusyMilliseconds()) * 100.0 : 100.0;
+		#if STATS
+			check(FrameState.StatsFrame.IsSet());
+			FEndOfPipeStats* Stats = FEndOfPipeStats::Get();
 
-			// Filter nodes according to cvar settings
-			const bool bAboveThreshold = Percent >= PercentThreshold;
-			const bool bNameMatches = bParentMatchedFilter || RootWildcard.IsMatch(Current->Name);
-			const bool bHasWork = bShowEmptyNodes || Current->Inclusive.HasWork();
+			if (!LastStatsFrame.IsSet() || *LastStatsFrame < *FrameState.StatsFrame)
+			{
+				Stats->AddMessage(FStatConstants::AdvanceFrame.GetEncodedName(), EStatOperation::AdvanceFrameEventEndOfPipe, *FrameState.StatsFrame);
+				LastStatsFrame = *FrameState.StatsFrame;
+			}
+		#endif
 
-			const bool bDisplayEvent = bNameMatches && bHasWork && bAboveThreshold;
-			
-			if (bDisplayEvent)
+		#if CSV_PROFILER_STATS
+			const bool bCsvStatsEnabled = !!CVarGPUCsvStatsEnabled.GetValueOnAnyThread();
+			FCsvProfiler* CsvProfiler = FCsvProfiler::Get();
+			CsvProfiler->BeginFrameEOP();
+		#else
+			const bool bCsvStatsEnabled = false;
+		#endif
+
+			for (auto const& [Queue, QueueTimestamps] : FrameState)
 			{
-				if (Root == nullptr)
+			#if WITH_RHI_BREADCRUMBS
+				// Compute the individual GPU stats
+				for (auto const& [Stat, StatState] : QueueTimestamps.Stats)
 				{
-					Root = Current;
+					StatState.EmitResults(Queue
+					#if HAS_GPU_STATS
+						, Stats
+						, *Stat.GPUStat
+					#endif
+					#if CSV_PROFILER_STATS
+						, bCsvStatsEnabled ? CsvProfiler : nullptr
+					#endif
+					);
 				}
+			#endif // WITH_RHI_BREADCRUMBS
 
-				Table.AddRow(
-					Root,
-					Current->Inclusive,
-					Current->Exclusive,
-					Current->Name,
-					Level
+				// Set the whole-frame per queue stat
+				QueueTimestamps.WholeQueueStat.EmitResults(Queue
+				#if HAS_GPU_STATS
+					, Stats
+					, GPUStat_Total
+				#endif
+				#if CSV_PROFILER_STATS
+					, bCsvStatsEnabled ? CsvProfiler : nullptr
+				#endif
 				);
 			}
 
-			FNode::FStats OtherChildrenInclusive;
-			FNode::FStats OtherChildrenExclusive;
-			uint32 NumHiddenChildren = 0;
-
-			for (FNode* Child : Current->Children)
+			// Compute the whole-frame total GPU time.
+			TArray<FTimestampStream::FState, TInlineAllocator<GetRHIPipelineCount() * MAX_NUM_GPUS>> StreamPointers;
+			for (auto const& [Queue, State] : FrameState)
 			{
-				bool bChildShown = Recurse(Recurse, Root, Child, bDisplayEvent, bDisplayEvent ? Level + 1 : Level);
-				if (!bChildShown)
-				{
-					OtherChildrenInclusive += Child->Inclusive;
-					OtherChildrenExclusive += Child->Exclusive;
-
-					NumHiddenChildren++;
-				}
+				StreamPointers.Emplace(State.Queue);
 			}
+			uint64 WholeFrameUnion = FTimestampStream::ComputeUnion(StreamPointers);
 
-			if (bDisplayEvent && NumHiddenChildren > 0)
+			// Update the global GPU frame time stats - need to convert to Cycles32 rather than Cycles64.
+			GGPUFrameTime = FPlatformMath::TruncToInt(FPlatformTime::ToSeconds64(WholeFrameUnion) / FPlatformTime::GetSecondsPerCycle());
+
+			// @todo set global csv GPU time
+			//RHISetGPUStatTotals(bCsvStatsEnabled, FPlatformTime::ToMilliseconds64(WholeFrameUnion));
+
+		#if STATS
+			Stats->Flush();
+		#endif
+		}
+
+		void ProcessAllQueues()
+		{
+			// Process the queue as far as possible
+			bool bProgress;
+			do
 			{
-				// Don't show the "other children" node if their total inclusive time is still below the percent threshold
-				if ((double(OtherChildrenInclusive.BusyCycles) / Root->Inclusive.BusyCycles) >= PercentThreshold)
+				bProgress = false;
+				for (auto& [Queue, QueueState] : QueueStates)
 				{
-					Table.AddRow(
-						Root,
-						OtherChildrenInclusive,
-						OtherChildrenExclusive,
-						FString::Printf(TEXT("%d Other %s"), NumHiddenChildren, NumHiddenChildren >= 2 ? TEXT("Children") : TEXT("Child")),
-						Level + 1
-					);
-				}
-			}
-			
-			return bDisplayEvent;
-		};
+					while (FIterator* Iterator = QueueState->PendingStreams.Peek())
+					{
+						FEvent const* Start = Iterator->Peek();
 
-		// Skip building the table if there was no useful work
-		if (State.First && State.First->Inclusive.BusyCycles > 0)
-		{
-			Recurse(Recurse, nullptr, State.First, false, 0);
+						bool bPaused = !ProcessQueue(*QueueState.Get(), *Iterator);
+
+						FEvent const* End = Iterator->Peek();
+						bProgress |= End != Start;
+
+						if (bPaused)
+						{
+							// The queue was paused by a Wait event
+							check(End);
+							break;
+						}
+
+						if (!End)
+						{
+							// This stream has been fully processed.
+							QueueState->PendingStreams.Dequeue();
+						}
+					}
+				}
+			} while (bProgress);
 		}
 
-		FString Final = FString::Printf(
-			TEXT("\n")
-			TEXT("GPU Profile for Frame %d - GPU %d - %s %d\n")
-			TEXT("\n")
-			TEXT("    - %-30s: %.2fms\n")
-			TEXT("    - %-30s: \"%s\"\n")
-			TEXT("    - %-30s: %.2f%%\n")
-			TEXT("    - %-30s: %s\n")
-			TEXT("\n")
-			TEXT("%s")
-			, FrameNumber
-			, Queue.GPU
-			, Queue.GetTypeString()
-			, Queue.Index
-			, TEXT("Frame Time")
-			, State.First->Inclusive.GetBusyMilliseconds()
-			, *IConsoleManager::Get().FindConsoleObjectName(GCVarProfileGPU_Root.AsVariable())
-			, *RootWildcard
-			, *IConsoleManager::Get().FindConsoleObjectName(GCVarProfileGPU_ThresholdPercent.AsVariable())
-			, PercentThreshold
-			, *IConsoleManager::Get().FindConsoleObjectName(GCVarProfileGPU_ShowLeafEvents.AsVariable())
-			, bShowEmptyNodes ? TEXT("true") : TEXT("false")
-			, Table.HasRows() ? *Table.ToString() : TEXT("    No recorded work for this queue.\n")
-		);
-
-		TArray<FString> Lines;
-		Final.ParseIntoArrayLines(Lines, false);
-
-		for (FString const& Line : Lines)
+		void ProcessStreams(TConstArrayView<TSharedRef<FEventStream>> EventStreams) override
 		{
-			UE_LOG(LogRHI, Display, TEXT("%s"), *Line);
+			for (TSharedRef<FEventStream> const& Stream : EventStreams)
+			{
+				FQueueState& QueueState = *QueueStates.FindChecked(Stream->Queue);
+				QueueState.PendingStreams.Enqueue(FIterator(Stream));
+			}
+
+			ProcessAllQueues();
 		}
-	}
+	} GGPUProfilerSink_StatSystem;
 
+#if WITH_PROFILEGPU
 	static FAutoConsoleCommand GCommand_ProfileGPU(
 		TEXT("ProfileGPU"),
 		TEXT("Captures statistics about a frame of GPU work and prints the results to the log."),
 		FConsoleCommandWithArgsDelegate::CreateLambda([](const TArray<FString>& Args)
 		{
-			GGPUProfilerSink_ProfileGPU.ProfileNextFrame();
+			GGPUProfilerSink_StatSystem.bTriggerProfile = true;
 		}));
 #endif // WITH_PROFILEGPU
 }
diff --git a/Engine/Source/Runtime/RHI/Private/GpuProfilerTrace.cpp b/Engine/Source/Runtime/RHI/Private/GpuProfilerTrace.cpp
--- a/Engine/Source/Runtime/RHI/Private/GpuProfilerTrace.cpp
+++ b/Engine/Source/Runtime/RHI/Private/GpuProfilerTrace.cpp
@@ -353,12 +353,21 @@ struct FGPUProfilerSink_Trace final : public FEventSink
 		}
 	}
 
-	bool ProcessStream(FQueueState& QueueState, FIterator& Iterator) override
+	void InitializeQueues(TConstArrayView<FQueue> Queues) override
 	{
+		for (FQueue Queue : Queues)
+		{
+			CachedQueues.Add(Queue);
+		}
+	}
+
+	void ProcessStream(TSharedRef<FEventStream> const& EventStream)
+	{
+		FQueue Queue = EventStream->Queue;
 		if (!UE_TRACE_CHANNELEXPR_IS_ENABLED(GpuChannel))
 		{
-			CachedQueues.Add(QueueState.Queue);
-			return true;
+			CachedQueues.Add(Queue);
+			return;
 		}
 
 		if (!bInitialized)
@@ -372,8 +381,7 @@ struct FGPUProfilerSink_Trace final : public FEventSink
 			CachedQueues.Empty();
 		}
 
-		FQueue Queue = QueueState.Queue;
-
+		FEventSink::FIterator Iterator = EventStream;
 		while (FEvent const* Event = Iterator.Pop())
 		{
 			switch (Event->GetType())
@@ -462,8 +470,14 @@ struct FGPUProfilerSink_Trace final : public FEventSink
 				}
 			}
 		}
+	}
 
-		return true;
+	void ProcessStreams(TConstArrayView<TSharedRef<FEventStream>> EventStreams) override
+	{
+		for (TSharedRef<FEventStream> const& Stream : EventStreams)
+		{
+			ProcessStream(Stream);
+		}
 	}
 
 } GGPUProfilerSink_Trace;
diff --git a/Engine/Source/Runtime/RHI/Private/RHICommandList.cpp b/Engine/Source/Runtime/RHI/Private/RHICommandList.cpp
--- a/Engine/Source/Runtime/RHI/Private/RHICommandList.cpp
+++ b/Engine/Source/Runtime/RHI/Private/RHICommandList.cpp
@@ -1078,9 +1078,12 @@ void FRHICommandListExecutor::FSubmitState::Submit(const FSubmitArgs& Args)
 		FDynamicRHI::FRHIEndFrameArgs EndFrameArgs
 		{
 			.FrameNumber = GRHICommandList.FrameNumber++,
-#if WITH_RHI_BREADCRUMBS
-			.GPUBreadcrumbs = Args.GPUBreadcrumbs
-#endif
+		#if WITH_RHI_BREADCRUMBS
+			.GPUBreadcrumbs = Args.GPUBreadcrumbs,
+		#endif
+		#if STATS
+			.StatsFrame = Args.StatsFrame,
+		#endif
 		};
 
 		GDynamicRHI->RHIEndFrame(EndFrameArgs);
@@ -1226,7 +1229,13 @@ RHI_API void FRHICommandListExecutor::Submit(TConstArrayView<FRHICommandListBase
 	if (EnumHasAnyFlags(SubmitFlags, ERHISubmitFlags::SubmitToGPU))
 	{
 		EnqueueDispatchTask({},
-			[this, State = SubmitState]() mutable
+			[
+				  this
+				, State = SubmitState
+			#if STATS
+				, LocalStatsFrame = StatsFrame
+			#endif
+			]() mutable
 			{
 				SCOPED_NAMED_EVENT(RHI_FinalizeAndSubmit, FColor::White);
 
@@ -1234,12 +1243,15 @@ RHI_API void FRHICommandListExecutor::Submit(TConstArrayView<FRHICommandListBase
 				State->FinalizeCurrent();
 
 				FSubmitState::FSubmitArgs Args;
-#if WITH_RHI_BREADCRUMBS
+			#if WITH_RHI_BREADCRUMBS
 				for (ERHIPipeline Pipeline : MakeFlagsRange(ERHIPipeline::All))
 				{
 					Args.GPUBreadcrumbs[Pipeline] = GRHICommandList.Breadcrumbs.GPU[Pipeline].Current;
 				}
-#endif
+			#endif
+			#if STATS
+				Args.StatsFrame = LocalStatsFrame;
+			#endif
 
 				// Submission thread
 				EnqueueSubmitTask(MoveTemp(State->TranslateEvents),
diff --git a/Engine/Source/Runtime/RHI/Public/DynamicRHI.h b/Engine/Source/Runtime/RHI/Public/DynamicRHI.h
--- a/Engine/Source/Runtime/RHI/Public/DynamicRHI.h
+++ b/Engine/Source/Runtime/RHI/Public/DynamicRHI.h
@@ -235,9 +235,12 @@ public:
 		// Increments once per call to RHIEndFrame
 		uint32 FrameNumber;
 
-#if WITH_RHI_BREADCRUMBS
+	#if WITH_RHI_BREADCRUMBS
 		const TRHIPipelineArray<FRHIBreadcrumbNode*>& GPUBreadcrumbs;
-#endif
+	#endif
+	#if STATS
+		int64 StatsFrame;
+	#endif
 	};
 	virtual void RHIEndFrame(const FRHIEndFrameArgs& Args) = 0;
 
diff --git a/Engine/Source/Runtime/RHI/Public/GPUProfiler.h b/Engine/Source/Runtime/RHI/Public/GPUProfiler.h
--- a/Engine/Source/Runtime/RHI/Public/GPUProfiler.h
+++ b/Engine/Source/Runtime/RHI/Public/GPUProfiler.h
@@ -102,9 +102,18 @@ namespace UE::RHI::GPUProfiler
 			FRHIBreadcrumbNode* Breadcrumb;
 		#endif
 
-			FFrameBoundary(uint64 CPUTimestamp, uint32 FrameNumber
+		#if STATS
+			int64 StatsFrame;
+		#endif
+
+			FFrameBoundary(
+				  uint64 CPUTimestamp
+				, uint32 FrameNumber
 			#if WITH_RHI_BREADCRUMBS
 				, FRHIBreadcrumbNode* Breadcrumb
+			#endif
+			#if STATS
+				, int64 StatsFrame
 			#endif
 				)
 				: CPUTimestamp(CPUTimestamp)
@@ -112,6 +121,9 @@ namespace UE::RHI::GPUProfiler
 			#if WITH_RHI_BREADCRUMBS
 				, Breadcrumb(Breadcrumb)
 			#endif
+			#if STATS
+				, StatsFrame(StatsFrame)
+			#endif
 			{}
 		};
 
@@ -330,12 +342,18 @@ namespace UE::RHI::GPUProfiler
 		FChunk* Current = nullptr;
 
 	public:
-		FEventStream() = default;
+		FQueue const Queue;
+
+		FEventStream(FQueue const Queue)
+			: Queue(Queue)
+		{}
+
 		FEventStream(FEventStream const&) = delete;
 
 		FEventStream(FEventStream&& Other)
-			: First(Other.First)
+			: First  (Other.First)
 			, Current(Other.Current)
+			, Queue  (Other.Queue)
 		{
 			Other.First = nullptr;
 			Other.Current = nullptr;
@@ -399,6 +417,8 @@ namespace UE::RHI::GPUProfiler
 
 		void Append(FEventStream&& Other)
 		{
+			check(Queue == Other.Queue);
+
 			if (IsEmpty())
 			{
 				Current = Other.Current;
@@ -417,9 +437,7 @@ namespace UE::RHI::GPUProfiler
 
 	struct FEventSink
 	{
-		friend RHI_API void ProcessEvents(FQueue Queue, FEventStream EventStream);
-		friend RHI_API void InitializeQueues(TConstArrayView<FQueue> Queues);
-
+	protected:
 		struct FIterator
 		{
 			friend FEventSink;
@@ -430,12 +448,12 @@ namespace UE::RHI::GPUProfiler
 			FEventStream::FChunk* Current;
 			uint32 Index = 0;
 
+		public:
 			FIterator(TSharedRef<FEventStream> const& Stream)
 				: Stream(Stream)
 				, Current(Stream->First)
 			{}
 
-		public:
 			FEvent const* Peek() const
 			{
 				return Current ? Current->GetElement(Index) : nullptr;
@@ -457,72 +475,20 @@ namespace UE::RHI::GPUProfiler
 
 				return Result;
 			}
-
-			void PopAll()
-			{
-				Current = nullptr;
-				Index = 0;
-			}
-		};
-
-	protected:
-		struct FResolvedWait
-		{
-			uint64 GPUTimestampTOP = 0;
-			uint64 CPUTimestamp    = 0;
-		};
-
-		struct FResolvedSignal
-		{
-			uint64 GPUTimestampBOP = 0;
-			uint64 Value           = 0;
 		};
 
-		struct FQueueState
-		{
-			FQueue const Queue;
-
-			// Array of fence signal history. Events are kept until all queues have processed events
-			// later than the CPU timestamps of these signals. The old events are then trimmed.
-			TArray<FResolvedSignal> Signals;
-
-			// The value of the latest signaled fence on this queue.
-			FResolvedSignal MaxSignal;
-
-			// The GPU timestamp of the last event processed.
-			uint64 LastGPUCycles = 0;
-			uint64 LastCPUCycles = 0;
-
-			TSpscQueue<FIterator> PendingStreams;
-
-			FQueueState(FQueue const Queue)
-				: Queue(Queue)
-			{}
-
-			void ResolveSignal(FEvent::FSignalFence const& SignalFenceEvent);
-		};
-
-		TMap<FQueue, TUniquePtr<FQueueState>> QueueStates;
-
-		void AppendStream(FQueue Queue, TSharedRef<FEventStream> const& NewEventStream);
-		void ProcessStreams();
-
 		RHI_API FEventSink();
 		RHI_API ~FEventSink();
 
 		FEventSink(FEventSink const&) = delete;
 		FEventSink(FEventSink&&) = delete;
 
-		// Attempts to retrieve the CPU and GPU timestamps of when a fence wait is resolved by a signal on another queue.
-		TOptional<FResolvedWait> ResolveWait(FQueueState& QueueState, FEvent::FWaitFence const& WaitFenceEvent);
-
-		virtual TUniquePtr<FQueueState> CreateQueueState(FQueue Queue) const;
-
 	public:
-		virtual bool ProcessStream(FQueueState& QueueState, FIterator& Iterator) = 0;
+		virtual void ProcessStreams(TConstArrayView<TSharedRef<FEventStream>> EventStreams) = 0;
+		virtual void InitializeQueues(TConstArrayView<FQueue> Queues) = 0;
 	};
 
-	RHI_API void ProcessEvents(FQueue Queue, FEventStream EventStream);
+	RHI_API void ProcessEvents(TArrayView<FEventStream> EventStreams);
 	RHI_API void InitializeQueues(TConstArrayView<FQueue> Queues);
 
 	struct FGPUStat
diff --git a/Engine/Source/Runtime/RHI/Public/RHICommandList.h b/Engine/Source/Runtime/RHI/Public/RHICommandList.h
--- a/Engine/Source/Runtime/RHI/Public/RHICommandList.h
+++ b/Engine/Source/Runtime/RHI/Public/RHICommandList.h
@@ -5746,9 +5746,12 @@ private:
 
 		struct FSubmitArgs
 		{
-#if WITH_RHI_BREADCRUMBS
+		#if WITH_RHI_BREADCRUMBS
 			TRHIPipelineArray<FRHIBreadcrumbNode*> GPUBreadcrumbs;
-#endif
+		#endif
+		#if STATS
+			int64 StatsFrame;
+		#endif
 		};
 		void Submit(const FSubmitArgs& Args);
 		void FinalizeCurrent();
@@ -5786,6 +5789,19 @@ private:
 	// Counts the number of calls to RHIEndFrame, and is used in GPU profiler frame boundary events.
 	uint32 FrameNumber = 0;
 
+#if STATS
+	// The engine stats system maintains its own frame counter. We need to propagate this to
+	// the end-of-pipe thread so that stat data it submits is attributed to the correct frame.
+	// This is done inside AdvanceRenderingThreadStats.
+	int64 StatsFrame = 0;
+public:
+	void SetStatsFrame(int64 InStatsFrame)
+	{
+		StatsFrame = InStatsFrame;
+	}
+private:
+#endif
+
 	bool AllowParallel() const;
 };
 
diff --git a/Engine/Source/Runtime/RenderCore/Private/RenderingThread.cpp b/Engine/Source/Runtime/RenderCore/Private/RenderingThread.cpp
--- a/Engine/Source/Runtime/RenderCore/Private/RenderingThread.cpp
+++ b/Engine/Source/Runtime/RenderCore/Private/RenderingThread.cpp
@@ -345,8 +345,12 @@ static void AdvanceRenderingThreadStats(int64 StatsFrame, int32 DisableChangeTag
 	}
 	FThreadStats::AddMessage(FStatConstants::AdvanceFrame.GetEncodedName(), EStatOperation::AdvanceFrameEventRenderThread, Frame);
 
-	// @todo remove this when the new GPU profiler is using the end-of-pipe stats timeline
+#if RHI_NEW_GPU_PROFILER
+	// Propagate the stats frame value down to the end-of-pipe thread.
+	GRHICommandList.SetStatsFrame(Frame);
+#else
 	FThreadStats::AddMessage(FStatConstants::AdvanceFrame.GetEncodedName(), EStatOperation::AdvanceFrameEventEndOfPipe, Frame);
+#endif
 
 	if( IsInActualRenderingThread() )
 	{
diff --git a/Engine/Source/Runtime/Windows/D3D11RHI/Private/D3D11RHI.cpp b/Engine/Source/Runtime/Windows/D3D11RHI/Private/D3D11RHI.cpp
--- a/Engine/Source/Runtime/Windows/D3D11RHI/Private/D3D11RHI.cpp
+++ b/Engine/Source/Runtime/Windows/D3D11RHI/Private/D3D11RHI.cpp
@@ -44,6 +44,9 @@ void FD3D11DynamicRHI::RHIEndFrame(const FRHIEndFrameArgs& Args)
 		EmplaceProfilerEvent<UE::RHI::GPUProfiler::FEvent::FFrameBoundary>(Timestamp, Args.FrameNumber
 	#if WITH_RHI_BREADCRUMBS
 			, Args.GPUBreadcrumbs[ERHIPipeline::Graphics]
+	#endif
+	#if STATS
+			, Args.StatsFrame
 	#endif
 		);
 
@@ -81,7 +84,7 @@ void FD3D11DynamicRHI::RHIEndFrame(const FRHIEndFrameArgs& Args)
 			PollQueryResults();
 
 			// Previous frame has completed and the data is available. Publish the profiler events.
-			UE::RHI::GPUProfiler::ProcessEvents(UE::RHI::GPUProfiler::FQueue(UE::RHI::GPUProfiler::FQueue::EType::Graphics, 0, 0), MoveTemp(PreviousFrame->EventStream));
+			UE::RHI::GPUProfiler::ProcessEvents(MakeArrayView(&PreviousFrame->EventStream, 1));
 
 			Profiler.EventPool.Push(MoveTemp(PreviousFrame->CompletionQuery));
 			Profiler.Pending.Pop();
diff --git a/Engine/Source/Runtime/Windows/D3D11RHI/Private/D3D11RHIPrivate.h b/Engine/Source/Runtime/Windows/D3D11RHI/Private/D3D11RHIPrivate.h
--- a/Engine/Source/Runtime/Windows/D3D11RHI/Private/D3D11RHIPrivate.h
+++ b/Engine/Source/Runtime/Windows/D3D11RHI/Private/D3D11RHIPrivate.h
@@ -792,6 +792,12 @@ protected:
 		#if WITH_RHI_BREADCRUMBS
 			FRHIBreadcrumbAllocatorArray BreadcrumbAllocators{};
 		#endif
+
+			FFrame()
+		#if RHI_NEW_GPU_PROFILER
+				: EventStream(UE::RHI::GPUProfiler::FQueue(UE::RHI::GPUProfiler::FQueue::EType::Graphics, 0, 0))
+		#endif
+			{}
 		};
 
 		FFrame Current;
