VulkanRHI: Further limit the textures where NeedsAllPlanes is applied by only targetting formats with stencil.  Move the condition into a function.

#rnx

diff --git a/Engine/Source/Runtime/RHI/Private/RHICommandList.cpp b/Engine/Source/Runtime/RHI/Private/RHICommandList.cpp
--- a/Engine/Source/Runtime/RHI/Private/RHICommandList.cpp
+++ b/Engine/Source/Runtime/RHI/Private/RHICommandList.cpp
@@ -2240,6 +2240,25 @@ void FRHIComputeCommandList::RayTraceDispatchIndirect(
 	GlobalResourceBindings.Reset();
 }
 
+// NeedsAllPlanes currently only applies to texture with depth+stencil
+bool FRHICommandListBase::TrackTransitioningResource(const FRHITransitionCreateInfoData& TransitionCreateInfoData, const FRHITransitionInfo& TransitionInfo)
+{
+	if (GRHIGlobals.NeedsAllPlanes && GRHIGlobals.EnableHighLevelTransitionPatching)
+	{
+		if (!EnumHasAnyFlags(TransitionCreateInfoData.m_Flags, ERHITransitionCreateFlags::NoSplit) &&
+			EnumHasAnyFlags(TransitionCreateInfoData.SrcPipelines, ERHIPipeline::Graphics) &&
+			EnumHasAnyFlags(TransitionCreateInfoData.DstPipelines, ERHIPipeline::Graphics) &&
+			(TransitionInfo.Type == FRHITransitionInfo::EType::Texture))
+		{
+			FRHITexture* Texture = (FRHITexture*)GetViewableResource(TransitionInfo);
+			if (Texture && IsStencilFormat(Texture->GetDesc().Format))
+			{
+				return true;
+			}
+		}
+	}
+	return false;
+}
 
 // track all transitions:
 // - if a BeforeAccess is Unknown add the transition to a list for patching before dispatch
@@ -2316,11 +2335,7 @@ void FRHICommandListBase::TrackResourceStateBegin(TArrayView<const FRHITransitio
 			const bool bShouldSkipAfterState = ((EnumHasAnyFlags(TransitionInfo->Flags, EResourceTransitionFlags::IgnoreAfterState)) || (TransitionInfo->AccessAfter == ERHIAccess::Unknown));
 			const bool bNeedUpdateAfterState = (bShouldSkipAfterState == false) && (bDiscardOnAsync == false);
 
-			const bool bTrackOnEnd = NeedsAllPlanes() && 
-				(TransitionInfo->Type == FRHITransitionInfo::EType::Texture) &&
-				EnumHasAnyFlags(TransitionCreateInfoData->SrcPipelines, ERHIPipeline::Graphics) && 
-				EnumHasAnyFlags(TransitionCreateInfoData->DstPipelines, ERHIPipeline::Graphics);
-			if (bNeedUpdateAfterState && !bTrackOnEnd)
+			if (bNeedUpdateAfterState && !TrackTransitioningResource(*TransitionCreateInfoData, *TransitionInfo))
 			{
 				// Update the TrackedResource State
 				if (TrackedResourceState)
@@ -2377,24 +2392,18 @@ void FRHICommandListBase::TrackResourceStateEnd(TArrayView<const FRHITransition*
 				{
 					const FRHITransitionInfo* TransitionInfo = &TransitionCreateInfoData->TransitionInfos[Index];
 
-					if (TransitionInfo->Type != FRHITransitionInfo::EType::Texture)
-					{
-						// NeedsAllPlanes only applied to textures
-						continue;
-					}
-
-					FRHIViewableResource* TransitionViewableResource = GetViewableResource(*TransitionInfo);
-					if (TransitionViewableResource == nullptr)
+					if (!TrackTransitioningResource(*TransitionCreateInfoData, *TransitionInfo))
 					{
 						continue;
 					}
 
-					TRHIResourceState<FRHICmdListBaseArrayAllocator>* TrackedResourceState = TrackedResources.Find(TransitionViewableResource);
-
 					const bool bShouldSkipAfterState = ((EnumHasAnyFlags(TransitionInfo->Flags, EResourceTransitionFlags::IgnoreAfterState)) || (TransitionInfo->AccessAfter == ERHIAccess::Unknown));
 					const bool bNeedUpdateAfterState = (bShouldSkipAfterState == false);
 					if (bNeedUpdateAfterState)
 					{
+						FRHIViewableResource* TransitionViewableResource = GetViewableResource(*TransitionInfo);
+						TRHIResourceState<FRHICmdListBaseArrayAllocator>* TrackedResourceState = TrackedResources.Find(TransitionViewableResource);
+
 						// Update the TrackedResource State
 						if (TrackedResourceState)
 						{
diff --git a/Engine/Source/Runtime/RHI/Public/RHICommandList.h b/Engine/Source/Runtime/RHI/Public/RHICommandList.h
--- a/Engine/Source/Runtime/RHI/Public/RHICommandList.h
+++ b/Engine/Source/Runtime/RHI/Public/RHICommandList.h
@@ -1530,6 +1530,9 @@ private:
 	inline void PatchTransitionsInPlace(TArrayView<const FRHITransition*> Transitions);
 	inline void UpdateTrackedAccessInPlace(TArrayView<const FRHITransition*> Transitions);
 
+	// (NeedsAllPlanes) Resource needs tracking while it's transitioning in case a RenderPass queries its state between Begin and End
+	RHI_API static bool TrackTransitioningResource(const FRHITransitionCreateInfoData& TransitionCreateInfoData, const FRHITransitionInfo& TransitionInfo);
+
 protected:
 	inline void ResolveNopAccessInPlace(FRHIRenderPassInfo& InInfo);
 	inline void TrackNopAccess(const FRHIRenderPassInfo& InInfo);
