Include uniform buffer name in the FRHIUniformBufferLayoutInitializer hash.
The hash is stored in compiled shaders and is used with FindUniformBufferStructByLayoutHash() to find the associated uniform buffer layout at runtime to resolve the uniform buffer slot.
Without the name being included it is easy to get a hash collision between two uniform buffers with similar layouts.
This leads to crashes when we retrieve the wrong uniform buffer (and slot) for our shader.
Note: There already exists a more detailed unique layout hash that is calculated and stored in FShaderParametersMetadata. But in our case we have code locations where the full layout isn't available and we only have the initializer.

diff --git a/Engine/Shaders/Public/ShaderVersion.ush b/Engine/Shaders/Public/ShaderVersion.ush
--- a/Engine/Shaders/Public/ShaderVersion.ush
+++ b/Engine/Shaders/Public/ShaderVersion.ush
@@ -3,5 +3,4 @@
 // Each time the console command is executed it generates a new GUID. As this file is included
 // in Platform.ush (which should be included in any shader) it allows to invalidate the shader DDC.
 // 
-#pragma message("UESHADERMETADATA_VERSION 75B6BD02-19BE-4FB4-B0B6-671BEBE9A2A1")
-
+#pragma message("UESHADERMETADATA_VERSION 51E571B6-FCD0-48DC-BE80-A580CB4BF409")
diff --git a/Engine/Source/Runtime/RHI/Public/RHIUniformBufferLayoutInitializer.h b/Engine/Source/Runtime/RHI/Public/RHIUniformBufferLayoutInitializer.h
--- a/Engine/Source/Runtime/RHI/Public/RHIUniformBufferLayoutInitializer.h
+++ b/Engine/Source/Runtime/RHI/Public/RHIUniformBufferLayoutInitializer.h
@@ -59,6 +59,12 @@ struct FRHIUniformBufferLayoutInitializer
 		return Hash;
 	}
 
+	/** 
+	 * This hash is stored with compiled shaders and is used for two things. 
+	 * 1. Validation that the data layout did not change between compilation and runtime.
+	 * 2. Passing to FindUniformBufferStructByLayoutHash() to find the associated uniform buffer layouts (and specifically the runtime slot mappings) for the shader.
+	 * Possibly it would be better to separate those concerns, but that would require more data stored per shader.
+	 */
 	void ComputeHash()
 	{
 		uint32 TmpHash = 0;
@@ -98,6 +104,11 @@ struct FRHIUniformBufferLayoutInitializer
 		{
 			TmpHash ^= Resources[--N].MemberType;
 		}
+
+		// Name is necessary to distinguish different uniform buffer declarations with matching layouts.
+		// Without it FindUniformBufferStructByLayoutHash() is vulnerable to hash collisions.
+		TmpHash ^= FCrc::StrCrc32(*Name);
+		
 		Hash = TmpHash;
 	}
 
