Merging RHI CreateBuffer refactor to Main.

diff --git a/Engine/Source/Runtime/Engine/Private/VolumeRendering.cpp b/Engine/Source/Runtime/Engine/Private/VolumeRendering.cpp
--- a/Engine/Source/Runtime/Engine/Private/VolumeRendering.cpp
+++ b/Engine/Source/Runtime/Engine/Private/VolumeRendering.cpp
@@ -51,22 +51,23 @@ void FVolumeRasterizeVertexBuffer::InitRHI(FRHICommandListBase& RHICmdList)
 {
 	// Used as a non-indexed triangle strip, so 4 vertices per quad
 	const uint32 Size = 4 * sizeof(FScreenVertex);
-	FRHIResourceCreateInfo CreateInfo(TEXT("FVolumeRasterizeVertexBuffer"));
-	VertexBufferRHI = RHICmdList.CreateBuffer(Size, BUF_Static | BUF_VertexBuffer, 0, ERHIAccess::VertexOrIndexBuffer, CreateInfo);
-	FScreenVertex* DestVertex = (FScreenVertex*)RHICmdList.LockBuffer(VertexBufferRHI, 0, Size, RLM_WriteOnly);
+
+	const FRHIBufferCreateDesc Desc =
+		FRHIBufferCreateDesc::CreateVertex<FScreenVertex>(TEXT("FVolumeRasterizeVertexBuffer"), 4)
+		.AddUsage(BUF_Static)
+		.SetInitialState(ERHIAccess::VertexOrIndexBuffer)
+		.SetInitActionInitializer();
+
+	TRHIBufferInitializer<FScreenVertex> Initializer = RHICmdList.CreateBufferInitializer(Desc);
 
 	// Setup a full - render target quad
 	// A viewport and UVScaleBias will be used to implement rendering to a sub region
-	DestVertex[0].Position = FVector2f(1, -GProjectionSignY);
-	DestVertex[0].UV = FVector2f(1, 1);
-	DestVertex[1].Position = FVector2f(1, GProjectionSignY);
-	DestVertex[1].UV = FVector2f(1, 0);
-	DestVertex[2].Position = FVector2f(-1, -GProjectionSignY);
-	DestVertex[2].UV = FVector2f(0, 1);
-	DestVertex[3].Position = FVector2f(-1, GProjectionSignY);
-	DestVertex[3].UV = FVector2f(0, 0);
-
-	RHICmdList.UnlockBuffer(VertexBufferRHI);
+	Initializer[0] = FScreenVertex { FVector2f( 1, -GProjectionSignY), FVector2f(1, 1) };
+	Initializer[1] = FScreenVertex { FVector2f( 1,  GProjectionSignY), FVector2f(1, 0) };
+	Initializer[2] = FScreenVertex { FVector2f(-1, -GProjectionSignY), FVector2f(0, 1) };
+	Initializer[3] = FScreenVertex { FVector2f(-1,  GProjectionSignY), FVector2f(0, 0) };
+
+	VertexBufferRHI = Initializer.Finalize();
 }
 
 TGlobalResource<FVolumeRasterizeVertexBuffer> GVolumeRasterizeVertexBuffer;
diff --git a/Engine/Source/Runtime/Engine/Public/ParticleHelper.h b/Engine/Source/Runtime/Engine/Public/ParticleHelper.h
--- a/Engine/Source/Runtime/Engine/Public/ParticleHelper.h
+++ b/Engine/Source/Runtime/Engine/Public/ParticleHelper.h
@@ -2376,11 +2376,19 @@ public:
 	virtual void InitRHI(FRHICommandListBase& RHICmdList) override
 	{
 		// create a static vertex buffer
-		FRHIResourceCreateInfo CreateInfo(TEXT("FNullDynamicParameterVertexBuffer"));
-		VertexBufferRHI = RHICmdList.CreateBuffer(sizeof(FParticleVertexDynamicParameter), BUF_Static | BUF_VertexBuffer, 0, ERHIAccess::VertexOrIndexBuffer, CreateInfo);
-		FParticleVertexDynamicParameter* Vertices = (FParticleVertexDynamicParameter*)RHICmdList.LockBuffer(VertexBufferRHI, 0, sizeof(FParticleVertexDynamicParameter), RLM_WriteOnly);
-		Vertices[0].DynamicValue[0] = Vertices[0].DynamicValue[1] = Vertices[0].DynamicValue[2] = Vertices[0].DynamicValue[3] = 1.0f;
-		RHICmdList.UnlockBuffer(VertexBufferRHI);
+		const FRHIBufferCreateDesc CreateDesc =
+			FRHIBufferCreateDesc::CreateVertex(TEXT("FNullDynamicParameterVertexBuffer"), sizeof(FParticleVertexDynamicParameter))
+			.AddUsage(EBufferUsageFlags::Static)
+			.SetInitialState(ERHIAccess::VertexOrIndexBuffer)
+			.SetInitActionInitializer();
+
+		FParticleVertexDynamicParameter Value;
+		Value.DynamicValue[0] = Value.DynamicValue[1] = Value.DynamicValue[2] = Value.DynamicValue[3] = 1.0f;
+
+		FRHIBufferInitializer Initializer = RHICmdList.CreateBufferInitializer(CreateDesc);
+		Initializer.WriteData(&Value, sizeof(Value));
+
+		VertexBufferRHI = Initializer.Finalize();
 	}
 };
 
diff --git a/Engine/Source/Runtime/Engine/Public/Rendering/SkeletalMeshDuplicatedVerticesBuffer.h b/Engine/Source/Runtime/Engine/Public/Rendering/SkeletalMeshDuplicatedVerticesBuffer.h
--- a/Engine/Source/Runtime/Engine/Public/Rendering/SkeletalMeshDuplicatedVerticesBuffer.h
+++ b/Engine/Source/Runtime/Engine/Public/Rendering/SkeletalMeshDuplicatedVerticesBuffer.h
@@ -112,11 +112,15 @@ public:
             FResourceArrayInterface* ResourceArray = DupVertData.GetResourceArray();
             check(ResourceArray->GetResourceDataSize() > 0);
 
-            FRHIResourceCreateInfo CreateInfo(TEXT("DuplicatedVerticesIndexBuffer"), ResourceArray);
-			CreateInfo.ClassName = DuplicatedVerticesBuffeName;
-			CreateInfo.OwnerName = GetOwnerName();
-            DuplicatedVerticesIndexBuffer.VertexBufferRHI = RHICmdList.CreateVertexBuffer(ResourceArray->GetResourceDataSize(), BUF_Static | BUF_ShaderResource, CreateInfo);
-			DuplicatedVerticesIndexBuffer.VertexBufferRHI->SetOwnerName(GetOwnerName());
+			const FRHIBufferCreateDesc Desc =
+				FRHIBufferCreateDesc::CreateVertex(TEXT("DuplicatedVerticesIndexBuffer"), ResourceArray->GetResourceDataSize())
+				.AddUsage(BUF_Static | BUF_ShaderResource)
+				.SetClassName(DuplicatedVerticesBuffeName)
+				.SetOwnerName(GetOwnerName())
+				.SetInitActionResourceArray(ResourceArray)
+				.DetermineInitialState();
+
+			DuplicatedVerticesIndexBuffer.VertexBufferRHI = RHICmdList.CreateBuffer(Desc);
             DuplicatedVerticesIndexBuffer.VertexBufferSRV = RHICmdList.CreateShaderResourceView(
 				DuplicatedVerticesIndexBuffer.VertexBufferRHI,
 				FRHIViewDesc::CreateBufferSRV()
@@ -128,11 +132,15 @@ public:
             FResourceArrayInterface* ResourceArray = DupVertIndexData.GetResourceArray();
             check(ResourceArray->GetResourceDataSize() > 0);
 
-            FRHIResourceCreateInfo CreateInfo(TEXT("LengthAndIndexDuplicatedVerticesIndexBuffer"), ResourceArray);
-			CreateInfo.ClassName = DuplicatedVerticesBuffeName;
-			CreateInfo.OwnerName = GetOwnerName();
-            LengthAndIndexDuplicatedVerticesIndexBuffer.VertexBufferRHI = RHICmdList.CreateVertexBuffer(ResourceArray->GetResourceDataSize(), BUF_Static | BUF_ShaderResource, CreateInfo);
-			LengthAndIndexDuplicatedVerticesIndexBuffer.VertexBufferRHI->SetOwnerName(GetOwnerName());
+			const FRHIBufferCreateDesc Desc =
+				FRHIBufferCreateDesc::CreateVertex(TEXT("LengthAndIndexDuplicatedVerticesIndexBuffer"), ResourceArray->GetResourceDataSize())
+				.AddUsage(BUF_Static | BUF_ShaderResource)
+				.SetClassName(DuplicatedVerticesBuffeName)
+				.SetOwnerName(GetOwnerName())
+				.SetInitActionResourceArray(ResourceArray)
+				.DetermineInitialState();
+
+            LengthAndIndexDuplicatedVerticesIndexBuffer.VertexBufferRHI = RHICmdList.CreateBuffer(Desc);
             LengthAndIndexDuplicatedVerticesIndexBuffer.VertexBufferSRV = RHICmdList.CreateShaderResourceView(
 				LengthAndIndexDuplicatedVerticesIndexBuffer.VertexBufferRHI,
 				FRHIViewDesc::CreateBufferSRV()
diff --git a/Engine/Source/Runtime/Experimental/GeometryCollectionEngine/Private/GeometryCollection/GeometryCollectionSceneProxy.h b/Engine/Source/Runtime/Experimental/GeometryCollectionEngine/Private/GeometryCollection/GeometryCollectionSceneProxy.h
--- a/Engine/Source/Runtime/Experimental/GeometryCollectionEngine/Private/GeometryCollection/GeometryCollectionSceneProxy.h
+++ b/Engine/Source/Runtime/Experimental/GeometryCollectionEngine/Private/GeometryCollection/GeometryCollectionSceneProxy.h
@@ -34,11 +34,14 @@ class FGeometryCollectionTransformBuffer : public FVertexBuffer
 public:
 	virtual void InitRHI(FRHICommandListBase& RHICmdList) override
 	{
-		FRHIResourceCreateInfo CreateInfo(TEXT("FGeometryCollectionTransformBuffer"));
-
 		// #note: This differs from instanced static mesh in that we are storing the entire transform in the buffer rather than
 		// splitting out the translation.  This is to simplify transferring data at runtime as a memcopy
-		VertexBufferRHI = RHICmdList.CreateVertexBuffer(NumTransforms * sizeof(FVector4f) * 4, BUF_Dynamic | BUF_ShaderResource, CreateInfo);		
+		const FRHIBufferCreateDesc CreateDesc =
+			FRHIBufferCreateDesc::CreateVertex<FVector4f>(TEXT("FGeometryCollectionTransformBuffer"), NumTransforms * 4)
+			.AddUsage(EBufferUsageFlags::Dynamic)
+			.DetermineInitialState();
+
+		VertexBufferRHI = RHICmdList.CreateBuffer(CreateDesc);
 		VertexBufferSRV = RHICmdList.CreateShaderResourceView(
 			VertexBufferRHI,
 			FRHIViewDesc::CreateBufferSRV()
diff --git a/Engine/Source/Runtime/GameplayMediaEncoder/Private/GameplayMediaEncoder.cpp b/Engine/Source/Runtime/GameplayMediaEncoder/Private/GameplayMediaEncoder.cpp
--- a/Engine/Source/Runtime/GameplayMediaEncoder/Private/GameplayMediaEncoder.cpp
+++ b/Engine/Source/Runtime/GameplayMediaEncoder/Private/GameplayMediaEncoder.cpp
@@ -673,7 +673,6 @@ PRAGMA_ENABLE_DEPRECATION_WARNINGS
 
 		if (RHIType == ERHIInterfaceType::D3D11)
 		{
-			FRHIResourceCreateInfo CreateInfo(TEXT("VideoCapturerBackBuffer"));
 			FTextureRHIRef Texture = RHICreateTexture(Desc);
 
 			PRAGMA_DISABLE_DEPRECATION_WARNINGS
@@ -683,7 +682,6 @@ PRAGMA_ENABLE_DEPRECATION_WARNINGS
 		}
 		else if (RHIType == ERHIInterfaceType::D3D12)
 		{
-			FRHIResourceCreateInfo CreateInfo(TEXT("VideoCapturerBackBuffer"));
 			FTextureRHIRef Texture = RHICreateTexture(Desc);
 
 			PRAGMA_DISABLE_DEPRECATION_WARNINGS
diff --git a/Engine/Source/Runtime/HeadMountedDisplay/Private/HeadMountedDisplayTypes.cpp b/Engine/Source/Runtime/HeadMountedDisplay/Private/HeadMountedDisplayTypes.cpp
--- a/Engine/Source/Runtime/HeadMountedDisplay/Private/HeadMountedDisplayTypes.cpp
+++ b/Engine/Source/Runtime/HeadMountedDisplay/Private/HeadMountedDisplayTypes.cpp
@@ -26,14 +26,20 @@ void FHMDViewMesh::BuildMesh(const FVector2D Positions[], uint32 VertexCount, EH
 	NumTriangles = NumVertices / 3;
 	NumIndices = NumVertices;
 
-	FRHIResourceCreateInfo CreateInfo(TEXT("FHMDViewMesh"));
-	VertexBufferRHI = RHICmdList.CreateVertexBuffer(sizeof(FFilterVertex) * NumVertices, BUF_Static, CreateInfo);
-	void* VoidPtr = RHICmdList.LockBuffer(VertexBufferRHI, 0, sizeof(FFilterVertex) * NumVertices, RLM_WriteOnly);
-	FFilterVertex* pVertices = reinterpret_cast<FFilterVertex*>(VoidPtr);
+	const FRHIBufferCreateDesc VertexCreateDesc =
+		FRHIBufferCreateDesc::CreateVertex<FFilterVertex>(TEXT("FHMDViewMesh"), NumVertices)
+		.AddUsage(EBufferUsageFlags::Static)
+		.SetInitActionInitializer()
+		.DetermineInitialState();
 
-	IndexBufferRHI = RHICmdList.CreateIndexBuffer(sizeof(uint16), sizeof(uint16) * NumIndices, BUF_Static, CreateInfo);
-	void* VoidPtr2 = RHICmdList.LockBuffer(IndexBufferRHI, 0, sizeof(uint16) * NumIndices, RLM_WriteOnly);
-	uint16* pIndices = reinterpret_cast<uint16*>(VoidPtr2);
+	const FRHIBufferCreateDesc IndexCreateDesc =
+		FRHIBufferCreateDesc::CreateIndex<uint16>(TEXT("FHMDViewMesh"), NumIndices)
+		.AddUsage(EBufferUsageFlags::Static)
+		.SetInitActionInitializer()
+		.DetermineInitialState();
+
+	TRHIBufferInitializer<FFilterVertex> pVertices = RHICmdList.CreateBufferInitializer(VertexCreateDesc);
+	TRHIBufferInitializer<uint16> pIndices = RHICmdList.CreateBufferInitializer(IndexCreateDesc);
 
 	for (uint32 VertexIndex = 0; VertexIndex < NumVertices; ++VertexIndex)
 	{
@@ -67,6 +73,6 @@ void FHMDViewMesh::BuildMesh(const FVector2D Positions[], uint32 VertexCount, EH
 		pIndices[VertexIndex] = static_cast<uint16>(VertexIndex);
 	}
 
-	RHICmdList.UnlockBuffer(VertexBufferRHI);
-	RHICmdList.UnlockBuffer(IndexBufferRHI);
+	VertexBufferRHI = pVertices.Finalize();
+	IndexBufferRHI = pIndices.Finalize();
 }
diff --git a/Engine/Source/Runtime/Landscape/Private/LandscapeCulling.cpp b/Engine/Source/Runtime/Landscape/Private/LandscapeCulling.cpp
--- a/Engine/Source/Runtime/Landscape/Private/LandscapeCulling.cpp
+++ b/Engine/Source/Runtime/Landscape/Private/LandscapeCulling.cpp
@@ -244,33 +244,45 @@ void FLandscapeTileMesh::InitRHI(FRHICommandListBase& RHICmdList)
 		const uint32 NumVertx = FMath::Square(TileSizeVertx);
 		const uint32 Stride = sizeof(FLandscapeVertex);
 
-		FRHIResourceCreateInfo CreateInfo(TEXT("FLandscapeTileMeshVertexBuffer"));
-		VertexBuffer.VertexBufferRHI = RHICmdList.CreateBuffer(NumVertx * Stride, BUF_Static | BUF_VertexBuffer, Stride, ERHIAccess::VertexOrIndexBuffer, CreateInfo);
-		FLandscapeVertex* Vertex = (FLandscapeVertex*)RHICmdList.LockBuffer(VertexBuffer.VertexBufferRHI, 0, NumVertx * Stride, RLM_WriteOnly);
+		const FRHIBufferCreateDesc CreateDesc =
+			FRHIBufferCreateDesc::Create(TEXT("FLandscapeTileMeshVertexBuffer"), NumVertx * Stride, Stride, EBufferUsageFlags::Static | EBufferUsageFlags::VertexBuffer)
+			.SetInitialState(ERHIAccess::VertexOrIndexBuffer)
+			.SetInitActionInitializer();
 
+		TRHIBufferInitializer<FLandscapeVertex> Initializer = RHICmdList.CreateBufferInitializer(CreateDesc);
+
+		uint32 VertexIndex = 0;
 		for (uint32 y = 0; y < TileSizeVertx; y++)
 		{
 			for (uint32 x = 0; x < TileSizeVertx; x++)
 			{
-				Vertex->VertexX = static_cast<uint8>(x);
-				Vertex->VertexY = static_cast<uint8>(y);
-				Vertex->SubX = 0;
-				Vertex->SubY = 0;
-				Vertex++;
+				FLandscapeVertex Vertex{};
+				Vertex.VertexX = static_cast<uint8>(x);
+				Vertex.VertexY = static_cast<uint8>(y);
+				Vertex.SubX = 0;
+				Vertex.SubY = 0;
+
+				Initializer.WriteValueAtIndex(VertexIndex, Vertex);
+				VertexIndex++;
 			}
 		}
 
-		RHICmdList.UnlockBuffer(VertexBuffer.VertexBufferRHI);
+		VertexBuffer.VertexBufferRHI = Initializer.Finalize();
 	}
 
 	// 
 	{
 		const int32 NumIndices = FMath::Square(LANDSCAPE_TILE_QUADS) * 6u;
 		const uint32 Stride = sizeof(uint16);
-		FRHIResourceCreateInfo CreateInfo(TEXT("FLandscapeTileMeshIndexBuffer"));
-		IndexBuffer.IndexBufferRHI = RHICmdList.CreateBuffer(NumIndices * Stride, BUF_Static | BUF_IndexBuffer, Stride, ERHIAccess::VertexOrIndexBuffer, CreateInfo);
-		uint16* Indices = (uint16*)RHICmdList.LockBuffer(IndexBuffer.IndexBufferRHI, 0, NumIndices * Stride, RLM_WriteOnly);
 
+		const FRHIBufferCreateDesc CreateDesc =
+			FRHIBufferCreateDesc::Create(TEXT("FLandscapeTileMeshIndexBuffer"), NumIndices * Stride, Stride, EBufferUsageFlags::Static | EBufferUsageFlags::IndexBuffer)
+			.SetInitialState(ERHIAccess::VertexOrIndexBuffer)
+			.SetInitActionInitializer();
+
+		TRHIBufferInitializer<uint16> Initializer = RHICmdList.CreateBufferInitializer(CreateDesc);
+
+		uint32 IndexOffset = 0;
 		for (uint32 y = 0; y < LANDSCAPE_TILE_QUADS; y++)
 		{
 			for (uint32 x = 0; x < LANDSCAPE_TILE_QUADS; x++)
@@ -280,18 +292,19 @@ void FLandscapeTileMesh::InitRHI(FRHICommandListBase& RHICmdList)
 				uint16 i11 = static_cast<uint16>((x + 1) + (y + 1) * (TileSizeVertx));
 				uint16 i01 = static_cast<uint16>((x + 0) + (y + 1) * (TileSizeVertx));
 
-				Indices[0] = i00;
-				Indices[1] = i11;
-				Indices[2] = i10;
+				Initializer[IndexOffset + 0] = i00;
+				Initializer[IndexOffset + 1] = i11;
+				Initializer[IndexOffset + 2] = i10;
 
-				Indices[3] = i00;
-				Indices[4] = i01;
-				Indices[5] = i11;
+				Initializer[IndexOffset + 3] = i00;
+				Initializer[IndexOffset + 4] = i01;
+				Initializer[IndexOffset + 5] = i11;
 
-				Indices += 6;
+				IndexOffset += 6;
 			}
 		}
-		RHICmdList.UnlockBuffer(IndexBuffer.IndexBufferRHI);
+
+		IndexBuffer.IndexBufferRHI = Initializer.Finalize();
 	}
 }
 
@@ -325,10 +338,14 @@ void FLandscapeTileDataBuffer::InitRHI(FRHICommandListBase& RHICmdList)
 	const uint32 ComponentTiles = SubsectionTiles * NumSubsections * NumSubsections;
 	const uint32 Stride = TILE_DATA_ENTRY_NUM_BYTES;
 
-	FRHIResourceCreateInfo CreateInfo(TEXT("FLandscapeTileVertexBuffer"));
-	VertexBufferRHI = RHICmdList.CreateBuffer(ComponentTiles * Stride, BUF_Static | BUF_VertexBuffer, Stride, ERHIAccess::VertexOrIndexBuffer, CreateInfo);
-	uint8* TileData = (uint8*)RHICmdList.LockBuffer(VertexBufferRHI, 0, ComponentTiles * Stride, RLM_WriteOnly);
+	const FRHIBufferCreateDesc CreateDesc =
+		FRHIBufferCreateDesc::Create(TEXT("FLandscapeTileVertexBuffer"), ComponentTiles * Stride, Stride, EBufferUsageFlags::Static | EBufferUsageFlags::VertexBuffer)
+		.SetInitialState(ERHIAccess::VertexOrIndexBuffer)
+		.SetInitActionInitializer();
+
+	TRHIBufferInitializer<uint8> Initializer = RHICmdList.CreateBufferInitializer(CreateDesc);
 
+	uint32 TileOffset = 0;
 	for (uint32 SubY = 0; SubY < NumSubsections; SubY++)
 	{
 		for (uint32 SubX = 0; SubX < NumSubsections; SubX++)
@@ -337,17 +354,17 @@ void FLandscapeTileDataBuffer::InitRHI(FRHICommandListBase& RHICmdList)
 			{
 				for (uint32 x = 0; x < SubsectionTilesRow; x++)
 				{
-					TileData[0] = static_cast<uint8>(x * LANDSCAPE_TILE_QUADS);
-					TileData[1] = static_cast<uint8>(y * LANDSCAPE_TILE_QUADS);
-					TileData[2] = static_cast<uint8>(SubX);
-					TileData[3] = static_cast<uint8>(SubY);
-					TileData += Stride;
+					Initializer[TileOffset + 0] = static_cast<uint8>(x * LANDSCAPE_TILE_QUADS);
+					Initializer[TileOffset + 1] = static_cast<uint8>(y * LANDSCAPE_TILE_QUADS);
+					Initializer[TileOffset + 2] = static_cast<uint8>(SubX);
+					Initializer[TileOffset + 3] = static_cast<uint8>(SubY);
+					TileOffset += Stride;
 				}
 			}
 		}
 	}
 
-	RHICmdList.UnlockBuffer(VertexBufferRHI);
+	VertexBufferRHI = Initializer.Finalize();
 }
 
 class FBuildLandscapeTileDataCS : public FGlobalShader
diff --git a/Engine/Source/Runtime/Landscape/Private/LandscapeEditLayers.cpp b/Engine/Source/Runtime/Landscape/Private/LandscapeEditLayers.cpp
--- a/Engine/Source/Runtime/Landscape/Private/LandscapeEditLayers.cpp
+++ b/Engine/Source/Runtime/Landscape/Private/LandscapeEditLayers.cpp
@@ -1016,13 +1016,13 @@ public:
 	/** Called when the resource is initialized. This is only called by the rendering thread. */
 	virtual void InitRHI(FRHICommandListBase& RHICmdList) override
 	{
-		FRHIResourceCreateInfo CreateInfo(TEXT("FLandscapeLayerWeightmapExtractMaterialLayersComputeShaderResource"));
-		ComponentsData = RHICmdList.CreateStructuredBuffer(sizeof(FLandscapeLayerWeightmapExtractMaterialLayersComponentData), OriginalComponentsData.Num() * sizeof(FLandscapeLayerWeightmapExtractMaterialLayersComponentData), BUF_ShaderResource | BUF_Volatile, CreateInfo);
+		ComponentsData = UE::RHIResourceUtils::CreateBufferFromArray(
+			RHICmdList,
+			TEXT("FLandscapeLayerWeightmapExtractMaterialLayersComputeShaderResource"),
+			EBufferUsageFlags::StructuredBuffer | EBufferUsageFlags::ShaderResource | EBufferUsageFlags::Volatile,
+			MakeConstArrayView(OriginalComponentsData)
+		);
 		ComponentsDataSRV = RHICmdList.CreateShaderResourceView(ComponentsData, FRHIViewDesc::CreateBufferSRV().SetTypeFromBuffer(ComponentsData));
-
-		uint8* Buffer = (uint8*)RHICmdList.LockBuffer(ComponentsData, 0, OriginalComponentsData.Num() * sizeof(FLandscapeLayerWeightmapExtractMaterialLayersComponentData), RLM_WriteOnly);
-		FMemory::Memcpy(Buffer, OriginalComponentsData.GetData(), OriginalComponentsData.Num() * sizeof(FLandscapeLayerWeightmapExtractMaterialLayersComponentData));
-		RHICmdList.UnlockBuffer(ComponentsData);
 	}
 
 	virtual void ReleaseRHI() override
@@ -1175,40 +1175,37 @@ public:
 	/** Called when the resource is initialized. This is only called by the rendering thread. */
 	virtual void InitRHI(FRHICommandListBase& RHICmdList) override
 	{
-		FRHIResourceCreateInfo CreateInfo(TEXT("ComponentsData"));
-		uint32 ComponentsDataMemSize = OriginalComponentsData.Num() * sizeof(FLandscapeLayerWeightmapPackMaterialLayersComponentData);
-		ComponentsData = RHICmdList.CreateStructuredBuffer(sizeof(FLandscapeLayerWeightmapPackMaterialLayersComponentData), ComponentsDataMemSize, BUF_ShaderResource | BUF_Volatile, CreateInfo);
+		ComponentsData = UE::RHIResourceUtils::CreateBufferFromArray(
+			RHICmdList,
+			TEXT("ComponentsData"),
+			EBufferUsageFlags::StructuredBuffer | EBufferUsageFlags::ShaderResource | EBufferUsageFlags::Volatile,
+			MakeConstArrayView(OriginalComponentsData)
+		);
 		ComponentsDataSRV = RHICmdList.CreateShaderResourceView(ComponentsData, FRHIViewDesc::CreateBufferSRV().SetTypeFromBuffer(ComponentsData));
 
-		uint8* Buffer = (uint8*)RHICmdList.LockBuffer(ComponentsData, 0, ComponentsDataMemSize, RLM_WriteOnly);
-		FMemory::Memcpy(Buffer, OriginalComponentsData.GetData(), ComponentsDataMemSize);
-		RHICmdList.UnlockBuffer(ComponentsData);
-
-		FRHIResourceCreateInfo WeightBlendCreateInfo(TEXT("WeightmapWeightBlendMode"));
-		uint32 WeightBlendMemSize = OriginalWeightmapWeightBlendModeData.Num() * sizeof(float);
-		WeightmapWeightBlendMode = RHICmdList.CreateVertexBuffer(WeightBlendMemSize, BUF_ShaderResource | BUF_Volatile, WeightBlendCreateInfo);
+		WeightmapWeightBlendMode = UE::RHIResourceUtils::CreateVertexBufferFromArray(
+			RHICmdList,
+			TEXT("WeightmapWeightBlendMode"),
+			EBufferUsageFlags::ShaderResource | EBufferUsageFlags::Volatile,
+			MakeConstArrayView(OriginalWeightmapWeightBlendModeData)
+		);
 		WeightmapWeightBlendModeSRV = RHICmdList.CreateShaderResourceView(
 			WeightmapWeightBlendMode, 
 			FRHIViewDesc::CreateBufferSRV()
 				.SetType(FRHIViewDesc::EBufferType::Typed)
 				.SetFormat(PF_R32_FLOAT));
 
-		void* WeightmapWeightBlendModePtr = RHICmdList.LockBuffer(WeightmapWeightBlendMode, 0, WeightBlendMemSize, RLM_WriteOnly);
-		FMemory::Memcpy(WeightmapWeightBlendModePtr, OriginalWeightmapWeightBlendModeData.GetData(), WeightBlendMemSize);
-		RHICmdList.UnlockBuffer(WeightmapWeightBlendMode);
-
-		FRHIResourceCreateInfo TextureOutputOffsetCreateInfo(TEXT("WeightmapTextureOutputOffset"));
-		uint32 TextureOutputOffsetMemSize = OriginalTextureOutputOffset.Num() * sizeof(FVector2f);
-		WeightmapTextureOutputOffset = RHICmdList.CreateVertexBuffer(TextureOutputOffsetMemSize, BUF_ShaderResource | BUF_Volatile, TextureOutputOffsetCreateInfo);
+		WeightmapTextureOutputOffset = UE::RHIResourceUtils::CreateVertexBufferFromArray(
+			RHICmdList,
+			TEXT("WeightmapTextureOutputOffset"),
+			EBufferUsageFlags::ShaderResource | EBufferUsageFlags::Volatile,
+			MakeConstArrayView(OriginalTextureOutputOffset)
+		);
 		WeightmapTextureOutputOffsetSRV = RHICmdList.CreateShaderResourceView(
 			WeightmapTextureOutputOffset,
 			FRHIViewDesc::CreateBufferSRV()
 				.SetType(FRHIViewDesc::EBufferType::Typed)
 				.SetFormat(PF_G32R32F));
-
-		void* TextureOutputOffsetPtr = RHICmdList.LockBuffer(WeightmapTextureOutputOffset, 0, TextureOutputOffsetMemSize, RLM_WriteOnly);
-		FMemory::Memcpy(TextureOutputOffsetPtr, OriginalTextureOutputOffset.GetData(), TextureOutputOffsetMemSize);
-		RHICmdList.UnlockBuffer(WeightmapTextureOutputOffset);
 	}
 
 	virtual void ReleaseRHI() override
diff --git a/Engine/Source/Runtime/Landscape/Private/LandscapeRender.cpp b/Engine/Source/Runtime/Landscape/Private/LandscapeRender.cpp
--- a/Engine/Source/Runtime/Landscape/Private/LandscapeRender.cpp
+++ b/Engine/Source/Runtime/Landscape/Private/LandscapeRender.cpp
@@ -1060,10 +1060,16 @@ void FLandscapeRenderSystem::UpdateBuffers(FRHICommandListBase& RHICmdList)
 	{
 		if (!SectionLODBiasBuffer.IsValid())
 		{
-			FRHIResourceCreateInfo CreateInfo(TEXT("SectionLODBiasBuffer"), &SectionLODBiases);
 			const static FLazyName ClassName(TEXT("FLandscapeRenderSystem"));
-			CreateInfo.ClassName = ClassName;
-			SectionLODBiasBuffer = RHICmdList.CreateVertexBuffer(SectionLODBiases.GetResourceDataSize(), BUF_ShaderResource | BUF_Dynamic, CreateInfo);
+
+			const FRHIBufferCreateDesc CreateDesc =
+				FRHIBufferCreateDesc::CreateVertex(TEXT("SectionLODBiasBuffer"), SectionLODBiases.GetResourceDataSize())
+				.AddUsage(EBufferUsageFlags::ShaderResource | EBufferUsageFlags::Dynamic)
+				.SetClassName(ClassName)
+				.SetInitActionResourceArray(&SectionLODBiases)
+				.DetermineInitialState();
+
+			SectionLODBiasBuffer = RHICmdList.CreateBuffer(CreateDesc);
 			SectionLODBiasSRV = RHICmdList.CreateShaderResourceView(
 				SectionLODBiasBuffer,
 				FRHIViewDesc::CreateBufferSRV()
@@ -1189,21 +1195,35 @@ void FLandscapeSceneViewExtension::PreRenderView_RenderThread(FRDGBuilder& Graph
 
 			for (FLandscapeViewData& LandscapeView : *LocalLandscapeViewsPtr)
 			{
-				FRHIResourceCreateInfo CreateInfoLODBuffer(TEXT("LandscapeLODDataBuffer"), &LandscapeView.LandscapeLODData);
-				FBufferRHIRef LandscapeLODDataBuffer = RHICmdList.CreateVertexBuffer(LandscapeView.LandscapeLODData.GetResourceDataSize(), BUF_ShaderResource | BUF_Volatile, CreateInfoLODBuffer);
-				LandscapeView.View->LandscapePerComponentDataBuffer = RHICmdList.CreateShaderResourceView(
-					LandscapeLODDataBuffer,
-					FRHIViewDesc::CreateBufferSRV()
-						.SetType(FRHIViewDesc::EBufferType::Typed)
-						.SetFormat(PF_R32_FLOAT));
-
-				FRHIResourceCreateInfo CreateInfoIndirection(TEXT("LandscapeIndirectionBuffer"), &LandscapeView.LandscapeIndirection);
-				FBufferRHIRef LandscapeIndirectionBuffer = RHICmdList.CreateVertexBuffer(LandscapeView.LandscapeIndirection.GetResourceDataSize(), BUF_ShaderResource | BUF_Volatile, CreateInfoIndirection);
-				LandscapeView.View->LandscapeIndirectionBuffer = RHICmdList.CreateShaderResourceView(
-					LandscapeIndirectionBuffer,
-					FRHIViewDesc::CreateBufferSRV()
-						.SetType(FRHIViewDesc::EBufferType::Typed)
-						.SetFormat(PF_R32_UINT));
+				{
+					const FRHIBufferCreateDesc CreateDesc =
+						FRHIBufferCreateDesc::CreateVertex(TEXT("LandscapeLODDataBuffer"), LandscapeView.LandscapeLODData.GetResourceDataSize())
+						.AddUsage(EBufferUsageFlags::ShaderResource | EBufferUsageFlags::Volatile)
+						.SetInitActionResourceArray(&LandscapeView.LandscapeLODData)
+						.DetermineInitialState();
+
+					FBufferRHIRef LandscapeLODDataBuffer = RHICmdList.CreateBuffer(CreateDesc);
+					LandscapeView.View->LandscapePerComponentDataBuffer = RHICmdList.CreateShaderResourceView(
+						LandscapeLODDataBuffer,
+						FRHIViewDesc::CreateBufferSRV()
+							.SetType(FRHIViewDesc::EBufferType::Typed)
+							.SetFormat(PF_R32_FLOAT));
+				}
+
+				{
+					const FRHIBufferCreateDesc CreateDesc =
+						FRHIBufferCreateDesc::CreateVertex(TEXT("LandscapeIndirectionBuffer"), LandscapeView.LandscapeIndirection.GetResourceDataSize())
+						.AddUsage(EBufferUsageFlags::ShaderResource | EBufferUsageFlags::Volatile)
+						.SetInitActionResourceArray(&LandscapeView.LandscapeIndirection)
+						.DetermineInitialState();
+
+					FBufferRHIRef LandscapeIndirectionBuffer = RHICmdList.CreateBuffer(CreateDesc);
+					LandscapeView.View->LandscapeIndirectionBuffer = RHICmdList.CreateShaderResourceView(
+						LandscapeIndirectionBuffer,
+						FRHIViewDesc::CreateBufferSRV()
+							.SetType(FRHIViewDesc::EBufferType::Typed)
+							.SetFormat(PF_R32_UINT));
+				}
 			}
 
 		}, GIsThreadedRendering && GLandscapeUseAsyncTasksForLODComputation);
@@ -3197,13 +3217,19 @@ void FLandscapeVertexBuffer::InitRHI(FRHICommandListBase& RHICmdList)
 {
 	SCOPED_LOADTIMER(FLandscapeVertexBuffer_InitRHI);
 
-	// create a static vertex buffer
-	FRHIResourceCreateInfo CreateInfo(TEXT("FLandscapeVertexBuffer"));
 	const static FLazyName ClassName(TEXT("FLandscapeVertexBuffer"));
-	CreateInfo.ClassName = ClassName;
-	CreateInfo.OwnerName = GetOwnerName();
-	VertexBufferRHI = RHICmdList.CreateBuffer(NumVertices * sizeof(FLandscapeVertex), BUF_Static | BUF_VertexBuffer, 0, ERHIAccess::VertexOrIndexBuffer, CreateInfo);
-	FLandscapeVertex* Vertex = (FLandscapeVertex*)RHICmdList.LockBuffer(VertexBufferRHI, 0, NumVertices * sizeof(FLandscapeVertex), RLM_WriteOnly);
+
+	// create a static vertex buffer
+	const FRHIBufferCreateDesc CreateDesc =
+		FRHIBufferCreateDesc::CreateVertex<FLandscapeVertex>(TEXT("FLandscapeVertexBuffer"), NumVertices)
+		.AddUsage(EBufferUsageFlags::Static)
+		.SetInitialState(ERHIAccess::VertexOrIndexBuffer)
+		.SetClassName(ClassName)
+		.SetOwnerName(GetOwnerName())
+		.SetInitActionInitializer();
+
+	TRHIBufferInitializer<FLandscapeVertex> Initializer = RHICmdList.CreateBufferInitializer(CreateDesc);
+
 	int32 VertexIndex = 0;
 	for (int32 SubY = 0; SubY < NumSubsections; SubY++)
 	{
@@ -3213,18 +3239,21 @@ void FLandscapeVertexBuffer::InitRHI(FRHICommandListBase& RHICmdList)
 			{
 				for (int32 x = 0; x < SubsectionSizeVerts; x++)
 				{
-					Vertex->VertexX = static_cast<uint8>(x);
-					Vertex->VertexY = static_cast<uint8>(y);
-					Vertex->SubX = static_cast<uint8>(SubX);
-					Vertex->SubY = static_cast<uint8>(SubY);
-					Vertex++;
+					FLandscapeVertex Vertex{};
+					Vertex.VertexX = static_cast<uint8>(x);
+					Vertex.VertexY = static_cast<uint8>(y);
+					Vertex.SubX = static_cast<uint8>(SubX);
+					Vertex.SubY = static_cast<uint8>(SubY);
+
+					Initializer.WriteValueAtIndex(VertexIndex, Vertex);
 					VertexIndex++;
 				}
 			}
 		}
 	}
 	check(NumVertices == VertexIndex);
-	RHICmdList.UnlockBuffer(VertexBufferRHI);
+
+	VertexBufferRHI = Initializer.Finalize();
 }
 
 //
diff --git a/Engine/Source/Runtime/MediaAssets/Private/Misc/MediaTextureResource.cpp b/Engine/Source/Runtime/MediaAssets/Private/Misc/MediaTextureResource.cpp
--- a/Engine/Source/Runtime/MediaAssets/Private/Misc/MediaTextureResource.cpp
+++ b/Engine/Source/Runtime/MediaAssets/Private/Misc/MediaTextureResource.cpp
@@ -1056,7 +1056,7 @@ void FMediaTextureResource::GetColorSpaceConversionMatrixForSample(const TShared
 
 		// draw full size quad into render target
 		// This needs to happen before we begin to setup the draw call, because on DX11, this might flush the command list more or less randomly
-		FBufferRHIRef VertexBuffer = CreateTempMediaVertexBuffer(); 
+		FBufferRHIRef VertexBuffer = CreateTempMediaVertexBuffer(RHICmdList);
 
 		FGraphicsPipelineStateInitializer GraphicsPSOInit;
 		FRHITexture* RenderTarget = IntermediateTarget;
@@ -1523,7 +1523,7 @@ void FMediaTextureResource::CopyFromExternalTexture(FRHICommandListImmediate& RH
 			SetShaderParametersLegacyPS(RHICmdList, CopyShader, SampleTexture, SamplerState, ScaleRotation, Offset);
 
 			// draw full size quad into render target
-			FBufferRHIRef VertexBuffer = CreateTempMediaVertexBuffer();
+			FBufferRHIRef VertexBuffer = CreateTempMediaVertexBuffer(RHICmdList);
 			RHICmdList.SetStreamSource(0, VertexBuffer, 0);
 			// set viewport to RT size
 			RHICmdList.SetViewport(0, 0, 0.0f, (float)OutputDim.X, (float)OutputDim.Y, 1.0f);
diff --git a/Engine/Source/Runtime/NullDrv/Public/NullRHI.h b/Engine/Source/Runtime/NullDrv/Public/NullRHI.h
--- a/Engine/Source/Runtime/NullDrv/Public/NullRHI.h
+++ b/Engine/Source/Runtime/NullDrv/Public/NullRHI.h
@@ -116,13 +116,39 @@ public:
 
 	}
 
-	virtual FBufferRHIRef RHICreateBuffer(FRHICommandListBase& RHICmdList, FRHIBufferDesc const& Desc, ERHIAccess ResourceState, FRHIResourceCreateInfo& CreateInfo) final override
-	{ 
-		if(CreateInfo.ResourceArray) 
-		{ 
-			CreateInfo.ResourceArray->Discard(); 
-		} 
-		return new FRHIBuffer(Desc);
+	struct FNullBufferInitializer : public FRHIBufferInitializer
+	{
+		FNullBufferInitializer(FRHICommandListBase& RHICmdList, FRHIBuffer* Buffer, void* InWritableData, uint64 InWritableDataSize)
+			: FRHIBufferInitializer(RHICmdList, Buffer, InWritableData, InWritableDataSize,
+				[Buffer = TRefCountPtr<FRHIBuffer>(Buffer)](FRHICommandListBase&) mutable
+				{
+					return MoveTemp(Buffer);
+				}
+			)
+		{
+		}
+	};
+
+	[[nodiscard]]
+	virtual FRHIBufferInitializer RHICreateBufferInitializer(FRHICommandListBase& RHICmdList, const FRHIBufferCreateDesc& CreateDesc) final override
+	{
+		if (CreateDesc.InitAction == ERHIBufferInitAction::ResourceArray && CreateDesc.InitialData)
+		{
+			CreateDesc.InitialData->Discard();
+		}
+
+		void* WritableData = nullptr;
+		uint64 WritableDataSize = 0;
+
+		if (CreateDesc.InitAction == ERHIBufferInitAction::Initializer)
+		{
+			WritableData = GetStaticBuffer(CreateDesc.Size);
+			WritableDataSize = CreateDesc.Size;
+		}
+
+		FRHIBuffer* Buffer = new FRHIBuffer(CreateDesc);
+
+		return FNullBufferInitializer(RHICmdList, Buffer, WritableData, WritableDataSize);
 	}
 
 	virtual void* LockBuffer_BottomOfPipe(FRHICommandListBase& RHICmdList, FRHIBuffer* Buffer, uint32 Offset, uint32 SizeRHI, EResourceLockMode LockMode) final override
diff --git a/Engine/Source/Runtime/OpenGLDrv/Private/OpenGLBuffer.cpp b/Engine/Source/Runtime/OpenGLDrv/Private/OpenGLBuffer.cpp
--- a/Engine/Source/Runtime/OpenGLDrv/Private/OpenGLBuffer.cpp
+++ b/Engine/Source/Runtime/OpenGLDrv/Private/OpenGLBuffer.cpp
@@ -8,6 +8,7 @@
 #include "Containers/ResourceArray.h"
 #include "HAL/IConsoleManager.h"
 #include "OpenGLDrv.h"
+#include "RHICoreBufferInitializer.h"
 
 namespace OpenGLConsoleVariables
 {
@@ -187,39 +188,66 @@ void FOpenGLBufferBase::OnBufferDeletion()
 	FOpenGLDynamicRHI::Get().OnBufferDeletion(Type, Resource);
 }
 
-FBufferRHIRef FOpenGLDynamicRHI::RHICreateBuffer(FRHICommandListBase& RHICmdList, FRHIBufferDesc const& Desc, ERHIAccess ResourceState, FRHIResourceCreateInfo& CreateInfo)
+static GLenum GetNewBufferType(const FRHIBufferCreateDesc& CreateDesc)
 {
-	const void* Data = nullptr;
 	GLenum BufferType = 0;
 
-	if (!Desc.IsNull())
+	if (!CreateDesc.IsNull())
 	{
-		// If a resource array was provided for the resource, create the resource pre-populated
-		if(CreateInfo.ResourceArray)
-		{
-			check(Desc.Size == CreateInfo.ResourceArray->GetResourceDataSize());
-			Data = CreateInfo.ResourceArray->GetResourceData();
-		}
-
 		BufferType = GL_ARRAY_BUFFER;
-		if (EnumHasAnyFlags(Desc.Usage, BUF_StructuredBuffer))
+		if (EnumHasAnyFlags(CreateDesc.Usage, BUF_StructuredBuffer))
 		{
 			BufferType = GL_SHADER_STORAGE_BUFFER;
 		}
-		else if (EnumHasAnyFlags(Desc.Usage, BUF_IndexBuffer))
+		else if (EnumHasAnyFlags(CreateDesc.Usage, BUF_IndexBuffer))
 		{
 			BufferType = GL_ELEMENT_ARRAY_BUFFER;
 		}
 	}
 
-	FOpenGLBuffer* Buffer = new FOpenGLBuffer(&RHICmdList, BufferType, Desc, Data);
+	return BufferType;
+}
+
+FRHIBufferInitializer FOpenGLDynamicRHI::RHICreateBufferInitializer(FRHICommandListBase& RHICmdList, const FRHIBufferCreateDesc& CreateDesc)
+{
+	const GLenum BufferType = GetNewBufferType(CreateDesc);
+
+	const void* InitialData = nullptr;
+	if (CreateDesc.InitAction == ERHIBufferInitAction::ResourceArray)
+	{
+		// TODO: move this initialization code path into this method
+		InitialData = CreateDesc.InitialData->GetResourceData();
+	}
+
+	FOpenGLBuffer* Buffer = new FOpenGLBuffer(&RHICmdList, BufferType, CreateDesc, InitialData);
+
+	if (CreateDesc.InitialData)
+	{
+		// Discard the resource array's contents.
+		CreateDesc.InitialData->Discard();
+	}
+
+	if (CreateDesc.IsNull() || CreateDesc.InitAction == ERHIBufferInitAction::ResourceArray || CreateDesc.InitAction == ERHIBufferInitAction::Default)
+	{
+		return UE::RHICore::FDefaultBufferInitializer(RHICmdList, Buffer);
+	}
+
+	if (CreateDesc.InitAction == ERHIBufferInitAction::Zeroed)
+	{
+		void* WritableData = RHICmdList.LockBuffer(Buffer, 0, CreateDesc.Size, RLM_WriteOnly);
+		FMemory::Memzero(WritableData, CreateDesc.Size);
+		RHICmdList.UnlockBuffer(Buffer);
+
+		return UE::RHICore::FDefaultBufferInitializer(RHICmdList, Buffer);
+	}
 
-	if (CreateInfo.ResourceArray)
+	if (CreateDesc.InitAction == ERHIBufferInitAction::Initializer)
 	{
-		CreateInfo.ResourceArray->Discard();
+		// Use LockBuffer + UnlockBuffer to allow the caller to write initial buffer data
+		return UE::RHICore::FLockBufferInitializer(RHICmdList, Buffer);
 	}
 
-	return Buffer;
+	return UE::RHICore::HandleUnknownBufferInitializerInitAction(RHICmdList, CreateDesc);
 }
 
 void* FOpenGLDynamicRHI::LockBuffer_BottomOfPipe(FRHICommandListBase& RHICmdList, FRHIBuffer* BufferRHI, uint32 Offset, uint32 Size, EResourceLockMode LockMode)
diff --git a/Engine/Source/Runtime/OpenGLDrv/Private/OpenGLCommands.cpp b/Engine/Source/Runtime/OpenGLDrv/Private/OpenGLCommands.cpp
--- a/Engine/Source/Runtime/OpenGLDrv/Private/OpenGLCommands.cpp
+++ b/Engine/Source/Runtime/OpenGLDrv/Private/OpenGLCommands.cpp
@@ -850,7 +850,11 @@ void FOpenGLDynamicRHI::UpdateSRV(FOpenGLShaderResourceView* SRV)
 		uint32 SizeY = Texture->GetSizeY();
 		
 		uint32 MipBytes = SizeX * SizeY;
-		TRefCountPtr<FOpenGLPixelBuffer> PixelBuffer = new FOpenGLPixelBuffer(nullptr, GL_PIXEL_UNPACK_BUFFER, FRHIBufferDesc(MipBytes, 0, BUF_Dynamic), nullptr);
+
+		const FRHIBufferCreateDesc CreateDesc =
+			FRHIBufferCreateDesc::Create(TEXT("PixelBuffer"), MipBytes, 0, EBufferUsageFlags::Dynamic);
+
+		TRefCountPtr<FOpenGLPixelBuffer> PixelBuffer = new FOpenGLPixelBuffer(nullptr, GL_PIXEL_UNPACK_BUFFER, CreateDesc, nullptr);
 		
 		glBindBuffer( GL_PIXEL_PACK_BUFFER, 0 );
 		glBindBuffer( GL_PIXEL_PACK_BUFFER, PixelBuffer->Resource );
diff --git a/Engine/Source/Runtime/OpenGLDrv/Private/OpenGLDrv.h b/Engine/Source/Runtime/OpenGLDrv/Private/OpenGLDrv.h
--- a/Engine/Source/Runtime/OpenGLDrv/Private/OpenGLDrv.h
+++ b/Engine/Source/Runtime/OpenGLDrv/Private/OpenGLDrv.h
@@ -332,7 +332,9 @@ public:
 
 	virtual FUniformBufferRHIRef RHICreateUniformBuffer(const void* Contents, const FRHIUniformBufferLayout* Layout, EUniformBufferUsage Usage, EUniformBufferValidation Validation) final override;
 	virtual void RHIUpdateUniformBuffer(FRHICommandListBase& RHICmdList, FRHIUniformBuffer* UniformBufferRHI, const void* Contents) final override;
-	virtual FBufferRHIRef RHICreateBuffer(FRHICommandListBase& RHICmdList, FRHIBufferDesc const& Desc, ERHIAccess ResourceState, FRHIResourceCreateInfo& CreateInfo) final override;
+
+	[[nodiscard]] virtual FRHIBufferInitializer RHICreateBufferInitializer(FRHICommandListBase& RHICmdList, const FRHIBufferCreateDesc& CreateDesc) final override;
+
 	virtual void RHIReplaceResources(FRHICommandListBase& RHICmdList, TArray<FRHIResourceReplaceInfo>&& ReplaceInfos) final override;
 	virtual void* LockBuffer_BottomOfPipe(FRHICommandListBase& RHICmdList, FRHIBuffer* Buffer, uint32 Offset, uint32 Size, EResourceLockMode LockMode) final override;
 	virtual void UnlockBuffer_BottomOfPipe(FRHICommandListBase& RHICmdList, FRHIBuffer* Buffer) final override;
diff --git a/Engine/Source/Runtime/OpenGLDrv/Private/OpenGLResources.h b/Engine/Source/Runtime/OpenGLDrv/Private/OpenGLResources.h
--- a/Engine/Source/Runtime/OpenGLDrv/Private/OpenGLResources.h
+++ b/Engine/Source/Runtime/OpenGLDrv/Private/OpenGLResources.h
@@ -145,15 +145,13 @@ class TOpenGLBuffer : public BaseType, public FOpenGLBufferBase
 	}
 
 public:
-	TOpenGLBuffer(FRHICommandListBase* RHICmdList, GLenum InType, FRHIBufferDesc const& BufferDesc, const void *InData)
-		: BaseType(BufferDesc)
+	TOpenGLBuffer(FRHICommandListBase* RHICmdList, GLenum InType, const FRHIBufferCreateDesc& CreateDesc, const void *InData)
+		: BaseType(CreateDesc)
 		, FOpenGLBufferBase(InType)
 	{
 		if (RHICmdList && RHICmdList->IsTopOfPipe() && InData)
 		{
-			void* LocalCopy = RHICmdList->Alloc(BaseType::GetSize(), 16);
-			FMemory::Memcpy(LocalCopy, InData, BaseType::GetSize());
-			InData = LocalCopy;
+			InData = RHICmdList->AllocCopy(InData, CreateDesc.Size, 16);
 		}
 
 		auto InitLambda = [this, InData]()
@@ -176,7 +174,7 @@ public:
 			OpenGLBufferStats::UpdateBufferStats(BaseType::GetDesc(), true);
 		};
 
-		if (!BufferDesc.IsNull())
+		if (!CreateDesc.IsNull())
 		{
 			if (RHICmdList)
 			{
@@ -525,8 +523,8 @@ private:
 class FOpenGLBasePixelBuffer : public FRefCountedObject
 {
 public:
-	FOpenGLBasePixelBuffer(FRHIBufferDesc const& BufferDesc)
-		: Desc(BufferDesc)
+	FOpenGLBasePixelBuffer(const FRHIBufferCreateDesc& CreateDesc)
+	: Desc(CreateDesc)
 	{}
 
 	const FRHIBufferDesc& GetDesc() const { return Desc; }
@@ -551,8 +549,8 @@ private:
 class FOpenGLBaseBuffer : public FRHIBuffer, public FOpenGLViewableResource
 {
 public:
-	FOpenGLBaseBuffer(FRHIBufferDesc const& BufferDesc)
-		: FRHIBuffer(BufferDesc)
+	FOpenGLBaseBuffer(const FRHIBufferCreateDesc& InCreateDesc)
+		: FRHIBuffer(InCreateDesc)
 	{
 		LLM_SCOPED_PAUSE_TRACKING_WITH_ENUM_AND_AMOUNT(ELLMTag::GraphicsPlatform, (int64)GetSize(), ELLMTracker::Platform, ELLMAllocType::None);
 		LLM_SCOPED_PAUSE_TRACKING_WITH_ENUM_AND_AMOUNT(ELLMTag::Meshes, (int64)GetSize(), ELLMTracker::Default, ELLMAllocType::None);
diff --git a/Engine/Source/Runtime/OpenGLDrv/Private/OpenGLTexture.cpp b/Engine/Source/Runtime/OpenGLDrv/Private/OpenGLTexture.cpp
--- a/Engine/Source/Runtime/OpenGLDrv/Private/OpenGLTexture.cpp
+++ b/Engine/Source/Runtime/OpenGLDrv/Private/OpenGLTexture.cpp
@@ -803,7 +803,10 @@ void FOpenGLDynamicRHI::ResolveTexture(FOpenGLTexture* Texture, uint32 MipIndex,
 	// Standard path with a PBO mirroring ever slice of a texture to allow multiple simulataneous maps
 	if (!IsValidRef(Texture->PixelBuffers[BufferIndex]))
 	{
-		Texture->PixelBuffers[BufferIndex] = new FOpenGLPixelBuffer(nullptr, GL_PIXEL_UNPACK_BUFFER, FRHIBufferDesc(MipBytes, 0, BUF_Dynamic), nullptr);
+		const FRHIBufferCreateDesc CreateDesc =
+			FRHIBufferCreateDesc::Create(TEXT("PixelBuffer"), MipBytes, 0, EBufferUsageFlags::Dynamic);
+
+		Texture->PixelBuffers[BufferIndex] = new FOpenGLPixelBuffer(nullptr, GL_PIXEL_UNPACK_BUFFER, CreateDesc, nullptr);
 	}
 	
 	TRefCountPtr<FOpenGLPixelBuffer> PixelBuffer = Texture->PixelBuffers[BufferIndex];
@@ -967,7 +970,9 @@ void* FOpenGLTexture::Lock(uint32 InMipIndex,uint32 ArrayIndex,EResourceLockMode
 		if (!IsValidRef(PixelBuffers[BufferIndex]))
 		{
 			bBufferExists = false;
-			PixelBuffers[BufferIndex] = new FOpenGLPixelBuffer(nullptr, GL_PIXEL_UNPACK_BUFFER, FRHIBufferDesc(MipBytes, 0, BUF_Dynamic), nullptr);
+			const FRHIBufferCreateDesc CreateDesc =
+				FRHIBufferCreateDesc::Create(TEXT("PixelBuffer"), MipBytes, 0, EBufferUsageFlags::Dynamic);
+			PixelBuffers[BufferIndex] = new FOpenGLPixelBuffer(nullptr, GL_PIXEL_UNPACK_BUFFER, CreateDesc, nullptr);
 		}
 
 		TRefCountPtr<FOpenGLPixelBuffer> PixelBuffer = PixelBuffers[BufferIndex];
diff --git a/Engine/Source/Runtime/RHI/Private/RHICommandList.cpp b/Engine/Source/Runtime/RHI/Private/RHICommandList.cpp
--- a/Engine/Source/Runtime/RHI/Private/RHICommandList.cpp
+++ b/Engine/Source/Runtime/RHI/Private/RHICommandList.cpp
@@ -4,6 +4,7 @@
 #include "Misc/App.h"
 #include "Async/TaskGraphInterfaces.h"
 #include "Containers/ConsumeAllMpmcQueue.h"
+#include "Containers/ResourceArray.h"
 #include "RHI.h"
 #include "ProfilingDebugging/CsvProfiler.h"
 #include "GenericPlatform/GenericPlatformCrashContext.h"
@@ -182,6 +183,22 @@ void FRHICommandListBase::AddDispatchPrerequisite(const FGraphEventRef& Prereq)
 
 void FRHICommandListBase::FinishRecording()
 {
+	if (PendingBufferUploads.Num() > 0)
+	{
+		FStringBuilderBase BufferList;
+
+		for (const FRHIBuffer* Buffer : PendingBufferUploads)
+		{
+			if (BufferList.Len() > 0)
+			{
+				BufferList.Append(TEXT(", "));
+			}
+			BufferList.Append(Buffer->GetName().ToString());
+		}
+
+		UE_LOG(LogRHI, Fatal, TEXT("Detected pending buffer uploads on RHICmdList submission: %s"), BufferList.ToString());
+	}
+
 	// "Complete" the dispatch event.
 	DispatchEvent->DispatchSubsequents();
 }
@@ -1751,9 +1768,12 @@ void FRHIComputeCommandList::BuildAccelerationStructure(FRHIRayTracingGeometry*
 	Params.BuildMode = EAccelerationStructureBuildMode::Build;
 
 	FRHIBufferRange ScratchBufferRange{};
-	
-	FRHIResourceCreateInfo ScratchBufferCreateInfo(TEXT("RHIScratchBuffer"));
-	ScratchBufferRange.Buffer = CreateBuffer(Geometry->GetSizeInfo().BuildScratchSize, BUF_StructuredBuffer | BUF_RayTracingScratch, 0, ERHIAccess::UAVCompute, ScratchBufferCreateInfo);
+
+	const FRHIBufferCreateDesc CreateDesc =
+		FRHIBufferCreateDesc::Create(TEXT("RHIScratchBuffer"), Geometry->GetSizeInfo().BuildScratchSize, 0, EBufferUsageFlags::StructuredBuffer | EBufferUsageFlags::RayTracingScratch)
+		.SetInitialState(ERHIAccess::UAVCompute);
+
+	ScratchBufferRange.Buffer = CreateBuffer(CreateDesc);
 
 	BuildAccelerationStructures(MakeArrayView(&Params, 1), ScratchBufferRange);
 }
@@ -1790,12 +1810,16 @@ void FRHIComputeCommandList::BuildAccelerationStructures(const TArrayView<const
 		// Allocate scratch buffer and build the acceleration structure for the selected sub-range
 		if (ParamCount > 0)
 		{
+			check(uint32(TotalRequiredScratchMemorySize) == TotalRequiredScratchMemorySize);
+
 			const TArrayView<const FRayTracingGeometryBuildParams> EffectiveParams = MakeArrayView<const FRayTracingGeometryBuildParams>(&Params[ParamBegin], ParamCount);
 
-			FRHIResourceCreateInfo ScratchBufferCreateInfo(TEXT("RHIScratchBuffer"));
+			const FRHIBufferCreateDesc CreateDesc =
+				FRHIBufferCreateDesc::Create(TEXT("RHIScratchBuffer"), TotalRequiredScratchMemorySize, 0, EBufferUsageFlags::StructuredBuffer | EBufferUsageFlags::RayTracingScratch)
+				.SetInitialState(ERHIAccess::UAVCompute);
+
 			FRHIBufferRange ScratchBufferRange{};
-			check(uint32(TotalRequiredScratchMemorySize) == TotalRequiredScratchMemorySize)
-			ScratchBufferRange.Buffer = CreateBuffer(TotalRequiredScratchMemorySize, BUF_StructuredBuffer | BUF_RayTracingScratch, 0, ERHIAccess::UAVCompute, ScratchBufferCreateInfo);
+			ScratchBufferRange.Buffer = CreateBuffer(CreateDesc);
 
 			BuildAccelerationStructures(EffectiveParams, ScratchBufferRange);
 		}
@@ -2094,9 +2118,13 @@ void FRHICommandList::GenerateMips(FRHITexture*)
 	UE_LOG(LogRHI, Fatal, TEXT("GenerateMips on RHI command lists is deprecated and no longer functions. Use the FGenerateMips helper class from the RenderCore module to generate mips on textures."));
 }
 
-void FRHICommandListBase::InvalidBufferFatalError(const FRHIResourceCreateInfo& CreateInfo, EBufferUsageFlags Usage, uint32 Stride)
+void FRHICommandListBase::InvalidBufferFatalError(const FRHIBufferCreateDesc& CreateDesc)
 {
-	UE_LOG(LogRHI, Fatal, TEXT("Attempt to create zero-sized buffer '%s', owner '%s', usage 0x%x, stride %u."), CreateInfo.DebugName ? CreateInfo.DebugName : TEXT("(nullptr)"), *CreateInfo.OwnerName.ToString(), (uint32)Usage, Stride);
+	UE_LOG(LogRHI, Fatal, TEXT("Attempt to create zero-sized buffer '%s', owner '%s', usage 0x%x, stride %u."),
+		   CreateDesc.DebugName ? CreateDesc.DebugName : TEXT("(nullptr)"),
+		   *CreateDesc.OwnerName.ToString(),
+		   static_cast<uint32>(CreateDesc.Usage), CreateDesc.Stride
+	);
 }
 
 void FRHICommandListBase::TransitionInternal(TConstArrayView<FRHITransitionInfo> Infos, ERHITransitionCreateFlags CreateFlags)
diff --git a/Engine/Source/Runtime/RHI/Public/DynamicRHI.h b/Engine/Source/Runtime/RHI/Public/DynamicRHI.h
--- a/Engine/Source/Runtime/RHI/Public/DynamicRHI.h
+++ b/Engine/Source/Runtime/RHI/Public/DynamicRHI.h
@@ -7,6 +7,7 @@ DynamicRHI.h: Dynamically bound Render Hardware Interface definitions.
 #pragma once
 
 #include "CoreTypes.h"
+#include "RHIBufferInitializer.h"
 #include "RHIContext.h"
 #include "RHIFeatureLevel.h"
 #include "RHIShaderLibrary.h"
@@ -28,6 +29,7 @@ class IRHITransientResourceAllocator;
 struct FDepthStencilStateInitializerRHI;
 struct FDisplayInformation;
 struct FRasterizerStateInitializerRHI;
+struct FRHIBufferCreateDesc;
 struct FRHIResourceCollectionMember;
 struct FRHIResourceCreateInfo;
 struct FRHIResourceInfo;
@@ -36,6 +38,7 @@ struct FSamplerStateInitializerRHI;
 struct FTextureMemoryStats;
 struct FRHIGPUMask;
 struct FUpdateTexture3DData;
+struct FResourceArrayUploadInterface;
 
 typedef TArray<FScreenResolutionRHI> FScreenResolutionArray;
 using FDisplayInformationArray = TArray<FDisplayInformation>;
@@ -435,10 +438,7 @@ public:
 	 */
 	virtual void RHIReplaceResources(FRHICommandListBase& RHICmdList, TArray<FRHIResourceReplaceInfo>&& ReplaceInfos) = 0;
 
-	/**
-	* @param ResourceArray - An optional pointer to a resource array containing the resource's data.
-	*/
-	virtual FBufferRHIRef RHICreateBuffer(FRHICommandListBase& RHICmdList, FRHIBufferDesc const& Desc, ERHIAccess ResourceState, FRHIResourceCreateInfo& CreateInfo) = 0;
+	[[nodiscard]] virtual FRHIBufferInitializer RHICreateBufferInitializer(FRHICommandListBase& RHICmdList, const FRHIBufferCreateDesc& CreateDesc) = 0;
 
 	RHI_API virtual void* RHILockBuffer(FRHICommandListBase& RHICmdList, FRHIBuffer* Buffer, uint32 Offset, uint32 Size, EResourceLockMode LockMode);
 	RHI_API virtual void* RHILockBufferMGPU(FRHICommandListBase& RHICmdList, FRHIBuffer* Buffer, uint32 GPUIndex, uint32 Offset, uint32 Size, EResourceLockMode LockMode);
@@ -1166,13 +1166,14 @@ FORCEINLINE FRHICalcTextureSizeResult RHICalcTexturePlatformSize(FRHITextureDesc
 	return GDynamicRHI->RHICalcTexturePlatformSize(Desc, FirstMipIndex);
 }
 
+PRAGMA_DISABLE_DEPRECATION_WARNINGS
+
 UE_DEPRECATED(5.5, "The separate RHICalcTexture... functions for each texture type are deprecated. Use RHICalcTexturePlatformSize instead.")
 FORCEINLINE uint64 RHICalcTexture2DPlatformSize(uint32 SizeX, uint32 SizeY, uint8 Format, uint32 NumMips, uint32 NumSamples, ETextureCreateFlags Flags, const FRHIResourceCreateInfo& CreateInfo, uint32& OutAlign)
 {
 	const uint16 Depth = 1;
 	const uint16 ArraySize = 1;
 
-PRAGMA_DISABLE_DEPRECATION_WARNINGS
 	FRHITextureDesc Desc(
 		ETextureDimension::Texture2D,
 		Flags,
@@ -1185,7 +1186,6 @@ PRAGMA_DISABLE_DEPRECATION_WARNINGS
 		(uint8)NumSamples,
 		CreateInfo.ExtData
 	);
-PRAGMA_ENABLE_DEPRECATION_WARNINGS
 
 	FRHICalcTextureSizeResult Result = RHICalcTexturePlatformSize(Desc, 0);
 	OutAlign = Result.Align;
@@ -1197,7 +1197,6 @@ FORCEINLINE uint64 RHICalcTexture2DArrayPlatformSize(uint32 SizeX, uint32 SizeY,
 {
 	const uint16 Depth = 1;
 
-PRAGMA_DISABLE_DEPRECATION_WARNINGS
 	FRHITextureDesc Desc(
 		ETextureDimension::Texture2DArray,
 		Flags,
@@ -1210,7 +1209,6 @@ PRAGMA_DISABLE_DEPRECATION_WARNINGS
 		(uint8)NumSamples,
 		CreateInfo.ExtData
 	);
-PRAGMA_ENABLE_DEPRECATION_WARNINGS
 
 	FRHICalcTextureSizeResult Result = RHICalcTexturePlatformSize(Desc, 0);
 	OutAlign = Result.Align;
@@ -1248,7 +1246,6 @@ FORCEINLINE uint64 RHICalcTexture3DPlatformSize(uint32 SizeX, uint32 SizeY, uint
 	const uint16 ArraySize  = 1;
 	const uint8 NumSamples = 1;
 
-PRAGMA_DISABLE_DEPRECATION_WARNINGS
 	FRHITextureDesc Desc(
 		ETextureDimension::Texture3D,
 		Flags,
@@ -1261,7 +1258,6 @@ PRAGMA_DISABLE_DEPRECATION_WARNINGS
 		NumSamples,
 		CreateInfo.ExtData
 	);
-PRAGMA_ENABLE_DEPRECATION_WARNINGS
 
 	FRHICalcTextureSizeResult Result = RHICalcTexturePlatformSize(Desc, 0);
 	OutAlign = Result.Align;
@@ -1275,7 +1271,6 @@ FORCEINLINE uint64 RHICalcTextureCubePlatformSize(uint32 Size, uint8 Format, uin
 	const uint16 ArraySize  = 1;
 	const uint8 NumSamples = 1;
 
-PRAGMA_DISABLE_DEPRECATION_WARNINGS
 	FRHITextureDesc Desc(
 		ETextureDimension::TextureCube,
 		Flags,
@@ -1288,12 +1283,12 @@ PRAGMA_DISABLE_DEPRECATION_WARNINGS
 		NumSamples,
 		CreateInfo.ExtData
 	);
-PRAGMA_ENABLE_DEPRECATION_WARNINGS
 
 	FRHICalcTextureSizeResult Result = RHICalcTexturePlatformSize(Desc, 0);
 	OutAlign = Result.Align;
 	return Result.Size;
 }
+PRAGMA_ENABLE_DEPRECATION_WARNINGS
 
 FORCEINLINE uint64 RHIGetMinimumAlignmentForBufferBackedSRV(EPixelFormat Format)
 {
diff --git a/Engine/Source/Runtime/RHI/Public/RHICommandList.h b/Engine/Source/Runtime/RHI/Public/RHICommandList.h
--- a/Engine/Source/Runtime/RHI/Public/RHICommandList.h
+++ b/Engine/Source/Runtime/RHI/Public/RHICommandList.h
@@ -715,7 +715,8 @@ public:
 
 private:
 	RHI_API void ActivatePipelines(ERHIPipeline Pipelines);
-	RHI_API void InvalidBufferFatalError(const FRHIResourceCreateInfo& CreateInfo, EBufferUsageFlags Usage, uint32 Stride);
+
+	RHI_API void InvalidBufferFatalError(const FRHIBufferCreateDesc& CreateDesc);
 
 public:
 	RHI_API void TransitionInternal(TConstArrayView<FRHITransitionInfo> Infos, ERHITransitionCreateFlags CreateFlags = ERHITransitionCreateFlags::None);
@@ -778,56 +779,106 @@ public:
 
 		GDynamicRHI->RHIUnlockBufferMGPU(*this, Buffer, GPUIndex);
 	}
-	
-	FORCEINLINE FBufferRHIRef CreateBuffer(uint32 Size, EBufferUsageFlags Usage, uint32 Stride, ERHIAccess ResourceState, FRHIResourceCreateInfo& CreateInfo)
+
+	[[nodiscard]]
+	FORCEINLINE_DEBUGGABLE FRHIBufferInitializer CreateBufferInitializer(const FRHIBufferCreateDesc& CreateDesc)
 	{
-		FRHIBufferDesc BufferDesc = CreateInfo.bWithoutNativeResource
-			? FRHIBufferDesc::Null()
-			: FRHIBufferDesc(Size, Stride, Usage);
+		FRHICommandListScopedPipelineGuard ScopedPipeline(*this);
+		return GDynamicRHI->RHICreateBufferInitializer(*this, CreateDesc);
+	}
 
-		if (Size == 0 && !BufferDesc.IsNull())
+	// Shortcut for creating a buffer without writing to an initializer
+	[[nodiscard]]
+	FORCEINLINE_DEBUGGABLE FBufferRHIRef CreateBuffer(const FRHIBufferCreateDesc& CreateDesc)
+	{
+		if (CreateDesc.Size == 0 && !CreateDesc.IsNull())
 		{
-			InvalidBufferFatalError(CreateInfo, Usage, Stride);
+			InvalidBufferFatalError(CreateDesc);
 		}
 
+		checkf(CreateDesc.InitAction != ERHIBufferInitAction::Initializer, TEXT("Buffer InitAction set to Initializer when calling CreateBuffer which doesn't write to its initializer"));
+
 		FRHICommandListScopedPipelineGuard ScopedPipeline(*this);
-		FBufferRHIRef Buffer = GDynamicRHI->RHICreateBuffer(*this, BufferDesc, ResourceState, CreateInfo);
-		Buffer->SetTrackedAccess_Unsafe(ResourceState);
-		return Buffer;
+		FRHIBufferInitializer Initializer = GDynamicRHI->RHICreateBufferInitializer(*this, CreateDesc);
+		return Initializer.Finalize();
+	}
+
+PRAGMA_DISABLE_DEPRECATION_WARNINGS
+	UE_DEPRECATED(5.6, "CreateBuffer without FRHIBufferCreateDesc is deprecated")
+	FORCEINLINE FBufferRHIRef CreateNullBuffer(ERHIAccess ResourceState, FRHIResourceCreateInfo& CreateInfo)
+	{
+		const FRHIBufferCreateDesc CreateDesc =
+			FRHIBufferCreateDesc::CreateNull(CreateInfo.DebugName)
+			.SetGPUMask(CreateInfo.GPUMask)
+			.SetInitialState(ResourceState)
+			.SetClassName(CreateInfo.ClassName)
+			.SetOwnerName(CreateInfo.OwnerName);
+
+		return CreateBuffer(CreateDesc);
 	}
 
+	UE_DEPRECATED(5.6, "CreateBuffer without FRHIBufferCreateDesc is deprecated")
+	FORCEINLINE_DEBUGGABLE FBufferRHIRef CreateBuffer(uint32 Size, EBufferUsageFlags Usage, uint32 Stride, ERHIAccess ResourceState, FRHIResourceCreateInfo& CreateInfo)
+	{
+		if (CreateInfo.bWithoutNativeResource)
+		{
+			return CreateNullBuffer(ResourceState, CreateInfo);
+		}
+
+		FRHIBufferCreateDesc CreateDesc =
+			FRHIBufferCreateDesc::Create(CreateInfo.DebugName, Size, Stride, Usage)
+			.SetGPUMask(CreateInfo.GPUMask)
+			.SetInitialState(ResourceState)
+			.SetClassName(CreateInfo.ClassName)
+			.SetOwnerName(CreateInfo.OwnerName);
+
+		if (CreateInfo.ResourceArray)
+		{
+			CreateDesc.SetInitActionResourceArray(CreateInfo.ResourceArray);
+		}
+
+		return CreateBuffer(CreateDesc);
+	}
+
+	UE_DEPRECATED(5.6, "CreateBuffer without FRHIBufferCreateDesc is deprecated")
 	FORCEINLINE FBufferRHIRef CreateVertexBuffer(uint32 Size, EBufferUsageFlags Usage, ERHIAccess ResourceState, FRHIResourceCreateInfo& CreateInfo)
 	{
 		return CreateBuffer(Size, Usage | EBufferUsageFlags::VertexBuffer, 0, ResourceState, CreateInfo);
 	}
 
+	UE_DEPRECATED(5.6, "CreateBuffer without FRHIBufferCreateDesc is deprecated")
 	FORCEINLINE FBufferRHIRef CreateVertexBuffer(uint32 Size, EBufferUsageFlags Usage, FRHIResourceCreateInfo& CreateInfo)
 	{
 		ERHIAccess ResourceState = RHIGetDefaultResourceState(Usage | EBufferUsageFlags::VertexBuffer, false);
 		return CreateVertexBuffer(Size, Usage, ResourceState, CreateInfo);
 	}
 
+	UE_DEPRECATED(5.6, "CreateBuffer without FRHIBufferCreateDesc is deprecated")
 	FORCEINLINE FBufferRHIRef CreateStructuredBuffer(uint32 Stride, uint32 Size, EBufferUsageFlags Usage, ERHIAccess ResourceState, FRHIResourceCreateInfo& CreateInfo)
 	{
 		return CreateBuffer(Size, Usage | EBufferUsageFlags::StructuredBuffer, Stride, ResourceState, CreateInfo);
 	}
 
+	UE_DEPRECATED(5.6, "CreateBuffer without FRHIBufferCreateDesc is deprecated")
 	FORCEINLINE FBufferRHIRef CreateStructuredBuffer(uint32 Stride, uint32 Size, EBufferUsageFlags Usage, FRHIResourceCreateInfo& CreateInfo)
 	{
 		ERHIAccess ResourceState = RHIGetDefaultResourceState(Usage | EBufferUsageFlags::StructuredBuffer, false);
 		return CreateStructuredBuffer(Stride, Size, Usage, ResourceState, CreateInfo);
 	}
 
+	UE_DEPRECATED(5.6, "CreateBuffer without FRHIBufferCreateDesc is deprecated")
 	FORCEINLINE FBufferRHIRef CreateIndexBuffer(uint32 Stride, uint32 Size, EBufferUsageFlags Usage, ERHIAccess ResourceState, FRHIResourceCreateInfo& CreateInfo)
 	{
 		return CreateBuffer(Size, Usage | EBufferUsageFlags::IndexBuffer, Stride, ResourceState, CreateInfo);
 	}
 
+	UE_DEPRECATED(5.6, "CreateBuffer without FRHIBufferCreateDesc is deprecated")
 	FORCEINLINE FBufferRHIRef CreateIndexBuffer(uint32 Stride, uint32 Size, EBufferUsageFlags Usage, FRHIResourceCreateInfo& CreateInfo)
 	{
 		ERHIAccess ResourceState = RHIGetDefaultResourceState(Usage | EBufferUsageFlags::IndexBuffer, false);
 		return CreateIndexBuffer(Stride, Size, Usage, ResourceState, CreateInfo);
 	}
+PRAGMA_ENABLE_DEPRECATION_WARNINGS
 
 	FORCEINLINE void UpdateUniformBuffer(FRHIUniformBuffer* UniformBufferRHI, const void* Contents)
 	{
@@ -1252,6 +1303,17 @@ protected:
 		PersistentState.SubpassIndex = 0;
 	}
 
+	void AddPendingBufferUpload(FRHIBuffer* InBuffer)
+	{
+		PendingBufferUploads.Emplace(InBuffer);
+	}
+
+	void RemovePendingBufferUpload(FRHIBuffer* InBuffer)
+	{
+		check(PendingBufferUploads.Contains(InBuffer));
+		PendingBufferUploads.Remove(InBuffer);
+	}
+
 protected:
 	FRHICommandBase*    Root            = nullptr;
 	FRHICommandBase**   CommandLink     = nullptr;
@@ -1419,6 +1481,8 @@ protected:
 
 	FRHIDrawStats DrawStats {};
 
+	TArray<FRHIBuffer*> PendingBufferUploads;
+
 public:
 #if WITH_RHI_BREADCRUMBS
 	friend FRHIBreadcrumbEventManual;
@@ -1480,6 +1544,7 @@ private:
 	friend class FRHICommandList_RecursiveHazardous;
 	friend class FRHIComputeCommandList_RecursiveHazardous;
 	friend struct FRHICommandSetGPUMask;
+	friend struct FRHIBufferInitializer;
 
 	template <typename RHICmdListType, typename LAMBDA>
 	friend struct TRHILambdaCommandMultiPipe;
@@ -1511,6 +1576,7 @@ private:
 	// Using FRHICmdListBaseArrayAllocator since elements Lifetime for all the following containers are 1 frame, between Record and Submit
 	// Track every resource state change happening during the command list Record
 	TMap<FRHIViewableResource*, TRHIResourceState<FRHICmdListBaseArrayAllocator>, FRHICmdListBaseSetAllocator> TrackedResources;
+
 	// Transitions with Unknown or Discard (when needed) that will patched at Dispatch time
 	TArray<FTransitionPatchInfo, FRHICmdListBaseArrayAllocator> TransitionInfoToPatch;
 	// Transitions with a single plane of DepthStencil specified where the other plane needs to be resolved with tracking
diff --git a/Engine/Source/Runtime/RHI/Public/RHIResourceUtils.h b/Engine/Source/Runtime/RHI/Public/RHIResourceUtils.h
--- a/Engine/Source/Runtime/RHI/Public/RHIResourceUtils.h
+++ b/Engine/Source/Runtime/RHI/Public/RHIResourceUtils.h
@@ -8,13 +8,46 @@
 
 namespace UE::RHIResourceUtils
 {
+	static FBufferRHIRef CreateBufferZeroed(FRHICommandListBase& RHICmdList, const FRHIBufferCreateDesc& Desc)
+	{
+		return RHICmdList.CreateBuffer(FRHIBufferCreateDesc(Desc).SetInitActionZeroData());
+	}
+
+	template<typename TElementType>
+	static FBufferRHIRef CreateBufferWithArray(FRHICommandListBase& RHICmdList, const FRHIBufferCreateDesc& Desc, TConstArrayView<TElementType> Array)
+	{
+		FResourceArrayUploadArrayView UploadView(Array);
+		return RHICmdList.CreateBuffer(FRHIBufferCreateDesc(Desc).SetInitActionResourceArray(&UploadView));
+	}
+
+	template<typename TElementType>
+	static FBufferRHIRef CreateBufferWithArray(FRHICommandListBase& RHICmdList, const FRHIBufferCreateDesc& Desc, const TArray<TElementType>& Array)
+	{
+		return CreateBufferWithArray(RHICmdList, Desc, TConstArrayView<TElementType>(Array));
+	}
+
+	template<typename TElementType, size_t TCount>
+	static FBufferRHIRef CreateBufferWithArray(FRHICommandListBase& RHICmdList, const FRHIBufferCreateDesc& Desc, const TElementType (&Array)[TCount])
+	{
+		return CreateBufferWithArray(RHICmdList, Desc, TConstArrayView<TElementType>(Array, TCount));
+	}
+
+	template<typename TElementType>
+	static FBufferRHIRef CreateBufferWithValue(FRHICommandListBase& RHICmdList, const FRHIBufferCreateDesc& Desc, const TElementType& Value)
+	{
+		return CreateBufferWithArray(RHICmdList, Desc, TConstArrayView<TElementType>(&Value, 1));
+	}
+
 	static FBufferRHIRef CreateBufferFromArray(FRHICommandListBase& RHICmdList, const TCHAR* Name, EBufferUsageFlags UsageFlags, uint32 InStride, const void* InData, uint32 InSizeInBytes)
 	{
-		const ERHIAccess InitialState = RHIGetDefaultResourceState(UsageFlags, false);
 		FResourceArrayUploadArrayView UploadView(InData, InSizeInBytes);
 
-		FRHIResourceCreateInfo CreateInfo(Name, &UploadView);
-		return RHICmdList.CreateBuffer(InSizeInBytes, UsageFlags, InStride, InitialState, CreateInfo);
+		const FRHIBufferCreateDesc CreateDesc =
+			FRHIBufferCreateDesc::Create(Name, InSizeInBytes, InStride, UsageFlags)
+			.DetermineInitialState()
+			.SetInitActionResourceArray(&UploadView);
+
+		return RHICmdList.CreateBuffer(CreateDesc);
 	}
 
 	template<typename TElementType>
@@ -22,8 +55,12 @@ namespace UE::RHIResourceUtils
 	{
 		FResourceArrayUploadArrayView UploadView(Array);
 
-		FRHIResourceCreateInfo CreateInfo(Name, &UploadView);
-		return RHICmdList.CreateBuffer(UploadView.GetResourceDataSize(), UsageFlags, InStride, InitialState, CreateInfo);
+		const FRHIBufferCreateDesc CreateDesc =
+			FRHIBufferCreateDesc::Create(Name, UploadView.GetResourceDataSize(), InStride, UsageFlags)
+			.SetInitialState(InitialState)
+			.SetInitActionResourceArray(&UploadView);
+
+		return RHICmdList.CreateBuffer(CreateDesc);
 	}
 
 	template<typename TElementType>
@@ -32,6 +69,12 @@ namespace UE::RHIResourceUtils
 		return CreateBufferFromArray(RHICmdList, Name, UsageFlags, Array.GetTypeSize(), InitialState, Array);
 	}
 
+	template<typename TElementType>
+	static FBufferRHIRef CreateBufferFromArray(FRHICommandListBase& RHICmdList, const TCHAR* Name, EBufferUsageFlags UsageFlags, TConstArrayView<TElementType> Array)
+	{
+		return CreateBufferFromArray(RHICmdList, Name, UsageFlags, RHIGetDefaultResourceState(UsageFlags, false), Array);
+	}
+
 	template<typename TElementType>
 	static FBufferRHIRef CreateVertexBufferFromArray(FRHICommandListBase& RHICmdList, const TCHAR* Name, EBufferUsageFlags ExtraFlags, TConstArrayView<TElementType> Array)
 	{
diff --git a/Engine/Source/Runtime/RHI/Public/RHIResources.h b/Engine/Source/Runtime/RHI/Public/RHIResources.h
--- a/Engine/Source/Runtime/RHI/Public/RHIResources.h
+++ b/Engine/Source/Runtime/RHI/Public/RHIResources.h
@@ -107,7 +107,7 @@ public:
 
 	inline ERHIResourceType GetType() const { return ResourceType; }
 
-	FName GetOwnerName() const
+	FORCEINLINE FName GetOwnerName() const
 	{
 #if RHI_ENABLE_RESOURCE_INFO
 		return OwnerName;
@@ -115,7 +115,8 @@ public:
 		return NAME_None;
 #endif
 	}
-	void SetOwnerName(FName InOwnerName)
+
+	FORCEINLINE void SetOwnerName(FName InOwnerName)
 	{
 #if RHI_ENABLE_RESOURCE_INFO
 		OwnerName = InOwnerName;
@@ -364,8 +365,11 @@ struct FClearValueBinding
 	static RHI_API const FClearValueBinding DefaultNormal8Bit;
 };
 
-struct FRHIResourceCreateInfo
+struct UE_DEPRECATED(5.6, "FRHIResourceCreateInfo is no longer used. Please use FRHIBufferCreateDesc or FRHITextureCreateDesc.") FRHIResourceCreateInfo
 {
+PRAGMA_DISABLE_DEPRECATION_WARNINGS
+	FRHIResourceCreateInfo() = default;
+
 	FRHIResourceCreateInfo(const TCHAR* InDebugName)
 		: DebugName(InDebugName)
 	{
@@ -373,13 +377,10 @@ struct FRHIResourceCreateInfo
 	}
 
 	// for CreateTexture calls
-	UE_DEPRECATED(5.5, "Please use FRHITextureCreateDesc for creating Textures with Bulk Data")
 	FRHIResourceCreateInfo(const TCHAR* InDebugName, FResourceBulkDataInterface* InBulkData)
 		: FRHIResourceCreateInfo(InDebugName)
 	{
-PRAGMA_DISABLE_DEPRECATION_WARNINGS
 		BulkData = InBulkData;
-PRAGMA_ENABLE_DEPRECATION_WARNINGS
 	}
 
 	// for CreateBuffer calls
@@ -389,23 +390,18 @@ PRAGMA_ENABLE_DEPRECATION_WARNINGS
 		ResourceArray = InResourceArray;
 	}
 
-	UE_DEPRECATED(5.5, "Please use FRHITextureCreateDesc for creating Textures with a Clear Value Binding")
 	FRHIResourceCreateInfo(const TCHAR* InDebugName, const FClearValueBinding& InClearValueBinding)
 		: FRHIResourceCreateInfo(InDebugName)
 	{
-PRAGMA_DISABLE_DEPRECATION_WARNINGS
 		ClearValueBinding = InClearValueBinding;
-PRAGMA_ENABLE_DEPRECATION_WARNINGS
 	}
 
-	UE_DEPRECATED(5.5, "Please use FRHITextureCreateDesc for creating Textures with Ext Data")
 	FRHIResourceCreateInfo(uint32 InExtData)
 		: FRHIResourceCreateInfo(TEXT(""))
 	{
-PRAGMA_DISABLE_DEPRECATION_WARNINGS
 		ExtData = InExtData;
-PRAGMA_ENABLE_DEPRECATION_WARNINGS
 	}
+PRAGMA_ENABLE_DEPRECATION_WARNINGS
 
 	FName GetTraceClassName() const									{ const static FLazyName FRHIBufferName(TEXT("FRHIBuffer")); return (ClassName == NAME_None) ? FRHIBufferName : ClassName; }
 
@@ -414,6 +410,7 @@ PRAGMA_ENABLE_DEPRECATION_WARNINGS
 	FResourceBulkDataInterface* BulkData = nullptr;
 
 	// for CreateBuffer calls
+	UE_DEPRECATED(5.6, "Please use FRHIBufferCreateDesc for creating buffers with Resource Arrays")
 	FResourceArrayUploadInterface* ResourceArray = nullptr;
 
 	// for binding clear colors to render targets.
@@ -1320,12 +1317,6 @@ private:
 class FRHIViewableResource : public FRHIResource
 {
 public:
-	// TODO (RemoveUnknowns) remove once FRHIBufferCreateDesc contains initial access.
-	void SetTrackedAccess_Unsafe(ERHIAccess Access)
-	{
-		TrackedAccess = Access;
-	}
-
 	ERHIAccess GetTrackedAccess_Unsafe() const
 	{
 		return TrackedAccess.GetResourceState();
@@ -1355,7 +1346,8 @@ protected:
 	FRHIViewableResource(ERHIResourceType InResourceType, ERHIAccess InAccess, uint16 SubResourceCount)
 		: FRHIResource(InResourceType)
 		, TrackedAccess(InAccess, SubResourceCount)
-	{}
+	{
+	}
 
 	virtual void SetTrackedAccessFromContext(ERHIAccess Access)
 	{
@@ -1384,33 +1376,259 @@ private:
 
 struct FRHIBufferDesc
 {
-	uint32 Size{};
-	uint32 Stride{};
-	EBufferUsageFlags Usage{};
+	/** Size of the buffer memory, in bytes */
+	uint32 Size = 0;
+
+	/** Distance in bytes between elements of the buffer */
+	uint32 Stride = 0;
+
+	/** Flags describing the usage of the buffer */
+	EBufferUsageFlags Usage = EBufferUsageFlags::None;
+
+	/* A mask representing which GPUs to create the resource on, in a multi-GPU system. */
+	FRHIGPUMask GPUMask = FRHIGPUMask::All();
 
 	FRHIBufferDesc() = default;
 	FRHIBufferDesc(uint32 InSize, uint32 InStride, EBufferUsageFlags InUsage)
 		: Size  (InSize)
 		, Stride(InStride)
 		, Usage (InUsage)
-	{}
+	{
+	}
+	FRHIBufferDesc(uint32 InSize, uint32 InStride, EBufferUsageFlags InUsage, FRHIGPUMask InGPUMask)
+		: Size   (InSize)
+		, Stride (InStride)
+		, Usage  (InUsage)
+		, GPUMask(InGPUMask)
+	{
+	}
 
 	static FRHIBufferDesc Null()
 	{
-		return FRHIBufferDesc(0, 0, BUF_NullResource);
+		return FRHIBufferDesc(0, 0, EBufferUsageFlags::NullResource);
 	}
 
 	bool IsNull() const
 	{
-		if (EnumHasAnyFlags(Usage, BUF_NullResource))
+		if (EnumHasAnyFlags(Usage, EBufferUsageFlags::NullResource))
 		{
 			// The null resource descriptor should have its other fields zeroed, and no additional flags.
-			check(Size == 0 && Stride == 0 && Usage == BUF_NullResource);
+			check(Size == 0 && Stride == 0 && Usage == EBufferUsageFlags::NullResource);
 			return true;
 		}
 
 		return false;
 	}
+
+	friend uint32 GetTypeHash(const FRHIBufferDesc& Desc)
+	{
+		uint32 Hash = GetTypeHash(Desc.Size);
+		Hash = HashCombine(Hash, GetTypeHash(Desc.Stride));
+		Hash = HashCombine(Hash, GetTypeHash(Desc.Usage));
+		Hash = HashCombine(Hash, GetTypeHash(Desc.GPUMask.GetNative()));
+		return Hash;
+	}
+
+	bool operator==(const FRHIBufferDesc& Other) const
+	{
+		return Size == Other.Size
+			&& Stride == Other.Stride
+			&& Usage == Other.Usage
+			&& GPUMask == Other.GPUMask;
+	}
+
+	bool operator!=(const FRHIBufferDesc& Other) const
+	{
+		return !(*this == Other);
+	}
+
+	FRHIBufferDesc& operator=(const FRHIBufferDesc& Other)
+	{
+		Size    = Other.Size;
+		Stride  = Other.Stride;
+		Usage   = Other.Usage;
+		GPUMask = Other.GPUMask;
+
+		return *this;
+	}
+};
+
+RHI_API ERHIAccess RHIGetDefaultResourceState(EBufferUsageFlags InUsage, bool bInHasInitialData);
+
+enum class ERHIBufferInitAction
+{
+	// Default for the RHI, data can be "undefined"
+	Default,
+
+	// Zero all buffer data
+	Zeroed,
+
+	// Upload data from a provided FResourceArrayUploadInterface. This data will be discarded after it's used.
+	ResourceArray,
+
+	// Caller will use FRHIBufferInitializer to set the initial buffer contents.
+	Initializer,
+};
+
+struct FRHIBufferCreateDesc : public FRHIBufferDesc
+{
+	static FRHIBufferCreateDesc Create(const TCHAR* InDebugName, EBufferUsageFlags InUsage)
+	{
+		return FRHIBufferCreateDesc(InDebugName, InUsage);
+	}
+
+	static FRHIBufferCreateDesc Create(const TCHAR* InDebugName, uint32 InSize, uint32 InStride, EBufferUsageFlags InUsage)
+	{
+		return FRHIBufferCreateDesc(InDebugName, InSize, InStride, InUsage);
+	}
+	
+	static FRHIBufferCreateDesc Create(const TCHAR* InDebugName, const FRHIBufferDesc& InDesc)
+	{
+		return FRHIBufferCreateDesc(InDebugName, InDesc);
+	}
+
+	static FRHIBufferCreateDesc CreateNull(const TCHAR* InDebugName)
+	{
+		return Create(InDebugName, 0, 0, EBufferUsageFlags::NullResource);
+	}
+
+	static FRHIBufferCreateDesc CreateVertex(const TCHAR* InDebugName)
+	{
+		return Create(InDebugName, EBufferUsageFlags::VertexBuffer);
+	}
+
+	static FRHIBufferCreateDesc CreateVertex(const TCHAR* InDebugName, uint32 InSize)
+	{
+		return Create(InDebugName, InSize, 0, EBufferUsageFlags::VertexBuffer);
+	}
+
+	template<typename TVertexType>
+	static FRHIBufferCreateDesc CreateVertex(const TCHAR* InDebugName, uint32 InCount)
+	{
+		return Create(InDebugName, InCount * sizeof(TVertexType), 0, EBufferUsageFlags::VertexBuffer);
+	}
+
+	static FRHIBufferCreateDesc CreateIndex(const TCHAR* InDebugName)
+	{
+		return Create(InDebugName, EBufferUsageFlags::IndexBuffer);
+	}
+
+	static FRHIBufferCreateDesc CreateIndex(const TCHAR* InDebugName, uint32 InSize, uint32 InStride)
+	{
+		return Create(InDebugName, InSize, InStride, EBufferUsageFlags::IndexBuffer);
+	}
+
+	template<typename TIndexType>
+	static FRHIBufferCreateDesc CreateIndex(const TCHAR* InDebugName, uint32 InCount)
+	{
+		return Create(InDebugName, InCount * sizeof(TIndexType), sizeof(TIndexType), EBufferUsageFlags::IndexBuffer);
+	}
+
+	static FRHIBufferCreateDesc CreateStructured(const TCHAR* InDebugName)
+	{
+		return Create(InDebugName, EBufferUsageFlags::StructuredBuffer);
+	}
+
+	static FRHIBufferCreateDesc CreateStructured(const TCHAR* InDebugName, uint32 InSize, uint32 InStride)
+	{
+		return Create(InDebugName, InSize, InStride, EBufferUsageFlags::StructuredBuffer);
+	}
+
+	template<typename TStructureType>
+	static FRHIBufferCreateDesc CreateStructured(const TCHAR* InDebugName, uint32 InCount)
+	{
+		return Create(InDebugName, InCount * sizeof(TStructureType), sizeof(TStructureType), EBufferUsageFlags::StructuredBuffer);
+	}
+
+	static FRHIBufferCreateDesc CreateByteAddress(const TCHAR* InDebugName)
+	{
+		return Create(InDebugName, EBufferUsageFlags::ByteAddressBuffer);
+	}
+
+	static FRHIBufferCreateDesc CreateByteAddress(const TCHAR* InDebugName, uint32 InSize, uint32 InStride)
+	{
+		return Create(InDebugName, InSize, InStride, EBufferUsageFlags::ByteAddressBuffer);
+	}
+
+	static FRHIBufferCreateDesc CreateUniform(const TCHAR* InDebugName)
+	{
+		return Create(InDebugName, EBufferUsageFlags::UniformBuffer);
+	}
+
+	FRHIBufferCreateDesc() = default;
+
+	FRHIBufferCreateDesc(const TCHAR* InDebugName, EBufferUsageFlags InUsage)
+		: DebugName(InDebugName)
+	{
+		Usage = InUsage;
+	}
+
+	FRHIBufferCreateDesc(const TCHAR* InDebugName, uint32 InSize, uint32 InStride, EBufferUsageFlags InUsage)
+		: FRHIBufferDesc(InSize, InStride, InUsage)
+		, DebugName(InDebugName)
+	{
+	}
+
+	FRHIBufferCreateDesc(const TCHAR* InDebugName, const FRHIBufferDesc& InOtherDesc)
+		: FRHIBufferDesc(InOtherDesc)
+		, DebugName(InDebugName)
+	{
+	}
+
+	FRHIBufferCreateDesc& SetSize(uint32 InSize)              { Size   = InSize;     return *this; }
+	FRHIBufferCreateDesc& SetStride(uint32 InStride)          { Stride = InStride;   return *this; }
+	FRHIBufferCreateDesc& SetUsage(EBufferUsageFlags InUsage) { Usage  = InUsage;    return *this; }
+	FRHIBufferCreateDesc& AddUsage(EBufferUsageFlags InUsage) { Usage |= InUsage;    return *this; }
+	FRHIBufferCreateDesc& SetGPUMask(FRHIGPUMask InGPUMask)   { GPUMask = InGPUMask; return *this; }
+
+	FRHIBufferCreateDesc& SetDebugName(const TCHAR* InDebugName)           { DebugName = InDebugName;       return *this; }
+	FRHIBufferCreateDesc& SetInitialState(ERHIAccess InInitialState)       { InitialState = InInitialState; return *this; }
+	FRHIBufferCreateDesc& DetermineInitialState()                          { if (InitialState == ERHIAccess::Unknown) InitialState = RHIGetDefaultResourceState(Usage, false); return *this; }
+	FRHIBufferCreateDesc& SetInitAction(ERHIBufferInitAction InInitAction) { InitAction = InInitAction;     return *this; }
+	FRHIBufferCreateDesc& SetClassName(const FName& InClassName)           { ClassName = InClassName;       return *this; }
+	FRHIBufferCreateDesc& SetOwnerName(const FName& InOwnerName)           { OwnerName = InOwnerName;       return *this; }
+
+	FRHIBufferCreateDesc& SetInitActionNone()
+	{
+		return SetInitAction(ERHIBufferInitAction::Default);
+	}
+	FRHIBufferCreateDesc& SetInitActionZeroData()
+	{
+		return SetInitAction(ERHIBufferInitAction::Zeroed);
+	}
+	FRHIBufferCreateDesc& SetInitActionInitializer()
+	{
+		return SetInitAction(ERHIBufferInitAction::Initializer);
+	}
+	FRHIBufferCreateDesc& SetInitActionResourceArray(FResourceArrayUploadInterface* InInitialData)
+	{
+		InitialData = InInitialData;
+		return SetInitAction(ERHIBufferInitAction::ResourceArray);
+	}
+
+	FName GetTraceClassName() const
+	{
+		const static FLazyName FRHIBufferName(TEXT("FRHIBuffer"));
+		return (ClassName == NAME_None) ? FRHIBufferName : ClassName;
+	}
+
+	/* A friendly name for the resource. */
+	const TCHAR* DebugName = nullptr;
+
+	/* Provider of initial data for the buffer. InitialData->Discard() will be called after the data is read. */
+	FResourceArrayUploadInterface* InitialData = nullptr;
+
+	/* The RHI access state that the resource will be created in. */
+	ERHIAccess InitialState = ERHIAccess::Unknown;
+
+	/* Tells how to initialize (or not) the buffer's data. */
+	ERHIBufferInitAction InitAction = ERHIBufferInitAction::Default;
+
+	/* The owner class of FRHIBuffer used for Insight asset metadata tracing */
+	FName ClassName = NAME_None;
+
+	/* The owner name used for Insight asset metadata tracing */
+	FName OwnerName = NAME_None;
 };
 
 class FRHIBuffer : public FRHIViewableResource
@@ -1420,25 +1638,47 @@ class FRHIBuffer : public FRHIViewableResource
 {
 public:
 	/** Initialization constructor. */
-	FRHIBuffer(FRHIBufferDesc const& InDesc)
-		: FRHIViewableResource(RRT_Buffer, ERHIAccess::Unknown /* TODO (RemoveUnknowns): Use InitialAccess from descriptor after refactor. */,  1)
-		, Desc(InDesc)
-	{}
+	FRHIBuffer(const FRHIBufferCreateDesc& InCreateDesc)
+		: FRHIViewableResource(RRT_Buffer, InCreateDesc.InitialState, 1)
+		, Desc(InCreateDesc)
+	{
+#if RHI_ENABLE_RESOURCE_INFO
+		if (InCreateDesc.DebugName)
+		{
+			Name = InCreateDesc.DebugName;
+		}
 
-	FRHIBufferDesc const& GetDesc() const { return Desc; }
+		SetOwnerName(InCreateDesc.OwnerName);
+#endif
+	}
+
+	const FRHIBufferDesc& GetDesc() const
+	{
+		return Desc;
+	}
 
 	/** @return The number of bytes in the buffer. */
-	uint32 GetSize() const { return Desc.Size; }
+	uint32 GetSize() const
+	{
+		return Desc.Size;
+	}
 
 	/** @return The stride in bytes of the buffer. */
-	uint32 GetStride() const { return Desc.Stride; }
+	uint32 GetStride() const
+	{
+		return Desc.Stride;
+	}
 
 	/** @return The usage flags used to create the buffer. */
-	EBufferUsageFlags GetUsage() const { return Desc.Usage; }
-
-	void SetName(const FName& InName) { Name = InName; }
+	EBufferUsageFlags GetUsage() const
+	{
+		return Desc.Usage;
+	}
 
-	virtual uint32 GetParentGPUIndex() const { return 0; }
+	void SetName(const FName& InName)
+	{
+		Name = InName;
+	}
 
 #if ENABLE_RHI_VALIDATION
 	virtual RHIValidation::FResource* GetValidationTrackerResource() final override
@@ -1448,12 +1688,6 @@ public:
 #endif
 
 protected:
-	// Used by RHI implementations that may adjust internal usage flags during object construction.
-	void SetUsage(EBufferUsageFlags InUsage)
-	{
-		Desc.Usage = InUsage;
-	}
-
 	void TakeOwnership(FRHIBuffer& Other)
 	{
 		FRHIViewableResource::TakeOwnership(Other);
@@ -5344,30 +5578,7 @@ public:
 };
 
 /** Descriptor used to create a buffer resource */
-struct FRHIBufferCreateInfo
-{
-	bool operator == (const FRHIBufferCreateInfo& Other) const
-	{
-		return (
-			Size == Other.Size &&
-			Stride == Other.Stride &&
-			Usage == Other.Usage);
-	}
-
-	bool operator != (const FRHIBufferCreateInfo& Other) const
-	{
-		return !(*this == Other);
-	}
-
-	/** Total size of the buffer. */
-	uint32 Size = 1;
-
-	/** Stride in bytes */
-	uint32 Stride = 1;
-
-	/** Bitfields describing the uses of that buffer. */
-	EBufferUsageFlags Usage = BUF_None;
-};
+using FRHIBufferCreateInfo = FRHIBufferDesc;
 
 struct FRHIBufferSRVCreateInfo
 {
diff --git a/Engine/Source/Runtime/RHI/Public/RHIUtilities.h b/Engine/Source/Runtime/RHI/Public/RHIUtilities.h
--- a/Engine/Source/Runtime/RHI/Public/RHIUtilities.h
+++ b/Engine/Source/Runtime/RHI/Public/RHIUtilities.h
@@ -227,13 +227,22 @@ struct FRWBuffer
 	{
 		// Provide a debug name if using Fast VRAM so the allocators diagnostics will work
 		ensure(!(EnumHasAnyFlags(AdditionalUsage, BUF_FastVRAM) && !InDebugName));
+
 		NumBytes = BytesPerElement * NumElements;
-		FRHIResourceCreateInfo CreateInfo(InDebugName);
-		CreateInfo.ClassName = ClassName;
-		CreateInfo.OwnerName = OwnerName;
-		CreateInfo.ResourceArray = InResourceArray;
-		Buffer = RHICmdList.CreateBuffer(NumBytes, BUF_VertexBuffer | BUF_UnorderedAccess | BUF_ShaderResource | AdditionalUsage, 0, InResourceState, CreateInfo);
-		
+
+		FRHIBufferCreateDesc CreateDesc =
+			FRHIBufferCreateDesc::CreateVertex(InDebugName, NumBytes)
+			.AddUsage(EBufferUsageFlags::UnorderedAccess | EBufferUsageFlags::ShaderResource | AdditionalUsage)
+			.SetInitialState(InResourceState)
+			.SetClassName(ClassName)
+			.SetOwnerName(OwnerName);
+
+		if (InResourceArray)
+		{
+			CreateDesc.SetInitActionResourceArray(InResourceArray);
+		}
+
+		Buffer = RHICmdList.CreateBuffer(CreateDesc);
 		UAV = RHICmdList.CreateUnorderedAccessView(
 			Buffer,
 			FRHIViewDesc::CreateBufferUAV()
@@ -252,6 +261,38 @@ struct FRWBuffer
 		Initialize(RHICmdList, InDebugName, BytesPerElement, NumElements, Format, ERHIAccess::UAVCompute, AdditionalUsage, InResourceArray);
 	}
 
+	template<typename LAMBDA>
+	void InitializeWithData(FRHICommandListBase& RHICmdList, const TCHAR* InDebugName, uint32 BytesPerElement, uint32 NumElements, EPixelFormat Format, ERHIAccess InResourceState, EBufferUsageFlags AdditionalUsage, LAMBDA&& InitializerFunction)
+	{
+		// Provide a debug name if using Fast VRAM so the allocators diagnostics will work
+		ensure(!(EnumHasAnyFlags(AdditionalUsage, BUF_FastVRAM) && !InDebugName));
+
+		NumBytes = BytesPerElement * NumElements;
+
+		const FRHIBufferCreateDesc CreateDesc =
+			FRHIBufferCreateDesc::CreateVertex(InDebugName, NumBytes)
+			.AddUsage(EBufferUsageFlags::UnorderedAccess | EBufferUsageFlags::ShaderResource | AdditionalUsage)
+			.SetInitialState(InResourceState)
+			.SetClassName(ClassName)
+			.SetOwnerName(OwnerName)
+			.SetInitActionInitializer();
+
+		FRHIBufferInitializer Initializer = RHICmdList.CreateBufferInitializer(CreateDesc);
+		InitializerFunction(Initializer);
+		Buffer = Initializer.Finalize();
+
+		UAV = RHICmdList.CreateUnorderedAccessView(
+			Buffer,
+			FRHIViewDesc::CreateBufferUAV()
+				.SetType(FRHIViewDesc::EBufferType::Typed)
+				.SetFormat(Format));
+		SRV = RHICmdList.CreateShaderResourceView(
+			Buffer,
+			FRHIViewDesc::CreateBufferSRV()
+			.SetType(FRHIViewDesc::EBufferType::Typed)
+			.SetFormat(Format));
+	}
+
 	void Release()
 	{
 		NumBytes = 0;
@@ -313,9 +354,40 @@ struct FReadBuffer
 	void Initialize(FRHICommandListBase& RHICmdList, const TCHAR* InDebugName, uint32 BytesPerElement, uint32 NumElements, EPixelFormat Format, EBufferUsageFlags AdditionalUsage = BUF_None, FResourceArrayUploadInterface* InResourceArray = nullptr)
 	{
 		NumBytes = BytesPerElement * NumElements;
-		FRHIResourceCreateInfo CreateInfo(InDebugName);
-		CreateInfo.ResourceArray = InResourceArray;
-		Buffer = RHICmdList.CreateVertexBuffer(NumBytes, BUF_ShaderResource | AdditionalUsage, ERHIAccess::SRVMask, CreateInfo);
+
+		FRHIBufferCreateDesc CreateDesc =
+			FRHIBufferCreateDesc::CreateVertex(InDebugName, NumBytes)
+			.AddUsage(EBufferUsageFlags::ShaderResource | AdditionalUsage)
+			.SetInitialState(ERHIAccess::SRVMask);
+
+		if (InResourceArray)
+		{
+			CreateDesc.SetInitActionResourceArray(InResourceArray);
+		}
+
+		Buffer = RHICmdList.CreateBuffer(CreateDesc);
+		SRV = RHICmdList.CreateShaderResourceView(
+			Buffer, 
+			FRHIViewDesc::CreateBufferSRV()
+				.SetType(FRHIViewDesc::EBufferType::Typed)
+				.SetFormat(Format));
+	}
+
+	template<typename LAMBDA>
+	void InitializeWithData(FRHICommandListBase& RHICmdList, const TCHAR* InDebugName, uint32 BytesPerElement, uint32 NumElements, EPixelFormat Format, EBufferUsageFlags AdditionalUsage, LAMBDA&& InitializerFunction)
+	{
+		NumBytes = BytesPerElement * NumElements;
+
+		const FRHIBufferCreateDesc CreateDesc =
+			FRHIBufferCreateDesc::CreateVertex(InDebugName, NumBytes)
+			.AddUsage(EBufferUsageFlags::ShaderResource | AdditionalUsage)
+			.SetInitialState(ERHIAccess::SRVMask)
+			.SetInitActionInitializer();
+
+		FRHIBufferInitializer Initializer = RHICmdList.CreateBufferInitializer(CreateDesc);
+		InitializerFunction(Initializer);
+		Buffer = Initializer.Finalize();
+
 
 		SRV = RHICmdList.CreateShaderResourceView(
 			Buffer, 
@@ -347,22 +419,52 @@ struct FRWBufferStructured
 		Release();
 	}
 
-	void Initialize(FRHICommandListBase& RHICmdList, const TCHAR* InDebugName, uint32 BytesPerElement, uint32 NumElements, EBufferUsageFlags AdditionalUsage = BUF_None, bool bUseUavCounter = false, bool bAppendBuffer = false, ERHIAccess InitialState = ERHIAccess::UAVMask)
+	void Initialize(FRHICommandListBase& RHICmdList, const TCHAR* InDebugName, uint32 BytesPerElement, uint32 NumElements, EBufferUsageFlags AdditionalUsage = EBufferUsageFlags::None, bool bUseUavCounter = false, bool bAppendBuffer = false, ERHIAccess InitialState = ERHIAccess::UAVMask)
 	{
 		// Provide a debug name if using Fast VRAM so the allocators diagnostics will work
-		ensure(!(EnumHasAnyFlags(AdditionalUsage, BUF_FastVRAM) && !InDebugName));
+		ensure(!(EnumHasAnyFlags(AdditionalUsage, EBufferUsageFlags::FastVRAM) && !InDebugName));
 
 		NumBytes = BytesPerElement * NumElements;
-		FRHIResourceCreateInfo CreateInfo(InDebugName);
-		Buffer = RHICmdList.CreateStructuredBuffer(BytesPerElement, NumBytes, BUF_UnorderedAccess | BUF_ShaderResource | AdditionalUsage, InitialState, CreateInfo);
-		
+
+		const FRHIBufferCreateDesc CreateDesc =
+			FRHIBufferCreateDesc::CreateStructured(InDebugName, NumBytes, BytesPerElement)
+			.AddUsage(EBufferUsageFlags::UnorderedAccess | EBufferUsageFlags::ShaderResource | AdditionalUsage)
+			.SetInitialState(InitialState);
+
+		Buffer = RHICmdList.CreateBuffer(CreateDesc);
+		UAV = RHICmdList.CreateUnorderedAccessView(
+			Buffer,
+			FRHIViewDesc::CreateBufferUAV()
+				.SetTypeFromBuffer(Buffer)
+				.SetAtomicCounter(bUseUavCounter)
+				.SetAppendBuffer(bAppendBuffer));
+		SRV = RHICmdList.CreateShaderResourceView(Buffer, FRHIViewDesc::CreateBufferSRV().SetTypeFromBuffer(Buffer));
+	}
+
+	template<typename LAMBDA>
+	void InitializeWithData(FRHICommandListBase& RHICmdList, const TCHAR* InDebugName, uint32 BytesPerElement, uint32 NumElements, LAMBDA&& InitializerFunction, EBufferUsageFlags AdditionalUsage = EBufferUsageFlags::None, bool bUseUavCounter = false, bool bAppendBuffer = false, ERHIAccess InitialState = ERHIAccess::UAVMask)
+	{
+		// Provide a debug name if using Fast VRAM so the allocators diagnostics will work
+		ensure(!(EnumHasAnyFlags(AdditionalUsage, EBufferUsageFlags::FastVRAM) && !InDebugName));
+
+		NumBytes = BytesPerElement * NumElements;
+
+		const FRHIBufferCreateDesc CreateDesc =
+			FRHIBufferCreateDesc::CreateStructured(InDebugName, NumBytes, BytesPerElement)
+			.AddUsage(EBufferUsageFlags::UnorderedAccess | EBufferUsageFlags::ShaderResource | AdditionalUsage)
+			.SetInitialState(InitialState)
+			.SetInitActionInitializer();
+
+		FRHIBufferInitializer Initializer = RHICmdList.CreateBufferInitializer(CreateDesc);
+		InitializerFunction(Initializer);
+		Buffer = Initializer.Finalize();
+
 		UAV = RHICmdList.CreateUnorderedAccessView(
 			Buffer,
 			FRHIViewDesc::CreateBufferUAV()
 				.SetTypeFromBuffer(Buffer)
 				.SetAtomicCounter(bUseUavCounter)
 				.SetAppendBuffer(bAppendBuffer));
-		
 		SRV = RHICmdList.CreateShaderResourceView(Buffer, FRHIViewDesc::CreateBufferSRV().SetTypeFromBuffer(Buffer));
 	}
 
@@ -383,12 +485,17 @@ struct FByteAddressBuffer
 
 	FByteAddressBuffer(): NumBytes(0) {}
 
-	void Initialize(FRHICommandListBase& RHICmdList, const TCHAR* InDebugName, uint32 InNumBytes, EBufferUsageFlags AdditionalUsage = BUF_None)
+	void Initialize(FRHICommandListBase& RHICmdList, const TCHAR* InDebugName, uint32 InNumBytes, EBufferUsageFlags AdditionalUsage = EBufferUsageFlags::None)
 	{
 		NumBytes = InNumBytes;
 		check( NumBytes % 4 == 0 );
-		FRHIResourceCreateInfo CreateInfo(InDebugName);
-		Buffer = RHICmdList.CreateStructuredBuffer(4, NumBytes, BUF_ShaderResource | BUF_ByteAddressBuffer | AdditionalUsage, ERHIAccess::SRVMask, CreateInfo);
+
+		const FRHIBufferCreateDesc CreateDesc =
+			FRHIBufferCreateDesc::CreateStructured(InDebugName, NumBytes, 4)
+			.AddUsage(EBufferUsageFlags::ShaderResource | EBufferUsageFlags::ByteAddressBuffer | AdditionalUsage)
+			.SetInitialState(ERHIAccess::SRVMask);
+
+		Buffer = RHICmdList.CreateBuffer(CreateDesc);
 		SRV = RHICmdList.CreateShaderResourceView(Buffer, FRHIViewDesc::CreateBufferSRV().SetType(FRHIViewDesc::EBufferType::Raw));
 	}
 
diff --git a/Engine/Source/Runtime/RHI/Public/RHIValidation.h b/Engine/Source/Runtime/RHI/Public/RHIValidation.h
--- a/Engine/Source/Runtime/RHI/Public/RHIValidation.h
+++ b/Engine/Source/Runtime/RHI/Public/RHIValidation.h
@@ -418,27 +418,26 @@ public:
 		});
 	}
 
-	/**
-	* @param ResourceArray - An optional pointer to a resource array containing the resource's data.
-	*/
-	virtual FBufferRHIRef RHICreateBuffer(FRHICommandListBase& RHICmdList, FRHIBufferDesc const& Desc, ERHIAccess ResourceState, FRHIResourceCreateInfo& CreateInfo) override final
+	struct FRHIValidationBufferInitializer : public FRHIBufferInitializer
 	{
-		if (EnumHasAnyFlags(Desc.Usage, BUF_ReservedResource))
+		FRHIValidationBufferInitializer(const FRHIBufferCreateDesc& CreateDesc, FRHIBufferInitializer&& InOther)
+			: FRHIBufferInitializer(MoveTemp(InOther))
 		{
-			RHI_VALIDATION_CHECK(GRHIGlobals.ReservedResources.Supported, TEXT("Reserved buffers are not supported by current RHI"));
-			RHI_VALIDATION_CHECK(!EnumHasAnyFlags(Desc.Usage, BUF_AnyDynamic), TEXT("Reserved buffers must not be dynamic"));
-			RHI_VALIDATION_CHECK(Desc.Stride <= GRHIGlobals.ReservedResources.TileSizeInBytes, TEXT("Reserved buffer stride must not be greater than reserved resource tile size"));
-		}
+			FFinalizeCallback OriginalFinalizeCallback = MoveTemp(FinalizeCallback);
+			FinalizeCallback = [OriginalFinalizeCallback = MoveTemp(OriginalFinalizeCallback), InitialState = CreateDesc.InitialState](FRHICommandListBase& RHICmdList)
+				{
+					FBufferRHIRef Buffer = OriginalFinalizeCallback(RHICmdList);
 
-		if (CreateInfo.ResourceArray && RHICmdList.IsInsideRenderPass())
-		{
-			FString Msg = FString::Printf(TEXT("Creating buffers with initial data during a render pass is not supported, buffer name: \"%s\""), CreateInfo.DebugName);
-			RHI_VALIDATION_CHECK(false, *Msg);
+					Buffer->InitBarrierTracking(InitialState, Buffer->GetName());
+
+					return Buffer;
+				};
 		}
+	};
 
-		FBufferRHIRef Buffer = RHI->RHICreateBuffer(RHICmdList, Desc, ResourceState, CreateInfo);
-		Buffer->InitBarrierTracking(ResourceState, CreateInfo.DebugName);
-		return Buffer;
+	[[nodiscard]] virtual FRHIBufferInitializer RHICreateBufferInitializer(FRHICommandListBase& RHICmdList, const FRHIBufferCreateDesc& CreateDesc) final override
+	{
+		return FRHIValidationBufferInitializer(CreateDesc, RHI->RHICreateBufferInitializer(RHICmdList, CreateDesc));
 	}
 
 	// FlushType: Flush RHI Thread
diff --git a/Engine/Source/Runtime/RHI/Public/RHIValidationCommon.h b/Engine/Source/Runtime/RHI/Public/RHIValidationCommon.h
--- a/Engine/Source/Runtime/RHI/Public/RHIValidationCommon.h
+++ b/Engine/Source/Runtime/RHI/Public/RHIValidationCommon.h
@@ -455,9 +455,10 @@ namespace RHIValidation
 	class FBufferResource : public FResource
 	{
 	public:
-		inline void InitBarrierTracking(ERHIAccess InResourceState, const TCHAR* InDebugName)
+		inline void InitBarrierTracking(ERHIAccess InResourceState, FName InDebugName)
 		{
-			FResource::InitBarrierTracking(1, 1, 1, InResourceState, InDebugName);
+			const FString LocalDebugName = InDebugName.ToString();
+			FResource::InitBarrierTracking(1, 1, 1, InResourceState, *LocalDebugName);
 		}
 	};
 
diff --git a/Engine/Source/Runtime/RHICore/Internal/RHICoreResourceCollection.h b/Engine/Source/Runtime/RHICore/Internal/RHICoreResourceCollection.h
--- a/Engine/Source/Runtime/RHICore/Internal/RHICoreResourceCollection.h
+++ b/Engine/Source/Runtime/RHICore/Internal/RHICoreResourceCollection.h
@@ -38,17 +38,14 @@ namespace UE::RHICore
 		return Handle;
 	}
 
-	template<typename TType>
-	inline void FillResourceCollectionMemory(TArray<uint32>& Destination, TConstArrayView<TType> InValues)
+	inline void FillResourceCollectionMemory(TRHIBufferInitializer<uint32>& Destination, TConstArrayView<FRHIResourceCollectionMember> InValues)
 	{
-		Destination.SetNumUninitialized(1 + InValues.Num());
-
 		int32 WriteIndex = 0;
 
 		Destination[WriteIndex] = static_cast<uint32>(InValues.Num());
 		++WriteIndex;
 
-		for (const TType& Value : InValues)
+		for (const FRHIResourceCollectionMember& Value : InValues)
 		{
 			const FRHIDescriptorHandle Handle = GetHandleForResourceCollectionValue(Value);
 			check(Handle.IsValid());
@@ -60,42 +57,20 @@ namespace UE::RHICore
 		}
 	}
 
-	template<typename TType>
-	inline TArray<uint32> CreateResourceCollectionArray(TConstArrayView<TType> InValues)
-	{
-		TArray<uint32> Result;
-		FillResourceCollectionMemory(Result, InValues);
-		return Result;
-	}
-
-	struct FResourceCollectionUpload : public FResourceArrayUploadInterface
+	inline FRHIBuffer* CreateResourceCollectionBuffer(FRHICommandListBase& RHICmdList, TConstArrayView<FRHIResourceCollectionMember> InMembers)
 	{
-		FResourceCollectionUpload(TConstArrayView<FRHIResourceCollectionMember> InMembers)
-			: Memory(CreateResourceCollectionArray(InMembers))
-		{
-		}
-		virtual const void* GetResourceData() const final
-		{
-			return Memory.GetData();
-		}
-		virtual uint32 GetResourceDataSize() const final
-		{
-			return Memory.Num() * Memory.GetTypeSize();
-		}
-		virtual void Discard() final
-		{
-			Memory.Reset();
-		}
+		const size_t BufferSize = UE::RHICore::CalculateResourceCollectionMemorySize(InMembers);
 
-		TArray<uint32> Memory;
-	};
+		const FRHIBufferCreateDesc CreateDesc =
+			FRHIBufferCreateDesc::CreateByteAddress(TEXT("ResourceCollection"), BufferSize, sizeof(uint32))
+			.AddUsage(EBufferUsageFlags::Static)
+			.SetInitialState(ERHIAccess::SRVMask)
+			.SetInitActionInitializer();
 
-	inline FRHIBuffer* CreateResourceCollectionBuffer(FRHICommandListBase& RHICmdList, TConstArrayView<FRHIResourceCollectionMember> InMembers)
-	{
-		FResourceCollectionUpload UploadData(InMembers);
+		TRHIBufferInitializer<uint32> Initializer = RHICmdList.CreateBufferInitializer(CreateDesc);
+		UE::RHICore::FillResourceCollectionMemory(Initializer, InMembers);
 
-		FRHIResourceCreateInfo CreateInfo(TEXT("ResourceCollection"), &UploadData);
-		return RHICmdList.CreateBuffer(UploadData.GetResourceDataSize(), EBufferUsageFlags::Static | EBufferUsageFlags::ByteAddressBuffer, sizeof(uint32), ERHIAccess::SRVMask, CreateInfo);
+		return Initializer.Finalize();
 	}
 
 	class FGenericResourceCollection : public FRHIResourceCollection
diff --git a/Engine/Source/Runtime/RHICore/Public/RHICoreTransientResourceAllocator.h b/Engine/Source/Runtime/RHICore/Public/RHICoreTransientResourceAllocator.h
--- a/Engine/Source/Runtime/RHICore/Public/RHICoreTransientResourceAllocator.h
+++ b/Engine/Source/Runtime/RHICore/Public/RHICoreTransientResourceAllocator.h
@@ -32,12 +32,20 @@ inline uint64 ComputeHash(const FRHITextureCreateInfo& InCreateInfo)
 
 inline uint64 ComputeHash(const FRHIBufferCreateInfo& InCreateInfo, uint64 HeapOffset)
 {
-	return CityHash64WithSeed((const char*)&InCreateInfo, sizeof(FRHIBufferCreateInfo), HeapOffset);
+	// Make sure all padding is removed.
+	FRHIBufferCreateInfo NewInfo;
+	FPlatformMemory::Memzero(&NewInfo, sizeof(NewInfo));
+	NewInfo = InCreateInfo;
+	return CityHash64WithSeed(reinterpret_cast<const char*>(&NewInfo), sizeof(NewInfo), HeapOffset);
 }
 
 inline uint64 ComputeHash(const FRHIBufferCreateInfo& InCreateInfo)
 {
-	return CityHash64((const char*)&InCreateInfo, sizeof(FRHIBufferCreateInfo));
+	// Make sure all padding is removed.
+	FRHIBufferCreateInfo NewInfo;
+	FPlatformMemory::Memzero(&NewInfo, sizeof(NewInfo));
+	NewInfo = InCreateInfo;
+	return CityHash64(reinterpret_cast<const char*>(&NewInfo), sizeof(NewInfo));
 }
 
 /** Tracks allocation statistics for buffer or texture resources. */
diff --git a/Engine/Source/Runtime/RenderCore/Private/ClearQuad.cpp b/Engine/Source/Runtime/RenderCore/Private/ClearQuad.cpp
--- a/Engine/Source/Runtime/RenderCore/Private/ClearQuad.cpp
+++ b/Engine/Source/Runtime/RenderCore/Private/ClearQuad.cpp
@@ -8,20 +8,20 @@
 #include "PipelineStateCache.h"
 #include "RendererInterface.h"
 #include "Logging/LogMacros.h"
+#include "RHIResourceUtils.h"
+
+static const FVector4f GClearVertexBufferVertices[] =
+{
+	FVector4f(-1.0f, 1.0f, 0.0f, 1.0f),
+	FVector4f(1.0f, 1.0f, 0.0f, 1.0f),
+	FVector4f(-1.0f, -1.0f, 0.0f, 1.0f),
+	FVector4f(1.0f, -1.0f, 0.0f, 1.0f),
+};
 
 void FClearVertexBuffer::InitRHI(FRHICommandListBase& RHICmdList)
 {
 	// create a static vertex buffer
-	FRHIResourceCreateInfo CreateInfo(TEXT("FClearVertexBuffer"));
-	VertexBufferRHI = RHICmdList.CreateVertexBuffer(sizeof(FVector4f) * 4, BUF_Static, CreateInfo);
-	void* VoidPtr = RHICmdList.LockBuffer(VertexBufferRHI, 0, sizeof(FVector4f) * 4, RLM_WriteOnly);
-	// Generate the vertices used
-	FVector4f* Vertices = reinterpret_cast<FVector4f*>(VoidPtr);
-	Vertices[0] = FVector4f(-1.0f, 1.0f, 0.0f, 1.0f);
-	Vertices[1] = FVector4f(1.0f, 1.0f, 0.0f, 1.0f);
-	Vertices[2] = FVector4f(-1.0f, -1.0f, 0.0f, 1.0f);
-	Vertices[3] = FVector4f(1.0f, -1.0f, 0.0f, 1.0f);
-	RHICmdList.UnlockBuffer(VertexBufferRHI);
+	VertexBufferRHI = UE::RHIResourceUtils::CreateVertexBufferFromArray(RHICmdList, TEXT("FClearVertexBuffer"), EBufferUsageFlags::Static, MakeConstArrayView(GClearVertexBufferVertices));
 }
 
 TGlobalResource<FClearVertexBuffer> GClearVertexBuffer;
@@ -204,28 +204,23 @@ void DrawClearQuadMRT(FRHICommandList& RHICmdList, bool bClearColor, int32 NumCl
 		InnerVertices[2].Set(FMath::Lerp(-1.0f, 1.0f, FractionRect.Z), FMath::Lerp(1.0f, -1.0f, FractionRect.W), Depth, 1.0f);
 		InnerVertices[3].Set(FMath::Lerp(-1.0f, 1.0f, FractionRect.X), FMath::Lerp(1.0f, -1.0f, FractionRect.W), Depth, 1.0f);
 
-		FRHIResourceCreateInfo CreateInfo(TEXT("DrawClearQuadMRT"));
-		FBufferRHIRef VertexBufferRHI = RHICmdList.CreateVertexBuffer(sizeof(FVector4f) * 10, BUF_Volatile, CreateInfo);
-		void* VoidPtr = RHICmdList.LockBuffer(VertexBufferRHI, 0, sizeof(FVector4f) * 10, RLM_WriteOnly);
-		
-		FVector4f* Vertices = reinterpret_cast<FVector4f*>(VoidPtr);
-		Vertices[0] = OuterVertices[0];
-		Vertices[1] = InnerVertices[0];
-		Vertices[2] = OuterVertices[1];
-		Vertices[3] = InnerVertices[1];
-		Vertices[4] = OuterVertices[2];
-		Vertices[5] = InnerVertices[2];
-		Vertices[6] = OuterVertices[3];
-		Vertices[7] = InnerVertices[3];
-		Vertices[8] = OuterVertices[0];
-		Vertices[9] = InnerVertices[0];
-
-		RHICmdList.UnlockBuffer(VertexBufferRHI);
-		RHICmdList.SetStreamSource(0, VertexBufferRHI, 0);
+		const FVector4f Vertices[] =
+		{
+			OuterVertices[0],
+			InnerVertices[0],
+			OuterVertices[1],
+			InnerVertices[1],
+			OuterVertices[2],
+			InnerVertices[2],
+			OuterVertices[3],
+			InnerVertices[3],
+			OuterVertices[0],
+			InnerVertices[0],
+		};
+		FBufferRHIRef VertexBufferRHI = UE::RHIResourceUtils::CreateVertexBufferFromArray(RHICmdList, TEXT("DrawClearQuadMRT"), EBufferUsageFlags::Volatile, MakeConstArrayView(Vertices));
 
+		RHICmdList.SetStreamSource(0, VertexBufferRHI, 0);
 		RHICmdList.DrawPrimitive(0, 8, 1);
-
-		VertexBufferRHI.SafeRelease();
 	}
 	else
 	{
