VulkanRHI: fix NeedsAllPlanes when transitions span multiple cmdlists: leave existing tracking untouched and track 'active' transitions (read renderpass resource states from active transition if one exists).

#rnx

diff --git a/Engine/Source/Runtime/RHI/Private/RHICommandList.cpp b/Engine/Source/Runtime/RHI/Private/RHICommandList.cpp
--- a/Engine/Source/Runtime/RHI/Private/RHICommandList.cpp
+++ b/Engine/Source/Runtime/RHI/Private/RHICommandList.cpp
@@ -2273,10 +2273,16 @@ bool FRHICommandListBase::TrackTransitioningResource(const FRHITransitionCreateI
 {
 	if (GRHIGlobals.NeedsAllPlanes && GRHIGlobals.EnableHighLevelTransitionPatching)
 	{
+		// NoSplit means we don't have to worry about render passes between Begin and End
+		// Only applies to pipelines that remain on the graphics queue (for render passes)
+		// Only applies to depth textures with a stencil
+		// Only worth tracking if it's not an unknown/discard state
 		if (!EnumHasAnyFlags(TransitionCreateInfoData.m_Flags, ERHITransitionCreateFlags::NoSplit) &&
 			EnumHasAnyFlags(TransitionCreateInfoData.SrcPipelines, ERHIPipeline::Graphics) &&
 			EnumHasAnyFlags(TransitionCreateInfoData.DstPipelines, ERHIPipeline::Graphics) &&
-			(TransitionInfo.Type == FRHITransitionInfo::EType::Texture))
+			(TransitionInfo.Type == FRHITransitionInfo::EType::Texture) &&
+			(TransitionInfo.AccessBefore != ERHIAccess::Unknown) && 
+			(TransitionInfo.AccessBefore != ERHIAccess::Discard))
 		{
 			FRHITexture* Texture = (FRHITexture*)GetViewableResource(TransitionInfo);
 			if (Texture && IsStencilFormat(Texture->GetDesc().Format))
@@ -2363,8 +2369,22 @@ void FRHICommandListBase::TrackResourceStateBegin(TArrayView<const FRHITransitio
 			const bool bShouldSkipAfterState = ((EnumHasAnyFlags(TransitionInfo->Flags, EResourceTransitionFlags::IgnoreAfterState)) || (TransitionInfo->AccessAfter == ERHIAccess::Unknown));
 			const bool bNeedUpdateAfterState = (bShouldSkipAfterState == false) && (bDiscardOnAsync == false);
 
-			if (bNeedUpdateAfterState && !TrackTransitioningResource(*TransitionCreateInfoData, *TransitionInfo))
+			if (bNeedUpdateAfterState)
 			{
+				if (TrackTransitioningResource(*TransitionCreateInfoData, *TransitionInfo))
+				{
+					if (ActiveTransitionContainer* TransitionArray = ActiveTransitions.Find(TransitionViewableResource))
+					{
+						TransitionArray->Add(TransitionInfo);
+					}
+					else
+					{
+						ActiveTransitionContainer NewArray;
+						NewArray.Add(TransitionInfo);
+						ActiveTransitions.Add(TransitionViewableResource, MoveTemp(NewArray));
+					}
+				}
+
 				// Update the TrackedResource State
 				if (TrackedResourceState)
 				{
@@ -2411,44 +2431,53 @@ void FRHICommandListBase::TrackResourceStateEnd(TArrayView<const FRHITransition*
 	{
 		const FRHITransitionCreateInfoData* TransitionCreateInfoData = Transition->CreationInfoData;
 
-		if (TransitionCreateInfoData)
+		if (!TransitionCreateInfoData)
 		{
-			const bool bTrackOnEnd = EnumHasAnyFlags(TransitionCreateInfoData->SrcPipelines, ERHIPipeline::Graphics) && EnumHasAnyFlags(TransitionCreateInfoData->DstPipelines, ERHIPipeline::Graphics);
-			if (bTrackOnEnd)
+			continue;
+		}
+
+		for (int32 Index = 0; Index < TransitionCreateInfoData->TransitionInfos.Num(); Index++)
+		{
+			const FRHITransitionInfo& TransitionInfo = TransitionCreateInfoData->TransitionInfos[Index];
+
+			if (!TrackTransitioningResource(*TransitionCreateInfoData, TransitionInfo))
 			{
-				for (int32 Index = 0; Index < TransitionCreateInfoData->TransitionInfos.Num(); Index++)
-				{
-					const FRHITransitionInfo* TransitionInfo = &TransitionCreateInfoData->TransitionInfos[Index];
+				continue;
+			}
 
-					if (!TrackTransitioningResource(*TransitionCreateInfoData, *TransitionInfo))
-					{
-						continue;
-					}
+			FRHIViewableResource* TransitionViewableResource = GetViewableResource(TransitionInfo);
+			if (TransitionViewableResource == nullptr)
+			{
+				continue;
+			}
 
-					const bool bShouldSkipAfterState = ((EnumHasAnyFlags(TransitionInfo->Flags, EResourceTransitionFlags::IgnoreAfterState)) || (TransitionInfo->AccessAfter == ERHIAccess::Unknown));
-					const bool bNeedUpdateAfterState = (bShouldSkipAfterState == false);
-					if (bNeedUpdateAfterState)
-					{
-						FRHIViewableResource* TransitionViewableResource = GetViewableResource(*TransitionInfo);
-						TRHIResourceState<FRHICmdListBaseArrayAllocator>* TrackedResourceState = TrackedResources.Find(TransitionViewableResource);
+			bool bWasFound = false;
+			if (ActiveTransitionContainer* TransitionArray = ActiveTransitions.Find(TransitionViewableResource))
+			{
+				const int32 NumRemoved = TransitionArray->RemoveSwap(&TransitionInfo);
+				bWasFound = (NumRemoved > 0);
+			}
 
-						// Update the TrackedResource State
-						if (TrackedResourceState)
-						{
-							TrackedResourceState->SetSubresourceState(*TransitionInfo, TransitionInfo->AccessAfter);
-						}
-						else
-						{
-							// Add a new element
-							TrackedResources.Add(TransitionViewableResource, TRHIResourceState<FRHICmdListBaseArrayAllocator>(*TransitionInfo, TransitionInfo->AccessAfter, TransitionViewableResource->GetTrackedSubResourceCount(), true));
-						}
-					}
+
+			if (!bWasFound)
+			{
+				if (ActiveTransitionContainer* TransitionArray = ClosedTransitions.Find(TransitionViewableResource))
+				{
+					TransitionArray->Add(&TransitionInfo);
+				}
+				else
+				{
+					ActiveTransitionContainer NewArray;
+					NewArray.Add(&TransitionInfo);
+					ClosedTransitions.Add(TransitionViewableResource, MoveTemp(NewArray));
 				}
 			}
 		}
 	}
 }
 
+TMap<FRHIViewableResource*, TArray<FRHITransitionInfo>> FRHICommandListBase::GlobalActiveTransitions;
+
 void FRHICommandListBase::UpdateTransitionsAndResourceState()
 {
 	TRACE_CPUPROFILER_EVENT_SCOPE(UpdateTransitionsAndResourceState);
@@ -2559,16 +2588,82 @@ void FRHICommandListBase::UpdateTransitionsAndResourceState()
 			// Iterate over all RenderPassInfo to patch any missing plane states
 			for (FRenderPassPatchInfo& PatchInfo : RenderPassInfoToPatch)
 			{
-				if (PatchInfo.State != ERHIAccess::Unknown)
+				PatchInfo.RenderPassInfo->NopAccess = PatchInfo.State;
+
+				if (PatchInfo.RenderPassInfo->NopAccess == ERHIAccess::Unknown)
 				{
-					PatchInfo.RenderPassInfo->NopAccess = PatchInfo.State;
+					FRHITexture* DepthStencilTexture = PatchInfo.RenderPassInfo->DepthStencilRenderTarget.DepthStencilTarget;				
+					if (TArray<FRHITransitionInfo>* TransitionArray = GlobalActiveTransitions.Find(DepthStencilTexture))
+					{
+						for (const FRHITransitionInfo& TransitionInfo : *TransitionArray)
+						{
+							if (!PatchInfo.RenderPassInfo->DepthStencilRenderTarget.ExclusiveDepthStencil.IsUsingDepth())
+							{
+								if (TransitionInfo.IsAllPlaneSlices() || TransitionInfo.PlaneSlice == FRHISubresourceRange::kDepthPlaneSlice)
+								{
+									PatchInfo.RenderPassInfo->NopAccess = TransitionInfo.AccessBefore;
+									break;
+								}
+							}
+							else if (!PatchInfo.RenderPassInfo->DepthStencilRenderTarget.ExclusiveDepthStencil.IsUsingStencil())
+							{
+								if (TransitionInfo.IsAllPlaneSlices() || TransitionInfo.PlaneSlice == FRHISubresourceRange::kStencilPlaneSlice)
+								{
+									PatchInfo.RenderPassInfo->NopAccess = TransitionInfo.AccessBefore;
+									break;
+								}
+							}
+						}
+					}
 				}
-				else
+
+				if (PatchInfo.RenderPassInfo->NopAccess == ERHIAccess::Unknown)
 				{
 					ResolveNopAccessInPlace(*PatchInfo.RenderPassInfo);
 				}
 			}
 			RenderPassInfoToPatch.Empty();
+
+			for (const auto& Pair : ClosedTransitions)
+			{
+				if (TArray<FRHITransitionInfo>* GlobalTransitionArray = GlobalActiveTransitions.Find(Pair.Key))
+				{
+					for (const FRHITransitionInfo* ClosedTransition : Pair.Value)
+					{
+						const int32 RemoveCount = GlobalTransitionArray->RemoveSwap(*ClosedTransition);
+						if (RemoveCount == 0)
+						{
+							UE_LOG(LogRHI, Warning, TEXT("Could not close transition that didn't exist in the tracked GlobalActiveTransitions."));
+						}
+
+						if (GlobalTransitionArray->Num() == 0)
+						{
+							GlobalActiveTransitions.Remove(Pair.Key);
+						}
+					}
+				}
+			}
+			ClosedTransitions.Empty();
+
+			for (const auto& Pair : ActiveTransitions)
+			{
+				if (Pair.Value.Num() == 0)
+				{
+					continue;
+				}
+
+				TArray<FRHITransitionInfo>* GlobalTransitionArray = GlobalActiveTransitions.Find(Pair.Key);
+				if (!GlobalTransitionArray)
+				{
+					GlobalTransitionArray = &GlobalActiveTransitions.Add(Pair.Key);
+				}
+
+				for (const FRHITransitionInfo* ActiveTransition : Pair.Value)
+				{
+					GlobalTransitionArray->Add(*ActiveTransition);
+				}
+			}
+			ActiveTransitions.Empty();
 		}
 
 		// Iterate over all tracked Resources to update their TrackedState with the latest state
diff --git a/Engine/Source/Runtime/RHI/Public/RHICommandList.h b/Engine/Source/Runtime/RHI/Public/RHICommandList.h
--- a/Engine/Source/Runtime/RHI/Public/RHICommandList.h
+++ b/Engine/Source/Runtime/RHI/Public/RHICommandList.h
@@ -1579,11 +1579,13 @@ private:
 
 	// Transitions with Unknown or Discard (when needed) that will patched at Dispatch time
 	TArray<FTransitionPatchInfo, FRHICmdListBaseArrayAllocator> TransitionInfoToPatch;
-	// Transitions with a single plane of DepthStencil specified where the other plane needs to be resolved with tracking
+	// (NeedsAllPlanes) Transitions with a single plane of DepthStencil specified where the other plane needs to be resolved with tracking
 	TArray<FTransitionPatchInfo, FRHICmdListBaseArrayAllocator> PlaneTransitionInfoToPatch;
-	// RenderPassInfo with a single plane of DepthStencil specified where the other plane needs to be resolved with tracking
+	// (NeedsAllPlanes) RenderPassInfo with a single plane of DepthStencil specified where the other plane needs to be resolved with tracking
 	TArray<FRenderPassPatchInfo, FRHICmdListBaseArrayAllocator> RenderPassInfoToPatch;
-	
+	// (NeedsAllPlanes) Track transitions happening over multiple command lists
+	static TMap<FRHIViewableResource*, TArray<FRHITransitionInfo>> GlobalActiveTransitions;
+
 	// This will scan the array and keep track of transitions that needs patching later at Dispatch time
 	RHI_API void TrackResourceStateBegin(TArrayView<const FRHITransition*> Transitions);
 	RHI_API void TrackResourceStateEnd(TArrayView<const FRHITransition*> Transitions);
@@ -1598,6 +1600,10 @@ private:
 
 	// (NeedsAllPlanes) Resource needs tracking while it's transitioning in case a RenderPass queries its state between Begin and End
 	RHI_API static bool TrackTransitioningResource(const FRHITransitionCreateInfoData& TransitionCreateInfoData, const FRHITransitionInfo& TransitionInfo);
+	// (NeedsAllPlanes) Track transitions that happen over multiple cmdlists for resources flagged by TrackTransitioningResource
+	using ActiveTransitionContainer = TArray<const FRHITransitionInfo*, TInlineAllocator<2>>;
+	TMap<FRHIViewableResource*, ActiveTransitionContainer, FRHICmdListBaseSetAllocator> ActiveTransitions;
+	TMap<FRHIViewableResource*, ActiveTransitionContainer, FRHICmdListBaseSetAllocator> ClosedTransitions;
 
 protected:
 	inline void ResolveNopAccessInPlace(FRHIRenderPassInfo& InInfo);
@@ -4363,6 +4369,7 @@ public:
 		// will need tracking to resolve the RHIAccess for the missing plane
 		const bool bNeedsNopResolve = NeedsAllPlanes() &&
 			InInfo.DepthStencilRenderTarget.DepthStencilTarget &&
+			IsStencilFormat(InInfo.DepthStencilRenderTarget.DepthStencilTarget->GetDesc().Format) &&
 			(!InInfo.DepthStencilRenderTarget.ExclusiveDepthStencil.IsUsingDepth() ||
 			 !InInfo.DepthStencilRenderTarget.ExclusiveDepthStencil.IsUsingStencil());
 
diff --git a/Engine/Source/Runtime/RHI/Public/RHICommandList.inl b/Engine/Source/Runtime/RHI/Public/RHICommandList.inl
--- a/Engine/Source/Runtime/RHI/Public/RHICommandList.inl
+++ b/Engine/Source/Runtime/RHI/Public/RHICommandList.inl
@@ -399,24 +399,50 @@ FORCEINLINE_DEBUGGABLE void FRHICommandListBase::TrackNopAccess(const FRHIRender
 	checkSlow(NeedsAllPlanes());
 	if (InInfo.DepthStencilRenderTarget.DepthStencilTarget)
 	{
-		TRHIResourceState<FRHICmdListBaseArrayAllocator>* TrackedResourceState = TrackedResources.Find(InInfo.DepthStencilRenderTarget.DepthStencilTarget);
-
 		ERHIAccess NopState = ERHIAccess::Unknown;
-		if (TrackedResourceState)
+
+		if (ActiveTransitionContainer* TransitionArray = ActiveTransitions.Find(InInfo.DepthStencilRenderTarget.DepthStencilTarget))
 		{
-			if (!InInfo.DepthStencilRenderTarget.ExclusiveDepthStencil.IsUsingDepth())
+			for (const FRHITransitionInfo* TransitionInfo : *TransitionArray)
 			{
-				const FRHISubresourceRange SubResourceIndex(0, 0, FRHISubresourceRange::kDepthPlaneSlice);
-				NopState = TrackedResourceState->GetSubresourceState(SubResourceIndex);
+				if (!InInfo.DepthStencilRenderTarget.ExclusiveDepthStencil.IsUsingDepth())
+				{
+					if (TransitionInfo->IsAllPlaneSlices() || TransitionInfo->PlaneSlice == FRHISubresourceRange::kDepthPlaneSlice)
+					{
+						NopState = TransitionInfo->AccessBefore;
+						break;
+					}
+				}
+				else if (!InInfo.DepthStencilRenderTarget.ExclusiveDepthStencil.IsUsingStencil())
+				{
+					if (TransitionInfo->IsAllPlaneSlices() || TransitionInfo->PlaneSlice == FRHISubresourceRange::kStencilPlaneSlice)
+					{
+						NopState = TransitionInfo->AccessBefore;
+						break;
+					}
+				}
 			}
-			else if (!InInfo.DepthStencilRenderTarget.ExclusiveDepthStencil.IsUsingStencil())
+		}
+
+		if (NopState == ERHIAccess::Unknown)
+		{
+			TRHIResourceState<FRHICmdListBaseArrayAllocator>* TrackedResourceState = TrackedResources.Find(InInfo.DepthStencilRenderTarget.DepthStencilTarget);
+			if (TrackedResourceState)
 			{
-				const FRHISubresourceRange SubResourceIndex(0, 0, FRHISubresourceRange::kStencilPlaneSlice);
-				NopState = TrackedResourceState->GetSubresourceState(SubResourceIndex);
+				if (!InInfo.DepthStencilRenderTarget.ExclusiveDepthStencil.IsUsingDepth())
+				{
+					const FRHISubresourceRange SubResourceIndex(0, 0, FRHISubresourceRange::kDepthPlaneSlice);
+					NopState = TrackedResourceState->GetSubresourceState(SubResourceIndex);
+				}
+				else if (!InInfo.DepthStencilRenderTarget.ExclusiveDepthStencil.IsUsingStencil())
+				{
+					const FRHISubresourceRange SubResourceIndex(0, 0, FRHISubresourceRange::kStencilPlaneSlice);
+					NopState = TrackedResourceState->GetSubresourceState(SubResourceIndex);
+				}
 			}
 		}
 
-		// Note: even if PreviousState is a valid state we cannot write into TransitionInfo since we are recording the cmdlist
+		// Note: even if the state is a valid state we cannot write into FRHIRenderPassInfo since we are recording the cmdlist
 		RenderPassInfoToPatch.Push(FRenderPassPatchInfo((FRHIRenderPassInfo*)&InInfo, NopState));
 	}
 }
