Removed LockTracker TotalMemoryOutstanding and checks that were causing a flush during unlock

diff --git a/Engine/Source/Runtime/Apple/MetalRHI/Private/MetalRHIBuffer.cpp b/Engine/Source/Runtime/Apple/MetalRHI/Private/MetalRHIBuffer.cpp
--- a/Engine/Source/Runtime/Apple/MetalRHI/Private/MetalRHIBuffer.cpp
+++ b/Engine/Source/Runtime/Apple/MetalRHI/Private/MetalRHIBuffer.cpp
@@ -718,7 +718,6 @@ void FMetalDynamicRHI::RHIUnlockBuffer(class FRHICommandListBase& RHICmdList, FR
 			
 			FRHICommandListScopedFlushAndExecute Flush(RHICmdList.GetAsImmediate());
 			Buffer->Unlock(RHICmdList);
-			GBufferLockTracker.TotalMemoryOutstanding = 0;
 			
 			delete LockData;
 		}
@@ -748,13 +747,6 @@ void FMetalDynamicRHI::RHIUnlockBuffer(class FRHICommandListBase& RHICmdList, FR
 				Buffer->Unlock(RHICmdList);
 			});
 			RHICmdList.RHIThreadFence(true);
-			
-			if (RHICmdList.IsImmediate() && GBufferLockTracker.TotalMemoryOutstanding > 256 * 1024)
-			{
-				QUICK_SCOPE_CYCLE_COUNTER(STAT_RHIMETHOD_UnlockBuffer_FlushForMem);
-				RHICmdList.GetAsImmediate().ImmediateFlush(EImmediateFlushType::DispatchToRHIThread);
-				GBufferLockTracker.TotalMemoryOutstanding = 0;
-			}
 		}
 	}
 	else
diff --git a/Engine/Source/Runtime/OpenGLDrv/Private/OpenGLDrv.h b/Engine/Source/Runtime/OpenGLDrv/Private/OpenGLDrv.h
--- a/Engine/Source/Runtime/OpenGLDrv/Private/OpenGLDrv.h
+++ b/Engine/Source/Runtime/OpenGLDrv/Private/OpenGLDrv.h
@@ -593,12 +593,9 @@ public:
 			}
 		};
 		TArray<FLockParams, TInlineAllocator<16> > OutstandingLocks;
-		uint32 TotalMemoryOutstanding;
 
 		FTextureLockTracker()
-		{
-			TotalMemoryOutstanding = 0;
-		}
+		{}
 
 		FORCEINLINE_DEBUGGABLE void Lock(void* RHIBuffer, void* Buffer, uint32 ArrayIndex, uint32 MipIndex, uint32 Stride, uint32 SizeRHI, EResourceLockMode LockMode)
 		{
@@ -609,7 +606,6 @@ public:
 			}
 //#endif
 			OutstandingLocks.Add(FLockParams(RHIBuffer, Buffer, ArrayIndex, MipIndex, Stride, SizeRHI, LockMode));
-			TotalMemoryOutstanding += SizeRHI;
 		}
 
 		FORCEINLINE_DEBUGGABLE FLockParams Unlock(void* RHIBuffer, uint32 ArrayIndex, uint32 MipIndex)
@@ -621,7 +617,6 @@ public:
 				{
 					FLockParams Result = OutstandingLocks[Index];
 					OutstandingLocks.RemoveAtSwap(Index, EAllowShrinking::No);
-					TotalMemoryOutstanding -= Result.BufferSize;
 					return Result;
 				}
 			}
diff --git a/Engine/Source/Runtime/OpenGLDrv/Private/OpenGLTexture.cpp b/Engine/Source/Runtime/OpenGLDrv/Private/OpenGLTexture.cpp
--- a/Engine/Source/Runtime/OpenGLDrv/Private/OpenGLTexture.cpp
+++ b/Engine/Source/Runtime/OpenGLDrv/Private/OpenGLTexture.cpp
@@ -1928,7 +1928,6 @@ void FOpenGLDynamicRHI::RHIUnlockTexture(FRHICommandListImmediate& RHICmdList, c
 	FTextureLockTracker::FLockParams Params = GLLockTracker.Unlock(Texture, ArrayIndex, Arguments.MipIndex);
 	if (Params.LockMode == RLM_ReadOnly)
 	{
-		GLLockTracker.TotalMemoryOutstanding = 0;
 		RHICmdList.EnqueueLambda([this, Texture, MipIndex = Arguments.MipIndex, ArrayIndex](FRHICommandListImmediate&)
 		{
 			Texture->Unlock(MipIndex, ArrayIndex);
diff --git a/Engine/Source/Runtime/RHI/Private/RHICommandList.cpp b/Engine/Source/Runtime/RHI/Private/RHICommandList.cpp
--- a/Engine/Source/Runtime/RHI/Private/RHICommandList.cpp
+++ b/Engine/Source/Runtime/RHI/Private/RHICommandList.cpp
@@ -85,11 +85,6 @@ static TAutoConsoleVariable<int32> CVarRHICmdBufferWriteLocks(
 	1,
 	TEXT("Only relevant with an RHI thread. Debugging option to diagnose problems with buffered locks."));
 
-static TAutoConsoleVariable<int32> CVarRHICmdMaxOutstandingMemoryBeforeFlush(
-	TEXT("r.RHICmdMaxOutstandingMemoryBeforeFlush"),
-	256,
-	TEXT("In kilobytes. The amount of outstanding memory before the RHI will force a flush. This should generally be set high enough that it doesn't happen on typical frames."));
-
 FAutoConsoleTaskPriority CPrio_SceneRenderingTask(
 	TEXT("TaskGraph.TaskPriorities.SceneRenderingTask"),
 	TEXT("Task and thread priority for various scene rendering tasks."),
@@ -1879,7 +1874,6 @@ void FDynamicRHI::RHIUnlockBuffer(class FRHICommandListBase& RHICmdList, FRHIBuf
 
 			FRHICommandListScopedFlushAndExecute Flush(RHICmdList.GetAsImmediate());
 			GDynamicRHI->UnlockBuffer_BottomOfPipe(RHICmdList, Buffer);
-			GLockTracker.TotalMemoryOutstanding = 0;
 		}
 		else
 		{
@@ -1898,14 +1892,6 @@ void FDynamicRHI::RHIUnlockBuffer(class FRHICommandListBase& RHICmdList, FRHIBuf
 			});
 			RHICmdList.RHIThreadFence(true);
 		}
-
-		if (RHICmdList.IsImmediate() && GLockTracker.TotalMemoryOutstanding > uint32(CVarRHICmdMaxOutstandingMemoryBeforeFlush.GetValueOnRenderThread()) * 1024u)
-		{
-			QUICK_SCOPE_CYCLE_COUNTER(STAT_RHIMETHOD_UnlockBuffer_FlushForMem);
-			// we could be loading a level or something, lets get this stuff going
-			RHICmdList.GetAsImmediate().ImmediateFlush(EImmediateFlushType::DispatchToRHIThread); 
-			GLockTracker.TotalMemoryOutstanding = 0;
-		}
 	}
 	else
 	{
diff --git a/Engine/Source/Runtime/RHI/Public/RHICommandList.h b/Engine/Source/Runtime/RHI/Public/RHICommandList.h
--- a/Engine/Source/Runtime/RHI/Public/RHICommandList.h
+++ b/Engine/Source/Runtime/RHI/Public/RHICommandList.h
@@ -321,11 +321,9 @@ struct FLockTracker
 
 	FCriticalSection CriticalSection;
 	TArray<FLockParams, TInlineAllocator<16> > OutstandingLocks;
-	uint32 TotalMemoryOutstanding;
 
 	FLockTracker()
 	{
-		TotalMemoryOutstanding = 0;
 	}
 
 	FORCEINLINE_DEBUGGABLE void Lock(void* RHIBuffer, void* Buffer, uint32 Offset, uint32 SizeRHI, EResourceLockMode LockMode)
@@ -338,7 +336,6 @@ struct FLockTracker
 		}
 #endif
 		OutstandingLocks.Add(FLockParams(RHIBuffer, Buffer, Offset, SizeRHI, LockMode));
-		TotalMemoryOutstanding += SizeRHI;
 	}
 	FORCEINLINE_DEBUGGABLE FLockParams Unlock(void* RHIBuffer)
 	{
diff --git a/Engine/Source/Runtime/RHI/Public/RHILockTracker.h b/Engine/Source/Runtime/RHI/Public/RHILockTracker.h
--- a/Engine/Source/Runtime/RHI/Public/RHILockTracker.h
+++ b/Engine/Source/Runtime/RHI/Public/RHILockTracker.h
@@ -93,13 +93,10 @@ struct FRHILockTracker
 	};
 
 	TArray<FLockParams, TInlineAllocator<16> > OutstandingLocks;
-	uint32 TotalMemoryOutstanding;
 	TArray<FUnlockFenceParams, TInlineAllocator<16> > OutstandingUnlocks;
 
 	FRHILockTracker()
-	{
-		TotalMemoryOutstanding = 0;
-	}
+	{}
 
 	FORCEINLINE_DEBUGGABLE void Lock(void* RHIBuffer, void* Buffer, uint32 Offset, uint32 SizeRHI, EResourceLockMode LockMode, bool bInDirectBufferWrite = false, bool bInCreateLock = false)
 	{
@@ -110,7 +107,6 @@ struct FRHILockTracker
 		}
 #endif
 		OutstandingLocks.Add(FLockParams(RHIBuffer, Buffer, Offset, SizeRHI, LockMode, bInDirectBufferWrite, bInCreateLock));
-		TotalMemoryOutstanding += SizeRHI;
 	}
 	FORCEINLINE_DEBUGGABLE FLockParams Unlock(void* RHIBuffer, uint32 Offset = 0)
 	{
