[LLM] Reset asset tags for ResourceTracking internal array reallocation

#rnx

diff --git a/Engine/Source/Runtime/RHI/Private/RHI.cpp b/Engine/Source/Runtime/RHI/Private/RHI.cpp
--- a/Engine/Source/Runtime/RHI/Private/RHI.cpp
+++ b/Engine/Source/Runtime/RHI/Private/RHI.cpp
@@ -7,11 +7,13 @@
 #include "RHI.h"
 #include "Async/ParallelFor.h"
 #include "HAL/IConsoleManager.h"
+#include "HAL/LowLevelMemTracker.h"
 #include "RHITransientResourceAllocator.h"
 #include "Misc/CommandLine.h"
 #include "Modules/ModuleManager.h"
 #include "Misc/ConfigCacheIni.h"
 #include "ProfilingDebugging/CsvProfiler.h"
+#include "ProfilingDebugging/MetadataTrace.h"
 #include "RHIFwd.h"
 #include "String/LexFromString.h"
 #include "RHIStrings.h"
@@ -522,6 +524,9 @@ void FRHIResource::BeginTrackingResource(FRHIResource* InResource)
 	if (GRHITrackingResources)
 	{
 		LLM_SCOPE_BYNAME(TEXT("RHIMisc/ResourceTracking"));
+		LLM_TAGSET_SCOPE_CLEAR(ELLMTagSet::Assets);
+		LLM_TAGSET_SCOPE_CLEAR(ELLMTagSet::AssetClasses);
+		UE_TRACE_METADATA_CLEAR_SCOPE();
 
 		FScopeLock Lock(&GRHIResourceTrackingCriticalSection);
 
@@ -535,6 +540,10 @@ void FRHIResource::EndTrackingResource(FRHIResource* InResource)
 {
 	if (InResource->bBeingTracked)
 	{
+		LLM_SCOPE_BYNAME(TEXT("RHIMisc/ResourceTracking"));
+		LLM_TAGSET_SCOPE_CLEAR(ELLMTagSet::Assets);
+		LLM_TAGSET_SCOPE_CLEAR(ELLMTagSet::AssetClasses);
+		UE_TRACE_METADATA_CLEAR_SCOPE();
 		FScopeLock Lock(&GRHIResourceTrackingCriticalSection);
 		GRHITrackedResources.Remove(InResource);
 		InResource->bBeingTracked = false;
@@ -548,6 +557,10 @@ void FRHIResource::StartTrackingAllResources()
 
 void FRHIResource::StopTrackingAllResources()
 {
+	LLM_SCOPE_BYNAME(TEXT("RHIMisc/ResourceTracking"));
+	LLM_TAGSET_SCOPE_CLEAR(ELLMTagSet::Assets);
+	LLM_TAGSET_SCOPE_CLEAR(ELLMTagSet::AssetClasses);
+	UE_TRACE_METADATA_CLEAR_SCOPE();
 	FScopeLock Lock(&GRHIResourceTrackingCriticalSection);
 	for (FRHIResource* Resource : GRHITrackedResources)
 	{
