VulkanRHI: Keep CreationInfoData for the entire lifetime of the transition when NeedsAllPlanes is active.

#rnx

diff --git a/Engine/Source/Runtime/RHI/Private/RHITransition.cpp b/Engine/Source/Runtime/RHI/Private/RHITransition.cpp
--- a/Engine/Source/Runtime/RHI/Private/RHITransition.cpp
+++ b/Engine/Source/Runtime/RHI/Private/RHITransition.cpp
@@ -63,13 +63,12 @@ void FRHITransition::CreatePlatformTransition(bool bCleanupInfoData) const
 	}
 }
 
-void FRHITransition::CreatePlatformTransitionThreadSafe() const
+void FRHITransition::CreatePlatformTransitionThreadSafe(bool bCleanupInfoData) const
 {
 	checkSlow(GRHIGlobals.EnableHighLevelTransitionPatching);
 	
 	UE::TUniqueLock Lock(TransitionCreationMutex);
 
-	const bool bCleanupInfoData = true;
 	CreatePlatformTransition(bCleanupInfoData);
 }
 
diff --git a/Engine/Source/Runtime/RHI/Public/RHICommandListCommandExecutes.inl b/Engine/Source/Runtime/RHI/Public/RHICommandListCommandExecutes.inl
--- a/Engine/Source/Runtime/RHI/Public/RHICommandListCommandExecutes.inl
+++ b/Engine/Source/Runtime/RHI/Public/RHICommandListCommandExecutes.inl
@@ -384,7 +384,8 @@ void FRHICommandBeginTransitions::Execute(FRHICommandListBase& CmdList)
 	{
 		for (const FRHITransition* Transition : Transitions)
 		{
-			Transition->CreatePlatformTransitionThreadSafe();
+			const bool bCleanupInfoData = !CmdList.NeedsAllPlanes();
+			Transition->CreatePlatformTransitionThreadSafe(bCleanupInfoData);
 			checkSlow(Transition->IsPlatformLevelCreated());
 		}
 	}
@@ -406,7 +407,8 @@ void FRHICommandEndTransitions::Execute(FRHICommandListBase& CmdList)
 	{
 		for (const FRHITransition* Transition : Transitions)
 		{
-			Transition->CreatePlatformTransitionThreadSafe();
+			const bool bCleanupInfoData = !CmdList.NeedsAllPlanes();
+			Transition->CreatePlatformTransitionThreadSafe(bCleanupInfoData);
 			checkSlow(Transition->IsPlatformLevelCreated());
 		}
 	}
@@ -426,7 +428,8 @@ void FRHICommandResourceTransition::Execute(FRHICommandListBase& CmdList)
 	// Create FRHITransition from FRHITransitionHighLevel with Platform information
 	if (GRHIGlobals.EnableHighLevelTransitionPatching)
 	{
-		Transition->CreatePlatformTransitionThreadSafe();
+		const bool bCleanupInfoData = true;
+		Transition->CreatePlatformTransitionThreadSafe(bCleanupInfoData);
 	}
 
 	checkSlow(Transition->IsPlatformLevelCreated());
diff --git a/Engine/Source/Runtime/RHI/Public/RHITransition.h b/Engine/Source/Runtime/RHI/Public/RHITransition.h
--- a/Engine/Source/Runtime/RHI/Public/RHITransition.h
+++ b/Engine/Source/Runtime/RHI/Public/RHITransition.h
@@ -862,7 +862,10 @@ private:
 
 	~FRHITransition()
 	{
-		checkSlow(CreationInfoData == nullptr);
+		if (CreationInfoData)
+		{
+			CleanupInfoData();
+		}
 	}
 
 	// Give private access to specific functions/RHI commands that need to allocate or control transitions.
@@ -885,7 +888,7 @@ private:
 	}
 	
 	RHI_API void CreatePlatformTransition(bool bCleanupInfoData) const;
-	RHI_API void CreatePlatformTransitionThreadSafe() const;
+	RHI_API void CreatePlatformTransitionThreadSafe(bool bCleanupInfoData) const;
 
 	inline void MarkBegin(ERHIPipeline Pipeline) const
 	{
