Fix CreationInfoData check from FRHITransition::Cleanup for Vulkan RHI

diff --git a/Engine/Source/Runtime/RHI/Private/RHITransition.cpp b/Engine/Source/Runtime/RHI/Private/RHITransition.cpp
--- a/Engine/Source/Runtime/RHI/Private/RHITransition.cpp
+++ b/Engine/Source/Runtime/RHI/Private/RHITransition.cpp
@@ -76,7 +76,7 @@ void FRHITransition::Cleanup() const
 {
 	FRHITransition* Transition = const_cast<FRHITransition*>(this);
 	RHIReleaseTransition(Transition);
-	checkSlow(CreationInfoData == nullptr);
+	checkSlow((CreationInfoData == nullptr) || GRHIGlobals.NeedsAllPlanes);
 
 	// Explicit destruction of the transition.
 	Transition->~FRHITransition();
