Gpu Fence visualisation when selecting an timing event in GPU Queue tracks.

diff --git a/Engine/Source/Developer/TraceInsights/Private/Insights/TimingProfiler/ViewModels/GpuFenceRelation.cpp b/Engine/Source/Developer/TraceInsights/Private/Insights/TimingProfiler/ViewModels/GpuFenceRelation.cpp
new file mode 100644
--- /dev/null
+++ b/Engine/Source/Developer/TraceInsights/Private/Insights/TimingProfiler/ViewModels/GpuFenceRelation.cpp
@@ -0,0 +1,191 @@
+// Copyright Epic Games, Inc. All Rights Reserved.
+
+#include "GpuFenceRelation.h"
+
+// TraceInsightsCore
+#include "InsightsCore/Common/PaintUtils.h"
+
+// TraceInsights
+#include "Insights/TimingProfiler/TimingProfilerManager.h"
+#include "Insights/ViewModels/TimingTrackViewport.h"
+#include "Insights/ViewModels/TimingViewDrawHelper.h"
+
+namespace UE::Insights::TimingProfiler
+{
+
+////////////////////////////////////////////////////////////////////////////////////////////////////
+
+INSIGHTS_IMPLEMENT_RTTI(FGpuFenceRelation)
+
+////////////////////////////////////////////////////////////////////////////////////////////////////
+
+FGpuFenceRelation::FGpuFenceRelation(double InSourceTime, int32 InSourceQueueId, double InTargetTime, int32 InTargetQueueId)
+{
+	SourceTime = InSourceTime;
+	SourceQueueId = InSourceQueueId;
+	TargetTime = InTargetTime;
+	TargetQueueId = InTargetQueueId;
+}
+
+////////////////////////////////////////////////////////////////////////////////////////////////////
+
+void FGpuFenceRelation::Draw(const FDrawContext& DrawContext, const FTimingTrackViewport& Viewport, const ITimingViewDrawHelper& Helper, const ITimingEventRelation::EDrawFilter Filter)
+{
+	int32 LayerId = Helper.GetRelationLayerId();
+
+	TSharedPtr<const FBaseTimingTrack> SourceTrackShared = SourceTrack.Pin();
+	TSharedPtr<const FBaseTimingTrack> TargetTrackShared = TargetTrack.Pin();
+
+	if (!SourceTrackShared.IsValid() || !TargetTrackShared.IsValid())
+	{
+		return;
+	}
+
+	if (Filter == ITimingEventRelation::EDrawFilter::BetweenScrollableTracks)
+	{
+		if (SourceTrackShared->GetLocation() != ETimingTrackLocation::Scrollable ||
+			TargetTrackShared->GetLocation() != ETimingTrackLocation::Scrollable)
+		{
+			return;
+		}
+	}
+
+	if (Filter == ITimingEventRelation::EDrawFilter::BetweenDockedTracks)
+	{
+		if (SourceTrackShared->GetLocation() == ETimingTrackLocation::Scrollable &&
+			TargetTrackShared->GetLocation() == ETimingTrackLocation::Scrollable)
+		{
+			return;
+		}
+
+		LayerId = DrawContext.LayerId;
+	}
+
+	const int32 OutlineLayerId = LayerId - 1;
+
+	float X1 = Viewport.TimeToSlateUnitsRounded(SourceTime);
+	float X2 = Viewport.TimeToSlateUnitsRounded(TargetTime);
+	if (FMath::Max(X1, X2) < 0.0f || FMath::Min(X1, X2) > Viewport.GetWidth())
+	{
+		return;
+	}
+
+	if (!SourceTrackShared->IsVisible() && !TargetTrackShared->IsVisible())
+	{
+		return;
+	}
+
+	const int32 MaxEventDepth = (int32)TimingProfiler::FTimingProfilerManager::Get()->GetEventDepthLimit() - 1;
+
+	float Y1 = 0.0f;
+	float Y2 = 0.0f;
+	if (SourceTrackShared->IsVisible())
+	{
+		int32 ActualSourceDepth = FMath::Min(SourceDepth, MaxEventDepth);
+		Y1 = SourceTrackShared->GetPosY();
+		Y1 += Viewport.GetLayout().GetLaneY(ActualSourceDepth) + Viewport.GetLayout().EventH / 2.0f;
+		Y1 += SourceTrackShared->GetChildTracksTopHeight(Viewport.GetLayout());
+	}
+	else
+	{
+		Y1 = TargetTrackShared->GetPosY();
+	}
+
+	if (TargetTrackShared->IsVisible())
+	{
+		int32 ActualTargetDepth = FMath::Min(TargetDepth, MaxEventDepth);
+		Y2 = TargetTrackShared->GetPosY();
+		Y2 += Viewport.GetLayout().GetLaneY(ActualTargetDepth) + Viewport.GetLayout().EventH / 2.0f;
+		Y2 += TargetTrackShared->GetChildTracksTopHeight(Viewport.GetLayout());
+	}
+	else
+	{
+		Y2 = SourceTrackShared->GetPosY();
+	}
+
+	const FVector2D StartPoint = FVector2D(X1, Y1);
+	const FVector2D EndPoint = FVector2D(X2, Y2);
+	const double Distance = FVector2D::Distance(StartPoint, EndPoint);
+
+	constexpr double LineHeightAtStart = 4.0;
+	constexpr double LineLengthAtStart = 4.0;
+	constexpr double LineLengthAtEnd = 12.0;
+
+	const FVector2D StartDir(FMath::Max(static_cast<double>(X2 - X1), 4.0 * (LineLengthAtStart + LineLengthAtEnd)), 0.0);
+
+	constexpr float OutlineThickness = 5.0f;
+	constexpr float LineThickness = 3.0f;
+
+	constexpr double ArrowDirectionLen = 10.0;
+	constexpr double ArrowRotationAngle = 20.0;
+	FVector2D ArrowDirection(-ArrowDirectionLen, 0.0);
+
+	auto ToLiniarColorNoAlpha = [](uint32 Value)
+		{
+			const float R = static_cast<float>((Value & 0xFF000000) >> 24) / 255.0f;
+			const float G = static_cast<float>((Value & 0x00FF0000) >> 16) / 255.0f;
+			const float B = static_cast<float>((Value & 0x0000FF00) >> 8) / 255.0f;
+			return FLinearColor(R, G, B);
+		};
+
+	const FLinearColor OutlineColor(0.0f, 0.0f, 0.0f, 1.0f);
+	const FLinearColor Color = ToLiniarColorNoAlpha(0xFF0000FF); // Red
+
+	TArray<FVector2D> LinePoints;
+	LinePoints.Add(StartPoint + FVector2D(0.0, -LineHeightAtStart / 2.0));
+	LinePoints.Add(StartPoint + FVector2D(0.0, +LineHeightAtStart / 2.0));
+	DrawContext.DrawLines(OutlineLayerId, 0.0f, 0.0f, LinePoints, OutlineColor, /*bAntialias=*/ true, OutlineThickness);
+	DrawContext.DrawLines(LayerId, 0.0f, 0.0f, LinePoints, Color, /*bAntialias=*/ true, LineThickness);
+
+	constexpr double MinDistance = 1.5 * (LineLengthAtStart + LineLengthAtEnd);
+	constexpr double MaxDistance = 10000.0; // arbitrary limit to avoid stack overflow in recursive FLineBuilder::Subdivide when rendering splines
+	if (Distance > MinDistance && Distance < MaxDistance && !FMath::IsNearlyEqual(StartPoint.Y, EndPoint.Y))
+	{
+		FVector2D SplineStart(StartPoint.X + LineLengthAtStart, StartPoint.Y);
+		FVector2D SplineEnd(EndPoint.X - LineLengthAtEnd, EndPoint.Y);
+		DrawContext.DrawSpline(OutlineLayerId, 0.0f, 0.0f, SplineStart, StartDir, SplineEnd, StartDir, OutlineThickness, OutlineColor);
+		DrawContext.DrawSpline(LayerId, 0.0f, 0.0f, SplineStart, StartDir, SplineEnd, StartDir, LineThickness, Color);
+
+		LinePoints.Empty();
+		LinePoints.Add(StartPoint);
+		LinePoints.Add(SplineStart);
+		DrawContext.DrawLines(OutlineLayerId, 0.0f, 0.0f, LinePoints, OutlineColor, /*bAntialias=*/ true, OutlineThickness);
+		DrawContext.DrawLines(LayerId, 0.0f, 0.0f, LinePoints, Color, /*bAntialias=*/ true, LineThickness);
+
+		LinePoints.Empty();
+		LinePoints.Add(SplineEnd);
+		LinePoints.Add(EndPoint);
+		DrawContext.DrawLines(OutlineLayerId, 0.0f, 0.0f, LinePoints, OutlineColor, /*bAntialias=*/ true, OutlineThickness);
+		DrawContext.DrawLines(LayerId, 0.0f, 0.0f, LinePoints, Color, /*bAntialias=*/ true, LineThickness);
+	}
+	else
+	{
+		LinePoints.Empty();
+		LinePoints.Add(StartPoint);
+		LinePoints.Add(EndPoint);
+		DrawContext.DrawLines(OutlineLayerId, 0.0f, 0.0f, LinePoints, OutlineColor, /*bAntialias=*/ true, OutlineThickness);
+		DrawContext.DrawLines(LayerId, 0.0f, 0.0f, LinePoints, Color, /*bAntialias=*/ true, LineThickness);
+
+		ArrowDirection = StartPoint - EndPoint;
+		ArrowDirection.Normalize();
+		ArrowDirection *= ArrowDirectionLen;
+	}
+
+	FVector2D ArrowOrigin = EndPoint;
+
+	LinePoints.Empty();
+	LinePoints.Add(ArrowOrigin);
+	LinePoints.Add(ArrowOrigin + ArrowDirection.GetRotated(-ArrowRotationAngle));
+	DrawContext.DrawLines(OutlineLayerId, 0.0f, 0.0f, LinePoints, OutlineColor, /*bAntialias=*/ true, OutlineThickness);
+	DrawContext.DrawLines(LayerId, 0.0f, 0.0f, LinePoints, Color, /*bAntialias=*/ true, LineThickness);
+
+	LinePoints.Empty();
+	LinePoints.Add(ArrowOrigin);
+	LinePoints.Add(ArrowOrigin + ArrowDirection.GetRotated(ArrowRotationAngle));
+	DrawContext.DrawLines(OutlineLayerId, 0.0f, 0.0f, LinePoints, OutlineColor, /*bAntialias=*/ true, OutlineThickness);
+	DrawContext.DrawLines(LayerId, 0.0f, 0.0f, LinePoints, Color, /*bAntialias=*/ true, LineThickness);
+}
+
+////////////////////////////////////////////////////////////////////////////////////////////////////
+
+} // namespace UE::Insights::TimingProfiler
diff --git a/Engine/Source/Developer/TraceInsights/Private/Insights/TimingProfiler/ViewModels/GpuFenceRelation.h b/Engine/Source/Developer/TraceInsights/Private/Insights/TimingProfiler/ViewModels/GpuFenceRelation.h
new file mode 100644
--- /dev/null
+++ b/Engine/Source/Developer/TraceInsights/Private/Insights/TimingProfiler/ViewModels/GpuFenceRelation.h
@@ -0,0 +1,50 @@
+// Copyright Epic Games, Inc. All Rights Reserved.
+
+#pragma once
+
+#include "CoreTypes.h"
+
+#include "Insights/ViewModels/ITimingEvent.h"
+
+namespace UE::Insights::TimingProfiler
+{
+
+class FGpuFenceRelation : public ITimingEventRelation
+{
+	INSIGHTS_DECLARE_RTTI(FGpuFenceRelation, ITimingEventRelation)
+
+public:
+	FGpuFenceRelation(double InSourceTime, int32 InSourceQueueId, double InTargetTime, int32 InTargetQueueId);
+	virtual ~FGpuFenceRelation() {}
+
+	virtual void Draw(const FDrawContext& DrawContext, const FTimingTrackViewport& Viewport, const ITimingViewDrawHelper& Helper, const ITimingEventRelation::EDrawFilter Filter) override;
+
+	void SetSourceTrack(TSharedPtr<const FBaseTimingTrack> InSourceTrack) { SourceTrack = InSourceTrack; }
+	TSharedPtr<const FBaseTimingTrack> GetSourceTrack() { return SourceTrack.Pin(); }
+
+	void SetTargetTrack(TSharedPtr<const FBaseTimingTrack> InTargetTrack) { TargetTrack = InTargetTrack; }
+	TSharedPtr<const FBaseTimingTrack> GetTargetTrack() { return TargetTrack.Pin(); }
+
+	double GetSourceTime() { return SourceTime; }
+	int32 GetSourceQueueId() { return SourceQueueId; }
+	void SetSourceDepth(int32 InDepth) { SourceDepth = InDepth; }
+	int32 GetSourceDepth() { return SourceDepth; }
+
+	double GetTargetTime() { return TargetTime; }
+	int32 GetTargetQueueId() { return TargetQueueId; }
+	void SetTargetDepth(int32 InDepth) { TargetDepth = InDepth; }
+	int32 GetTargetDepth() { return TargetDepth; }
+
+private:
+	double SourceTime;
+	int32 SourceQueueId;
+	int32 SourceDepth = -1;
+	double TargetTime;
+	int32 TargetQueueId;
+	int32 TargetDepth = -1;
+
+	TWeakPtr<const FBaseTimingTrack> SourceTrack;
+	TWeakPtr<const FBaseTimingTrack> TargetTrack;
+};
+
+} // namespace UE::Insights::TimingProfiler
diff --git a/Engine/Source/Developer/TraceInsights/Private/Insights/TimingProfiler/ViewModels/ThreadTimingSharedState.cpp b/Engine/Source/Developer/TraceInsights/Private/Insights/TimingProfiler/ViewModels/ThreadTimingSharedState.cpp
--- a/Engine/Source/Developer/TraceInsights/Private/Insights/TimingProfiler/ViewModels/ThreadTimingSharedState.cpp
+++ b/Engine/Source/Developer/TraceInsights/Private/Insights/TimingProfiler/ViewModels/ThreadTimingSharedState.cpp
@@ -23,6 +23,7 @@
 #include "Insights/ITimingViewSession.h"
 #include "Insights/TimingProfiler/Tracks/CpuTimingTrack.h"
 #include "Insights/TimingProfiler/Tracks/GpuTimingTrack.h"
+#include "Insights/TimingProfiler/ViewModels/GpuFenceRelation.h"
 #include "Insights/Widgets/STimingView.h"
 
 #define LOCTEXT_NAMESPACE "UE::Insights::TimingProfiler::ThreadTiming"
@@ -279,6 +280,8 @@ void FThreadTimingSharedState::OnBeginSession(Timing::ITimingViewSession& InSess
 
 	TimingProfilerTimelineCount = 0;
 	LoadTimeProfilerTimelineCount = 0;
+
+	TimingView->OnSelectedEventChanged().AddSP(this, &FThreadTimingSharedState::OnTimingEventSelected);
 }
 
 ////////////////////////////////////////////////////////////////////////////////////////////////////
@@ -301,6 +304,8 @@ void FThreadTimingSharedState::OnEndSession(Timing::ITimingViewSession& InSessio
 
 	TimingProfilerTimelineCount = 0;
 	LoadTimeProfilerTimelineCount = 0;
+
+	TimingView->OnSelectedEventChanged().RemoveAll(this);
 }
 
 ////////////////////////////////////////////////////////////////////////////////////////////////////
@@ -847,6 +852,120 @@ void FThreadTimingSharedState::RemoveGpuWorkChildTracks()
 
 ////////////////////////////////////////////////////////////////////////////////////////////////////
 
+void FThreadTimingSharedState::OnTimingEventSelected(TSharedPtr<const ITimingEvent> InSelectedEvent)
+{
+	if (TimingView)
+	{
+		TimingView->EditCurrentRelations().RemoveAll([](TUniquePtr<ITimingEventRelation>& Item)
+			{
+				return Item->Is<FGpuFenceRelation>();
+			});
+	}
+
+	if (!InSelectedEvent.IsValid())
+	{
+		return;
+	}
+
+	TSharedRef<const FBaseTimingTrack> BaseTrack = InSelectedEvent->GetTrack();
+
+	uint32 QueueId = 0;
+	if (BaseTrack->Is<FGpuQueueTimingTrack>())
+	{
+		const FGpuQueueTimingTrack& GpuQueueTrack = BaseTrack->As<const FGpuQueueTimingTrack>();
+		QueueId = GpuQueueTrack.GetThreadId();
+	}
+	else if (BaseTrack->Is<FGpuQueueWorkTimingTrack>())
+	{
+		const FGpuQueueWorkTimingTrack& GpuQueueWorkTrack = BaseTrack->As<const FGpuQueueWorkTimingTrack>();
+		QueueId = GpuQueueWorkTrack.GetThreadId();
+	}
+	else
+	{
+		return;
+	}
+
+	using namespace TraceServices;
+
+	TSharedPtr<const TraceServices::IAnalysisSession> Session = FInsightsManager::Get()->GetSession();
+	if (!Session.IsValid())
+	{
+		return;
+	}
+
+	const TraceServices::ITimingProfilerProvider* TimingProfilerProvider = TraceServices::ReadTimingProfilerProvider(*Session);
+
+	if (TimingProfilerProvider == nullptr)
+	{
+		return;
+	}
+
+	auto AddFenceRelation = [this, TimingProfilerProvider](const FGpuSignalFence& SignalFence, const FGpuWaitFence& WaitFence, uint32 WaitFenceQueueId)
+	{
+		uint32 Index = 0;
+		uint32 SignalFenceQueueId = WaitFence.QueueToWaitForId;
+		bool Result = TimingProfilerProvider->GetGpuQueueTimelineIndex(SignalFenceQueueId, Index);
+		if (!Result)
+		{
+			return;
+		}
+
+		int32 SourceDepth = 0;
+		TimingProfilerProvider->ReadTimeline(Index, [SignalFence, &SourceDepth](const TraceServices::ITimingProfilerProvider::Timeline& Timeline)
+		{
+			SourceDepth = FMath::Max(Timeline.GetDepthAt(SignalFence.Timestamp) - 1, 0);
+		});
+
+		Result = TimingProfilerProvider->GetGpuQueueTimelineIndex(WaitFenceQueueId, Index);
+		if (!Result)
+		{
+			return;
+		}
+
+		int32 TargetDepth = 0;
+		TimingProfilerProvider->ReadTimeline(Index, [WaitFence, &TargetDepth](const TraceServices::ITimingProfilerProvider::Timeline& Timeline)
+		{
+			TargetDepth = FMath::Max(Timeline.GetDepthAt(WaitFence.Timestamp) - 1, 0);
+		});
+
+		TUniquePtr<ITimingEventRelation> RelationBase = MakeUnique<FGpuFenceRelation>(SignalFence.Timestamp, SignalFenceQueueId, WaitFence.Timestamp, WaitFenceQueueId);
+		FGpuFenceRelation& Relation = RelationBase->As<FGpuFenceRelation>();
+
+		Relation.SetSourceDepth(SourceDepth);
+		Relation.SetTargetDepth(TargetDepth);
+
+		TSharedPtr<FGpuQueueTimingTrack>* Track = GpuTracks.Find(SignalFenceQueueId);
+		if (Track == nullptr)
+		{
+			return;
+		}
+		Relation.SetSourceTrack(*Track);
+
+		Track = GpuTracks.Find(WaitFenceQueueId);
+		if (Track == nullptr)
+		{
+			return;
+		}
+		Relation.SetTargetTrack(*Track);
+
+		if (TimingView)
+		{
+			TimingView->AddRelation(RelationBase);
+		}
+	};
+
+	TraceServices::FAnalysisSessionReadScope SessionReadScope(*Session);
+
+	TimingProfilerProvider->EnumerateResolvedGpuFences(QueueId, InSelectedEvent->GetStartTime(), InSelectedEvent->GetEndTime(), 
+		[AddFenceRelation, TimingProfilerProvider, QueueId](uint32 SignalFenceQueueId, const FGpuSignalFence& SignalFence, uint32 WaitFenceQueueId, const FGpuWaitFence& WaitFence)
+	{
+		AddFenceRelation(SignalFence, WaitFence, WaitFenceQueueId);
+		return EEnumerateResult::Continue;
+	});
+}
+
+////////////////////////////////////////////////////////////////////////////////////////////////////
+
 } // namespace UE::Insights::TimingProfiler
 
 #undef LOCTEXT_NAMESPACE
diff --git a/Engine/Source/Developer/TraceInsights/Private/Insights/TimingProfiler/ViewModels/ThreadTimingSharedState.h b/Engine/Source/Developer/TraceInsights/Private/Insights/TimingProfiler/ViewModels/ThreadTimingSharedState.h
--- a/Engine/Source/Developer/TraceInsights/Private/Insights/TimingProfiler/ViewModels/ThreadTimingSharedState.h
+++ b/Engine/Source/Developer/TraceInsights/Private/Insights/TimingProfiler/ViewModels/ThreadTimingSharedState.h
@@ -116,6 +116,8 @@ private:
 	void AddGpuWorkChildTracks();
 	void RemoveGpuWorkChildTracks();
 
+	void OnTimingEventSelected(TSharedPtr<const ITimingEvent> InSelectedEvent);
+
 private:
 	STimingView* TimingView = nullptr;
 
diff --git a/Engine/Source/Developer/TraceServices/Private/Analyzers/GpuProfilerTraceAnalysis.cpp b/Engine/Source/Developer/TraceServices/Private/Analyzers/GpuProfilerTraceAnalysis.cpp
--- a/Engine/Source/Developer/TraceServices/Private/Analyzers/GpuProfilerTraceAnalysis.cpp
+++ b/Engine/Source/Developer/TraceServices/Private/Analyzers/GpuProfilerTraceAnalysis.cpp
@@ -32,6 +32,8 @@ void FGpuProfilerAnalyzer::OnAnalysisBegin(const FOnAnalysisContext& Context)
 	Builder.RouteEvent(RouteId_EventBeginWork, "GpuProfiler", "EventBeginWork");
 	Builder.RouteEvent(RouteId_EventEndWork, "GpuProfiler", "EventEndWork");
 	Builder.RouteEvent(RouteId_EventStats, "GpuProfiler", "EventStats");
+	Builder.RouteEvent(RouteId_SignalFence, "GpuProfiler", "SignalFence");
+	Builder.RouteEvent(RouteId_WaitFence, "GpuProfiler", "WaitFence");
 
 	// The Old GPU Profiler (deprecated in UE 5.6)
 	// Analysis code is maintained for backward compatibility with old traces.
@@ -122,6 +124,13 @@ bool FGpuProfilerAnalyzer::OnEvent(uint16 RouteId, EStyle Style, const FOnEventC
 	{
 		const uint32 QueueId = EventData.GetValue<uint32>("QueueId");
 		const uint64 GPUTimestampTOP = EventData.GetValue<uint64>("GPUTimestampTOP");
+		
+		// Ignore events for which the timestamp could not be determined.
+		if (GPUTimestampTOP == 0)
+		{
+			break;
+		}
+
 		const double Time = Context.EventTime.AsSeconds(GPUTimestampTOP);
 
 		FString Name;
@@ -138,6 +147,13 @@ bool FGpuProfilerAnalyzer::OnEvent(uint16 RouteId, EStyle Style, const FOnEventC
 	{
 		const uint32 QueueId = EventData.GetValue<uint32>("QueueId");
 		const uint64 GPUTimestampBOP = EventData.GetValue<uint64>("GPUTimestampBOP");
+
+		// Ignore events for which the timestamp could not be determined.
+		if (GPUTimestampBOP == 0)
+		{
+			break;
+		}
+
 		const double Time = Context.EventTime.AsSeconds(GPUTimestampBOP);
 
 		FString Name;
@@ -208,6 +224,39 @@ bool FGpuProfilerAnalyzer::OnEvent(uint16 RouteId, EStyle Style, const FOnEventC
 
 		break;
 	}
+	case RouteId_SignalFence:
+	{
+		uint32 QueueId = EventData.GetValue<uint32>("QueueId");
+		uint64 Timestamp = EventData.GetValue<uint64>("CPUTimestamp");
+		uint64 Value = EventData.GetValue<uint64>("Value");;
+
+		FQueue& Queue = GetOrAddQueue(QueueId);
+		FGpuSignalFence Fence;
+		Fence.Timestamp = Context.EventTime.AsSeconds(Timestamp);
+		Fence.Value = Value;
+
+		FAnalysisSessionEditScope _(Session);
+		TimingProfilerProvider.AddGpuSignalFence(QueueId, Fence);
+
+		break;
+	}
+	case RouteId_WaitFence:
+	{
+		uint32 QueueId = EventData.GetValue<uint32>("QueueId");
+		uint64 Timestamp = EventData.GetValue<uint64>("CPUTimestamp");
+		uint32 QueueToWaitForId = EventData.GetValue<uint32>("QueueToWaitForId");
+		uint64 Value = EventData.GetValue<uint64>("Value");;
+
+		FGpuWaitFence Fence;
+		Fence.Timestamp = Context.EventTime.AsSeconds(Timestamp);
+		Fence.Value = Value;
+		Fence.QueueToWaitForId = QueueToWaitForId;
+
+		FAnalysisSessionEditScope _(Session);
+		TimingProfilerProvider.AddGpuWaitFence(QueueId, Fence);
+
+		break;
+	}
 
 	//////////////////////////////////////////////////
 	// The Old GPU Profiler
diff --git a/Engine/Source/Developer/TraceServices/Private/Analyzers/GpuProfilerTraceAnalysis.h b/Engine/Source/Developer/TraceServices/Private/Analyzers/GpuProfilerTraceAnalysis.h
--- a/Engine/Source/Developer/TraceServices/Private/Analyzers/GpuProfilerTraceAnalysis.h
+++ b/Engine/Source/Developer/TraceServices/Private/Analyzers/GpuProfilerTraceAnalysis.h
@@ -26,6 +26,8 @@ private:
 		RouteId_EventBeginWork,
 		RouteId_EventEndWork,
 		RouteId_EventStats,
+		RouteId_SignalFence,
+		RouteId_WaitFence,
 
 		// The Old GPU Profiler
 		RouteId_EventSpec,
diff --git a/Engine/Source/Developer/TraceServices/Private/Model/TimingProfiler.cpp b/Engine/Source/Developer/TraceServices/Private/Model/TimingProfiler.cpp
--- a/Engine/Source/Developer/TraceServices/Private/Model/TimingProfiler.cpp
+++ b/Engine/Source/Developer/TraceServices/Private/Model/TimingProfiler.cpp
@@ -189,11 +189,14 @@ FTimingProfilerProvider::TimelineInternal& FTimingProfilerProvider::EditGpu2Time
 
 void FTimingProfilerProvider::AddGpuQueue(uint32 QueueId, uint8 GPU, uint8 Index, uint8 Type, const TCHAR* Name)
 {
+	Session.WriteAccessCheck();
+
 	const uint32* FoundQueueIndex = GpuQueueIdToQueueIndexMap.Find(QueueId);
 	if (!FoundQueueIndex)
 	{
 		GpuQueueIdToQueueIndexMap.Add(QueueId, (uint32)GpuQueues.Num());
 
+		GpuQueueData.AddDefaulted();
 		FGpuQueueInfo& GpuQueue = GpuQueues.AddDefaulted_GetRef();
 
 		GpuQueue.Id = QueueId;
@@ -212,6 +215,180 @@ void FTimingProfilerProvider::AddGpuQueue(uint32 QueueId, uint8 GPU, uint8 Index
 	}
 }
 
+void FTimingProfilerProvider::AddGpuSignalFence(uint32 QueueId, const FGpuSignalFence& Fence)
+{
+	Session.WriteAccessCheck();
+
+	const uint32* FoundQueueIndex = GpuQueueIdToQueueIndexMap.Find(QueueId);
+	if (FoundQueueIndex)
+	{
+		check(*FoundQueueIndex < (uint32)GpuQueueData.Num());
+
+		FGpuQueueData& Data = GpuQueueData[*FoundQueueIndex];
+
+		if (Data.SignalFenceArray.Num())
+		{
+			ensure(Fence.Timestamp >= Data.SignalFenceArray.Last().Timestamp);
+		}
+
+		Data.SignalFenceArray.Add(Fence);
+	}
+}
+
+void FTimingProfilerProvider::AddGpuWaitFence(uint32 QueueId, const FGpuWaitFence& Fence)
+{
+	Session.WriteAccessCheck();
+
+	const uint32* FoundQueueIndex = GpuQueueIdToQueueIndexMap.Find(QueueId);
+	if (FoundQueueIndex)
+	{
+		check(*FoundQueueIndex < (uint32)GpuQueueData.Num());
+
+		FGpuQueueData& Data = GpuQueueData[*FoundQueueIndex];
+
+		if (Data.WaitFenceArray.Num())
+		{
+			ensure(Fence.Timestamp >= Data.WaitFenceArray.Last().Timestamp);
+		}
+
+		Data.WaitFenceArray.Add(Fence);
+	}
+}
+
+void FTimingProfilerProvider::EnumerateGpuSignalFences(uint32 QueueId, double StartTime, double EndTime, EnumerateGpuSignalFencesCallback Callback) const
+{
+	Session.ReadAccessCheck();
+
+	const uint32* FoundQueueIndex = GpuQueueIdToQueueIndexMap.Find(QueueId);
+	if (FoundQueueIndex == nullptr)
+	{
+		return;
+	}
+
+	const FGpuQueueData& Data = GpuQueueData[*FoundQueueIndex];
+	const TArray<FGpuSignalFence>& Array = Data.SignalFenceArray;
+
+	int32 StartIndex = Algo::LowerBoundBy(Array, StartTime,
+		[](const FGpuSignalFence& Fence) { return Fence.Timestamp; });
+
+	for (int32 Index = StartIndex; Index < Array.Num() && Array[Index].Timestamp <= EndTime; ++Index)
+	{
+		if (Callback(Array[Index]) == EEnumerateResult::Stop)
+		{
+			break;
+		}
+	}
+}
+
+void FTimingProfilerProvider::EnumerateGpuWaitFences(uint32 QueueId, double StartTime, double EndTime, EnumerateGpuWaitFencesCallback Callback) const
+{
+	Session.ReadAccessCheck();
+
+	const uint32* FoundQueueIndex = GpuQueueIdToQueueIndexMap.Find(QueueId);
+	if (FoundQueueIndex == nullptr)
+	{
+		return;
+	}
+
+	const FGpuQueueData& Data = GpuQueueData[*FoundQueueIndex];
+	const TArray<FGpuWaitFence>& Array = Data.WaitFenceArray;
+
+	int32 StartIndex = Algo::LowerBoundBy(Array, StartTime,
+		[](const FGpuWaitFence& Fence) { return Fence.Timestamp; });
+
+	for (int32 Index = StartIndex; Index < Array.Num() && Array[Index].Timestamp <= EndTime; ++Index)
+	{
+		if (Callback(Array[Index]) == EEnumerateResult::Stop)
+		{
+			break;
+		}
+	}
+}
+
+void FTimingProfilerProvider::EnumerateResolvedGpuFences(uint32 QueueId, double StartTime, double EndTime, EnumerateResolvedGpuFencesCallback Callback) const
+{
+	Session.ReadAccessCheck();
+
+	const uint32* FoundQueueIndex = GpuQueueIdToQueueIndexMap.Find(QueueId);
+	if (FoundQueueIndex == nullptr)
+	{
+		return;
+	}
+
+	const FGpuQueueData& QueueData = GpuQueueData[*FoundQueueIndex];
+	const TArray<FGpuSignalFence>& SignalFenceArray = QueueData.SignalFenceArray;
+
+	int32 SignalFenceIndex = Algo::LowerBoundBy(SignalFenceArray, StartTime,
+		[](const FGpuSignalFence& Fence) { return Fence.Timestamp; });
+
+	while (SignalFenceIndex < SignalFenceArray.Num() && SignalFenceArray[SignalFenceIndex].Timestamp <= EndTime)
+	{
+		const FGpuSignalFence& SignalFence = SignalFenceArray[SignalFenceIndex];
+		uint64 PrevSignalFenceValue = 0;
+		if (SignalFenceIndex > 0)
+		{
+			PrevSignalFenceValue = SignalFenceArray[SignalFenceIndex - 1].Value;
+		}
+
+		for (uint32 QueueIndex = 0; QueueIndex < (uint32) GpuQueueData.Num(); ++QueueIndex)
+		{
+			if (QueueIndex == *FoundQueueIndex)
+			{
+				continue;
+			}
+
+			const FGpuQueueData& TargetQueueData = GpuQueueData[QueueIndex];
+			const TArray<FGpuWaitFence>& TargetQueueArray = TargetQueueData.WaitFenceArray;
+
+			int32 TargetIndex = Algo::LowerBoundBy(TargetQueueArray, PrevSignalFenceValue,
+				[](const FGpuWaitFence& Fence) { return Fence.Value; });
+
+			while (TargetIndex < TargetQueueArray.Num() &&
+				   TargetQueueArray[TargetIndex].QueueToWaitForId == QueueId &&
+				   TargetQueueArray[TargetIndex].Value <= SignalFence.Value)
+			{
+				if (((TargetQueueArray[TargetIndex].Value > PrevSignalFenceValue) || SignalFence.Value == 0))
+				{
+					uint32 WaitFenceQueueId = GpuQueues[QueueIndex].Id;
+					Callback(QueueId, SignalFence, WaitFenceQueueId, TargetQueueArray[TargetIndex]);
+				}
+				++TargetIndex;
+			}
+		}
+		++SignalFenceIndex;
+	}
+	
+	const TArray<FGpuWaitFence>& WaitFenceArray = QueueData.WaitFenceArray;
+
+	int32 WaitFenceIndex = Algo::LowerBoundBy(WaitFenceArray, StartTime,
+		[](const FGpuWaitFence& Fence) { return Fence.Timestamp; });
+
+	while (WaitFenceIndex < WaitFenceArray.Num() && WaitFenceArray[WaitFenceIndex].Timestamp <= EndTime)
+	{
+		const FGpuWaitFence& WaitFence = WaitFenceArray[WaitFenceIndex];
+
+		const uint32* SignalQueueIndex = GpuQueueIdToQueueIndexMap.Find(WaitFence.QueueToWaitForId);
+		if (SignalQueueIndex == nullptr)
+		{
+			continue;
+		}
+
+		const FGpuQueueData& SignalQueueData = GpuQueueData[*SignalQueueIndex];
+
+		const TArray<FGpuSignalFence>& TargetQueueArray = SignalQueueData.SignalFenceArray;
+
+		int32 TargetIndex = Algo::LowerBoundBy(TargetQueueArray, WaitFence.Value,
+			[](const FGpuSignalFence& Fence) { return Fence.Value; });
+
+		if (TargetIndex < TargetQueueArray.Num())
+		{
+			Callback(WaitFence.QueueToWaitForId, TargetQueueArray[TargetIndex], QueueId, WaitFence);
+		}
+
+		++WaitFenceIndex;
+	}
+}
+
 IEditableTimeline<FTimingProfilerEvent>* FTimingProfilerProvider::GetGpuQueueEditableTimeline(uint32 QueueId)
 {
 	Session.WriteAccessCheck();
@@ -283,6 +460,19 @@ void FTimingProfilerProvider::EnumerateGpuQueues(TFunctionRef<void(const FGpuQue
 	}
 }
 
+bool FTimingProfilerProvider::GetGpuQueueTimelineIndex(uint32 QueueId, uint32& OutTimelineIndex) const
+{
+	Session.ReadAccessCheck();
+
+	const uint32* Index = GpuQueueIdToQueueIndexMap.Find(QueueId);
+	if (Index)
+	{
+		OutTimelineIndex = GpuQueues[*Index].TimelineIndex;
+		return true;
+	}
+	return false;
+}
+
 bool FTimingProfilerProvider::ReadTimeline(uint32 Index, TFunctionRef<void(const Timeline &)> Callback) const
 {
 	Session.ReadAccessCheck();
diff --git a/Engine/Source/Developer/TraceServices/Private/Model/TimingProfilerPrivate.h b/Engine/Source/Developer/TraceServices/Private/Model/TimingProfilerPrivate.h
--- a/Engine/Source/Developer/TraceServices/Private/Model/TimingProfilerPrivate.h
+++ b/Engine/Source/Developer/TraceServices/Private/Model/TimingProfilerPrivate.h
@@ -13,6 +13,12 @@ namespace TraceServices
 class FAnalysisSessionLock;
 class FStringStore;
 
+struct FGpuQueueData
+{
+	TArray<FGpuSignalFence> SignalFenceArray;
+	TArray<FGpuWaitFence> WaitFenceArray;
+};
+
 class FTimingProfilerProvider
 	: public ITimingProfilerProvider
 	, public ITimingProfilerTimerReader
@@ -32,6 +38,12 @@ public:
 	virtual bool GetGpuTimelineIndex(uint32& OutTimelineIndex) const override;
 	virtual bool GetGpu2TimelineIndex(uint32& OutTimelineIndex) const override;
 	virtual void EnumerateGpuQueues(TFunctionRef<void(const FGpuQueueInfo&)> Callback) const override;
+	virtual bool GetGpuQueueTimelineIndex(uint32 QueueId, uint32& OutTimelineIndex) const override;
+
+	virtual void EnumerateGpuSignalFences(uint32 QueueId, double StartTime, double EndTime, EnumerateGpuSignalFencesCallback Callback) const override;
+	virtual void EnumerateGpuWaitFences(uint32 QueueId, double StartTime, double EndTime, EnumerateGpuWaitFencesCallback Callback) const override;
+	virtual void EnumerateResolvedGpuFences(uint32 QueueId, double StartTime, double EndTime, EnumerateResolvedGpuFencesCallback Callback) const override;
+
 	virtual bool ReadTimeline(uint32 Index, TFunctionRef<void(const Timeline&)> Callback) const override;
 	virtual uint32 GetTimelineCount() const override { return Timelines.Num(); }
 	virtual void EnumerateTimelines(TFunctionRef<void(const Timeline&)> Callback) const override;
@@ -54,10 +66,15 @@ public:
 	virtual void SetTimerNameAndLocation(uint32 TimerId, FStringView Name, const TCHAR* File, uint32 Line) override;
 	virtual uint32 AddMetadata(uint32 OriginalTimerId, TArray<uint8>&& Metadata) override;
 	virtual TArrayView<uint8> GetEditableMetadata(uint32 MetadataTimerId) override;
+
 	virtual void AddGpuQueue(uint32 QueueId, uint8 GPU, uint8 Index, uint8 Type, const TCHAR* Name) override;
+	virtual void AddGpuSignalFence(uint32 QueueId, const FGpuSignalFence& SignalFence) override;
+	virtual void AddGpuWaitFence(uint32 QueueId, const FGpuWaitFence& WaitFence) override;
+
 	virtual IEditableTimeline<FTimingProfilerEvent>& GetCpuThreadEditableTimeline(uint32 ThreadId) override;
 	virtual IEditableTimeline<FTimingProfilerEvent>* GetGpuQueueEditableTimeline(uint32 QueueId) override;
 	virtual IEditableTimeline<FTimingProfilerEvent>* GetGpuQueueWorkEditableTimeline(uint32 QueueId) override;
+	
 	virtual const ITimingProfilerProvider* GetReadProvider() const override { return this; }
 
 private:
@@ -77,6 +94,7 @@ private:
 	static const uint32 GpuTimelineIndex = 0;
 	static const uint32 Gpu2TimelineIndex = 1;
 	TArray<FGpuQueueInfo> GpuQueues;
+	TArray<FGpuQueueData> GpuQueueData;
 	TMap<uint32, uint32> GpuQueueIdToQueueIndexMap; // GPU Queue Id --> queue index, in GpuQueues
 	TTableLayout<FTimingProfilerAggregatedStats> AggregatedStatsTableLayout;
 };
diff --git a/Engine/Source/Developer/TraceServices/Public/TraceServices/Model/TimingProfiler.h b/Engine/Source/Developer/TraceServices/Public/TraceServices/Model/TimingProfiler.h
--- a/Engine/Source/Developer/TraceServices/Public/TraceServices/Model/TimingProfiler.h
+++ b/Engine/Source/Developer/TraceServices/Public/TraceServices/Model/TimingProfiler.h
@@ -130,6 +130,19 @@ public:
 	virtual TArrayView<const uint8> GetMetadata(uint32 MetadataTimerId) const { return TArrayView<const uint8>(); }
 };
 
+struct FGpuSignalFence
+{
+	double Timestamp = 0.0f;
+	uint64 Value = 0;
+};
+
+struct FGpuWaitFence
+{
+	double Timestamp = 0.0f;
+	uint64 Value = 0;
+	uint32 QueueToWaitForId = 0;
+};
+
 struct FGpuQueueInfo
 {
 	uint32 Id = 0;
@@ -141,11 +154,20 @@ struct FGpuQueueInfo
 	uint32 WorkTimelineIndex = ~0;
 };
 
+enum class EEnumerateResult
+{
+	Continue,
+	Stop,
+};
+
 class ITimingProfilerProvider
 	: public IProvider
 {
 public:
 	typedef ITimeline<FTimingProfilerEvent> Timeline;
+	typedef TFunctionRef<EEnumerateResult(const FGpuSignalFence&)> EnumerateGpuSignalFencesCallback;
+	typedef TFunctionRef<EEnumerateResult(const FGpuWaitFence&)> EnumerateGpuWaitFencesCallback;
+	typedef TFunctionRef<EEnumerateResult(uint32 /*SignalFenceQueueId*/, const FGpuSignalFence& /*SignalFence*/, uint32 /*WaitFenceQueueId*/, const FGpuWaitFence& /*WaitFence*/)> EnumerateResolvedGpuFencesCallback;
 
 	virtual ~ITimingProfilerProvider() = default;
 
@@ -153,6 +175,11 @@ public:
 	virtual bool GetGpuTimelineIndex(uint32& OutTimelineIndex) const = 0;
 	virtual bool GetGpu2TimelineIndex(uint32& OutTimelineIndex) const = 0;
 	virtual void EnumerateGpuQueues(TFunctionRef<void(const FGpuQueueInfo&)> Callback) const {}
+	virtual bool GetGpuQueueTimelineIndex(uint32 QueueId, uint32& OutTimelineIndex) const { return false; };
+
+	virtual void EnumerateGpuSignalFences(uint32 QueueId, double StartTime, double EndTime, EnumerateGpuSignalFencesCallback Callback) const {}
+	virtual void EnumerateGpuWaitFences(uint32 QueueId, double StartTime, double EndTime, EnumerateGpuWaitFencesCallback Callback) const {}
+	virtual void EnumerateResolvedGpuFences(uint32 QueueId, double StartTime, double EndTime, EnumerateResolvedGpuFencesCallback Callback) const {}
 
 	virtual bool ReadTimeline(uint32 Index, TFunctionRef<void(const Timeline&)> Callback) const = 0;
 	virtual uint32 GetTimelineCount() const = 0;
@@ -265,6 +292,22 @@ public:
 	*/
 	virtual void AddGpuQueue(uint32 QueueId, uint8 GPU, uint8 Index, uint8 Type, const TCHAR* Name) {}
 
+	/*
+	* Adds a new GPU signal fence to a queue.
+	*
+	* @param QueueId		The GPU queue that issues the signal.
+	* @param SignalFence	The signal fence.
+	*/
+	virtual void AddGpuSignalFence(uint32 QueueId, const FGpuSignalFence& SignalFence) {};
+
+	/*
+	* Adds a new GPU wait fence to a queue.
+	*
+	* @param QueueId		The GPU queue that waits.
+	* @param WaitFence		The wait fence.
+	*/
+	virtual void AddGpuWaitFence(uint32 QueueId, const FGpuWaitFence& WaitFence) {};
+
 	/*
 	* Gets an object to receive ordered timeline events for a CPU thread.
 	*
diff --git a/Engine/Source/Runtime/RHI/Private/GPUProfiler.cpp b/Engine/Source/Runtime/RHI/Private/GPUProfiler.cpp
--- a/Engine/Source/Runtime/RHI/Private/GPUProfiler.cpp
+++ b/Engine/Source/Runtime/RHI/Private/GPUProfiler.cpp
@@ -1644,6 +1644,8 @@ namespace UE::RHI::GPUProfiler
 
 			bool bBusy = false;
 
+			bool bWasTraced = false;
+
 		#if WITH_RHI_BREADCRUMBS
 			TMap<FRHIBreadcrumbData_Stats, int32> ActiveStats;
 			FRHIBreadcrumbNode* Breadcrumb = nullptr;
@@ -1835,6 +1837,8 @@ namespace UE::RHI::GPUProfiler
 				Result.GPUTimestampBOP = FMath::Max(LastGPUCycles, Event.CPUTimestamp);
 				Result.Value = Event.Value;
 
+				FGpuProfilerTrace::SignalFence(Queue.Value, Result.GPUTimestampBOP, Event.Value);
+
 				//
 				// Fences signals *MUST* be sequential, to remove ambiguity caused by trimming the Signals array.
 				// 
@@ -1903,6 +1907,9 @@ namespace UE::RHI::GPUProfiler
 					: 0;
 
 				AccumulateTime(0, 0, Idle);
+
+				FGpuProfilerTrace::BeginWork(Queue.Value, Event.GPUTimestampTOP, Event.CPUTimestamp);
+
 				LastGPUCycles = FMath::Max(LastGPUCycles, Event.GPUTimestampTOP);
 			}
 
@@ -1915,6 +1922,9 @@ namespace UE::RHI::GPUProfiler
 					: 0;
 
 				AccumulateTime(Busy, 0, 0);
+
+				FGpuProfilerTrace::EndWork(Queue.Value, Event.GPUTimestampBOP);
+
 				LastGPUCycles = FMath::Max(LastGPUCycles, Event.GPUTimestampBOP);
 			}
 
@@ -1941,12 +1951,15 @@ namespace UE::RHI::GPUProfiler
 
 				Breadcrumb = Event.Breadcrumb;
 
+				FRHIBreadcrumb::FBuffer Buffer;
+				const TCHAR* Name = Event.Breadcrumb->Name.GetTCHAR(Buffer);
+				FGpuProfilerTrace::BeginBreadcrumb(Queue.Value, Event.GPUTimestampTOP, Name);
+
 			#if WITH_PROFILEGPU
 				if (Profile.bProfileFrame)
 				{
 					// Push a new node
-					FRHIBreadcrumb::FBuffer Buffer;
-					Profile.PushNode(Breadcrumb->Name.GetTCHAR(Buffer));
+					Profile.PushNode(Name);
 				}
 			#endif
 			}
@@ -1973,6 +1986,10 @@ namespace UE::RHI::GPUProfiler
 
 				Breadcrumb = Event.Breadcrumb->GetParent();
 
+				FRHIBreadcrumb::FBuffer Buffer;
+				const TCHAR* Name = Event.Breadcrumb->Name.GetTCHAR(Buffer);
+				FGpuProfilerTrace::EndBreadcrumb(Queue.Value, Event.GPUTimestampBOP, Name);
+
 			#if WITH_PROFILEGPU
 				if (Profile.bProfileFrame)
 				{
@@ -1994,9 +2011,10 @@ namespace UE::RHI::GPUProfiler
 					}
 				}
 			#endif
+				FGpuProfilerTrace::Stats(Queue.Value, Event.NumDraws, Event.NumPrimitives);
 			}
 
-			void Wait(FResolvedWait const& ResolvedWait)
+			void Wait(FResolvedWait const& ResolvedWait, const FEvent::FWaitFence& WaitFence)
 			{
 				// Time the queue was idle between the last EndWork event, and the Wait command being submitted to the GPU driver.
 				uint64 Idle = ResolvedWait.CPUTimestamp > LastGPUCycles
@@ -2005,6 +2023,8 @@ namespace UE::RHI::GPUProfiler
 
 				uint64 WaitStart = FMath::Max(ResolvedWait.CPUTimestamp, LastGPUCycles);
 
+				FGpuProfilerTrace::WaitFence(Queue.Value, ResolvedWait.GPUTimestampTOP, WaitFence.Queue.Value, WaitFence.Value);
+
 				// Time the queue spent waiting for the fence to signal on another queue.
 				uint64 Wait = ResolvedWait.GPUTimestampTOP > WaitStart
 					? ResolvedWait.GPUTimestampTOP - WaitStart
@@ -2137,6 +2157,8 @@ namespace UE::RHI::GPUProfiler
 
 		void InitializeQueues(TConstArrayView<FQueue> Queues) override
 		{
+			FGpuProfilerTrace::Initialize();
+
 			for (FQueue Queue : Queues)
 			{
 				TUniquePtr<FQueueState>& Ptr = QueueStates.FindOrAdd(Queue);
@@ -2149,6 +2171,12 @@ namespace UE::RHI::GPUProfiler
 
 		bool ProcessQueue(FQueueState& QueueState, FIterator& Iterator)
 		{
+			if (FGpuProfilerTrace::IsAvailable() && !QueueState.bWasTraced)
+			{
+				FGpuProfilerTrace::InitializeQueue(QueueState.Queue.Value, QueueState.Queue.GetTypeString());
+				QueueState.bWasTraced = true;
+			}
+
 			while (FEvent const* Event = Iterator.Peek())
 			{
 				switch (Event->GetType())
@@ -2204,6 +2232,7 @@ namespace UE::RHI::GPUProfiler
 				case FEvent::EType::WaitFence:
 					{
 						check(!QueueState.bBusy);
+						const FEvent::FWaitFence& WaitFence= Event->Value.Get<FEvent::FWaitFence>();
 						TOptional<FResolvedWait> ResolvedWait = ResolveWait(QueueState, Event->Value.Get<FEvent::FWaitFence>());
 
 						if (!ResolvedWait.IsSet())
@@ -2212,7 +2241,7 @@ namespace UE::RHI::GPUProfiler
 							return false;
 						}
 
-						QueueState.Wait(*ResolvedWait);
+						QueueState.Wait(*ResolvedWait, WaitFence);
 					}
 					break;
 
diff --git a/Engine/Source/Runtime/RHI/Private/GpuProfilerTrace.cpp b/Engine/Source/Runtime/RHI/Private/GpuProfilerTrace.cpp
--- a/Engine/Source/Runtime/RHI/Private/GpuProfilerTrace.cpp
+++ b/Engine/Source/Runtime/RHI/Private/GpuProfilerTrace.cpp
@@ -10,11 +10,6 @@
 #include "Trace/Detail/Important/ImportantLogScope.h"
 #include "Trace/Detail/LogScope.h"
 
-// Tracing for the new GPU Profiler
-#ifndef UE_TRACE_GPU_PROFILER_ENABLED
-#define UE_TRACE_GPU_PROFILER_ENABLED UE_TRACE_ENABLED && RHI_NEW_GPU_PROFILER && !UE_BUILD_SHIPPING
-#endif
-
 // Both (old and new) GPU profilers uses same "GPU" trace channel.
 #if GPUPROFILERTRACE_ENABLED || UE_TRACE_GPU_PROFILER_ENABLED
 	UE_TRACE_CHANNEL_EXTERN(GpuChannel, RHI_API)
@@ -325,162 +320,107 @@ UE_TRACE_EVENT_BEGIN(GpuProfiler, EventStats)
 	//UE_TRACE_EVENT_FIELD(uint32, NumVertices)
 UE_TRACE_EVENT_END()
 
-struct FGPUProfilerSink_Trace final : public FEventSink
-{
-	bool bInitialized = false;
-	TSet<FQueue> CachedQueues;
+UE_TRACE_EVENT_BEGIN(GpuProfiler, SignalFence)
+	UE_TRACE_EVENT_FIELD(uint32, QueueId)
+	UE_TRACE_EVENT_FIELD(uint64, CPUTimestamp)
+	UE_TRACE_EVENT_FIELD(uint64, Value)
+UE_TRACE_EVENT_END()
 
-	void Initialize()
-	{
-		ensure(!bInitialized);
+UE_TRACE_EVENT_BEGIN(GpuProfiler, WaitFence)
+	UE_TRACE_EVENT_FIELD(uint32, QueueId)
+	UE_TRACE_EVENT_FIELD(uint64, CPUTimestamp)
+	UE_TRACE_EVENT_FIELD(uint32, QueueToWaitForId)
+	UE_TRACE_EVENT_FIELD(uint64, Value)
+UE_TRACE_EVENT_END()
 
-		constexpr uint8 Version = 2;
+bool FGpuProfilerTrace::IsAvailable()
+{
+	return UE_TRACE_CHANNELEXPR_IS_ENABLED(GpuChannel);
+}
 
-		UE_TRACE_LOG(GpuProfiler, Init, GpuChannel)
-			<< Init.Version(Version);
+void FGpuProfilerTrace::Initialize()
+{
+	static bool bInitialized = false;
+	ensure(!bInitialized);
 
-		bInitialized = true;
-	}
+	constexpr uint8 Version = 2;
 
-	void TraceQueues(TSet<FQueue>& Queues)
-	{
-		for (FQueue const& Queue : Queues)
-		{
-			FStringView TypeString(Queue.GetTypeString());
-			UE_TRACE_LOG(GpuProfiler, QueueSpec, GpuChannel, TypeString.Len() * sizeof(TCHAR))
-				<< QueueSpec.QueueId(Queue.Value)
-				<< QueueSpec.TypeString(TypeString.GetData(), TypeString.Len());
-		}
-	}
+	UE_TRACE_LOG(GpuProfiler, Init, GpuChannel)
+		<< Init.Version(Version);
 
-	void InitializeQueues(TConstArrayView<FQueue> Queues) override
-	{
-		for (FQueue Queue : Queues)
-		{
-			CachedQueues.Add(Queue);
-		}
-	}
+	bInitialized = true;
+}
 
-	void ProcessStream(TSharedRef<FEventStream> const& EventStream)
-	{
-		FQueue Queue = EventStream->Queue;
-		if (!UE_TRACE_CHANNELEXPR_IS_ENABLED(GpuChannel))
-		{
-			CachedQueues.Add(Queue);
-			return;
-		}
+void FGpuProfilerTrace::InitializeQueue(uint32 QueueId, const TCHAR* Name)
+{
+	FStringView TypeString(Name);
+	UE_TRACE_LOG(GpuProfiler, QueueSpec, GpuChannel, TypeString.Len() * sizeof(TCHAR))
+		<< QueueSpec.QueueId(QueueId)
+		<< QueueSpec.TypeString(TypeString.GetData(), TypeString.Len());
+}
 
-		if (!bInitialized)
-		{
-			Initialize();
-		}
+void FGpuProfilerTrace::FrameBoundary(uint32 QueueId, uint32 FrameId)
+{
+	UE_TRACE_LOG(GpuProfiler, EventFrameBoundary, GpuChannel)
+		<< EventFrameBoundary.QueueId(QueueId)
+		<< EventFrameBoundary.FrameNumber(FrameId);
+}
 
-		if (!CachedQueues.IsEmpty())
-		{
-			TraceQueues(CachedQueues);
-			CachedQueues.Empty();
-		}
+void FGpuProfilerTrace::BeginBreadcrumb(uint32 QueueId, uint64 GPUTimestampTOP, const TCHAR* Name)
+{
+	UE_TRACE_LOG(GpuProfiler, EventBeginBreadcrumb, GpuChannel)
+		<< EventBeginBreadcrumb.QueueId(QueueId)
+		<< EventBeginBreadcrumb.GPUTimestampTOP(GPUTimestampTOP)
+		<< EventBeginBreadcrumb.Name(Name);
+}
 
-		FEventSink::FIterator Iterator = EventStream;
-		while (FEvent const* Event = Iterator.Pop())
-		{
-			switch (Event->GetType())
-			{
-				case FEvent::EType::FrameBoundary:
-				{
-					FEvent::FFrameBoundary const& FrameBoundary = Event->Value.Get<FEvent::FFrameBoundary>();
-					UE_TRACE_LOG(GpuProfiler, EventFrameBoundary, GpuChannel)
-						<< EventFrameBoundary.QueueId(Queue.Value)
-						<< EventFrameBoundary.FrameNumber(FrameBoundary.FrameNumber);
-					break;
-				}
+void FGpuProfilerTrace::EndBreadcrumb(uint32 QueueId, uint64 GPUTimestampTOP, const TCHAR* Name)
+{
+	UE_TRACE_LOG(GpuProfiler, EventEndBreadcrumb, GpuChannel)
+		<< EventEndBreadcrumb.QueueId(QueueId)
+		<< EventEndBreadcrumb.GPUTimestampBOP(GPUTimestampTOP)
+		<< EventEndBreadcrumb.Name(Name); // for debugging
+}
 
-#if WITH_RHI_BREADCRUMBS
-				case FEvent::EType::BeginBreadcrumb:
-				{
-					FEvent::FBeginBreadcrumb const& BeginBreadcrumb = Event->Value.Get<FEvent::FBeginBreadcrumb>();
-					FRHIBreadcrumb::FBuffer Buffer;
-					const TCHAR* Name = BeginBreadcrumb.Breadcrumb->Name.GetTCHAR(Buffer);
-					UE_TRACE_LOG(GpuProfiler, EventBeginBreadcrumb, GpuChannel)
-						<< EventBeginBreadcrumb.QueueId(Queue.Value)
-						<< EventBeginBreadcrumb.GPUTimestampTOP(BeginBreadcrumb.GPUTimestampTOP)
-						<< EventBeginBreadcrumb.Name(Name);
-					break;
-				}
-
-				case FEvent::EType::EndBreadcrumb:
-				{
-					FEvent::FEndBreadcrumb const& EndBreadcrumb = Event->Value.Get<FEvent::FEndBreadcrumb>();
-					FRHIBreadcrumb::FBuffer Buffer;
-					const TCHAR* Name = EndBreadcrumb.Breadcrumb->Name.GetTCHAR(Buffer);
-					UE_TRACE_LOG(GpuProfiler, EventEndBreadcrumb, GpuChannel)
-						<< EventEndBreadcrumb.QueueId(Queue.Value)
-						<< EventEndBreadcrumb.GPUTimestampBOP(EndBreadcrumb.GPUTimestampBOP)
-						<< EventEndBreadcrumb.Name(Name); // for debugging
-					break;
-				}
-#endif // WITH_RHI_BREADCRUMBS
+void FGpuProfilerTrace::BeginWork(uint32 QueueId, uint64 GPUTimestampTOP, uint64 CPUTimestamp)
+{
+	UE_TRACE_LOG(GpuProfiler, EventBeginWork, GpuChannel)
+		<< EventBeginWork.QueueId(QueueId)
+		<< EventBeginWork.GPUTimestampTOP(GPUTimestampTOP)
+		<< EventBeginWork.CPUTimestamp(CPUTimestamp);
+}
 
-				case FEvent::EType::BeginWork:
-				{
-					FEvent::FBeginWork const& BeginWork = Event->Value.Get<FEvent::FBeginWork>();
-					UE_TRACE_LOG(GpuProfiler, EventBeginWork, GpuChannel)
-						<< EventBeginWork.QueueId(Queue.Value)
-						<< EventBeginWork.GPUTimestampTOP(BeginWork.GPUTimestampTOP)
-						<< EventBeginWork.CPUTimestamp(BeginWork.CPUTimestamp);
-					break;
-				}
-
-				case FEvent::EType::EndWork:
-				{
-					FEvent::FEndWork const& EndWork = Event->Value.Get<FEvent::FEndWork>();
-					UE_TRACE_LOG(GpuProfiler, EventEndWork, GpuChannel)
-						<< EventEndWork.QueueId(Queue.Value)
-						<< EventEndWork.GPUTimestampBOP(EndWork.GPUTimestampBOP);
-					break;
-				}
-
-				case FEvent::EType::Stats:
-				{
-					FEvent::FStats const& Stats = Event->Value.Get<FEvent::FStats>();
-					UE_TRACE_LOG(GpuProfiler, EventStats, GpuChannel)
-						<< EventStats.QueueId(Queue.Value)
-						<< EventStats.NumDraws(Stats.NumDraws)
-						<< EventStats.NumPrimitives(Stats.NumPrimitives)
-						// @todo add num dispatches / num vertices stats
-						//<< EventStats.NumDispatches(Stats.NumDispatches)
-						//<< EventStats.NumVertices(Stats.NumVertices)
-						;
-					break;
-				}
-
-				case FEvent::EType::SignalFence:
-				case FEvent::EType::WaitFence:
-				case FEvent::EType::Flip:
-				case FEvent::EType::VSync:
-				{
-					//TODO
-					break;
-				}
-
-				default:
-				{
-					unimplemented();
-					break;
-				}
-			}
-		}
-	}
+void FGpuProfilerTrace::EndWork(uint32 QueueId, uint64 GPUTimestampBOP)
+{
+	UE_TRACE_LOG(GpuProfiler, EventEndWork, GpuChannel)
+		<< EventEndWork.QueueId(QueueId)
+		<< EventEndWork.GPUTimestampBOP(GPUTimestampBOP);
+}
 
-	void ProcessStreams(TConstArrayView<TSharedRef<FEventStream>> EventStreams) override
-	{
-		for (TSharedRef<FEventStream> const& Stream : EventStreams)
-		{
-			ProcessStream(Stream);
-		}
-	}
+void FGpuProfilerTrace::Stats(uint32 QueueId, uint32 NumDraws, uint32 NumPrimitives)
+{
+	UE_TRACE_LOG(GpuProfiler, EventStats, GpuChannel)
+		<< EventStats.QueueId(QueueId)
+		<< EventStats.NumDraws(NumDraws)
+		<< EventStats.NumPrimitives(NumPrimitives);
+}
 
-} GGPUProfilerSink_Trace;
+void FGpuProfilerTrace::SignalFence(uint32 QueueId, uint64 ResolvedTimestamp, uint64 Value)
+{
+	UE_TRACE_LOG(GpuProfiler, SignalFence, GpuChannel)
+		<< SignalFence.QueueId(QueueId)
+		<< SignalFence.CPUTimestamp(ResolvedTimestamp)
+		<< SignalFence.Value(Value);
+}
+
+void FGpuProfilerTrace::WaitFence(uint32 QueueId, uint64 ResolvedTimestamp, uint32 QueueToWaitForId, uint64 Value)
+{
+	UE_TRACE_LOG(GpuProfiler, WaitFence, GpuChannel)
+		<< WaitFence.QueueId(QueueId)
+		<< WaitFence.CPUTimestamp(ResolvedTimestamp)
+		<< WaitFence.QueueToWaitForId(QueueToWaitForId)
+		<< WaitFence.Value(Value);
+}
 
 } // namespace UE::RHI::GPUProfiler
 
diff --git a/Engine/Source/Runtime/RHI/Public/GpuProfilerTrace.h b/Engine/Source/Runtime/RHI/Public/GpuProfilerTrace.h
--- a/Engine/Source/Runtime/RHI/Public/GpuProfilerTrace.h
+++ b/Engine/Source/Runtime/RHI/Public/GpuProfilerTrace.h
@@ -15,6 +15,11 @@
 #endif
 #endif
 
+// Tracing for the new GPU Profiler
+#ifndef UE_TRACE_GPU_PROFILER_ENABLED
+	#define UE_TRACE_GPU_PROFILER_ENABLED UE_TRACE_ENABLED && RHI_NEW_GPU_PROFILER && !UE_BUILD_SHIPPING
+#endif
+
 #if GPUPROFILERTRACE_ENABLED
 
 #if RHI_NEW_GPU_PROFILER
@@ -60,3 +65,83 @@ struct FGpuProfilerTrace
 #define TRACE_GPUPROFILER_END_EVENT(...)                 UE_DEPRECATED_MACRO(5.6, "TRACE_GPUPROFILER_END_EVENT has been deprecated and it is non functional. Use standard RHI breadcrumb events instead.")
 #define TRACE_GPUPROFILER_END_FRAME(...)                 UE_DEPRECATED_MACRO(5.6, "TRACE_GPUPROFILER_END_FRAME has been deprecated and it is non functional. Use standard RHI breadcrumb events instead.")
 #define TRACE_GPUPROFILER_DEINITIALIZE(...)              UE_DEPRECATED_MACRO(5.6, "TRACE_GPUPROFILER_DEINITIALIZE has been deprecated and it is non functional. Use standard RHI breadcrumb events instead.")
+
+namespace UE::RHI::GPUProfiler
+{
+
+struct FGpuProfilerTrace
+{
+#if UE_TRACE_GPU_PROFILER_ENABLED
+
+	/**
+	 * Returns true if the Gpu channel is enabled.
+	 */
+	RHI_API static bool IsAvailable();
+
+	/**
+	 * Initialize Gpu Profiling Trace. Must only be called once.
+	 */
+	RHI_API static void Initialize();
+
+	/**
+	 * Trace an initialization event for a Gpu Queue.
+	 */
+	RHI_API static void InitializeQueue(uint32 QueueId, const TCHAR* Name);
+
+	/**
+	 * Trace a frame boundary for a Gpu Queue. 
+	 */
+	RHI_API static void FrameBoundary(uint32 QueueId, uint32 FrameId);
+
+	/**
+	 * Trace the start of a breadcrump on a Gpu Queue. 
+	 */
+	RHI_API static void BeginBreadcrumb(uint32 QueueId, uint64 GPUTimestampTOP, const TCHAR* Name);
+
+	/**
+	 * Trace the end of a breadcrump on a Gpu Queue. 
+	 */
+	RHI_API static void EndBreadcrumb(uint32 QueueId, uint64 GPUTimestampBOP, const TCHAR* Name);
+	
+	/**
+	 * Trace the start of a work event on a Gpu Queue. 
+	 */
+	RHI_API static void BeginWork(uint32 QueueId, uint64 GPUTimestampTOP, uint64 CPUTimestamp);
+
+	/**
+	 * Trace the end of a work event on a Gpu Queue. 
+	 */
+	RHI_API static void EndWork(uint32 QueueId, uint64 GPUTimestampBOP);
+
+	/**
+	 * Trace Gpu stats. 
+	 */
+	RHI_API static void Stats(uint32 QueueId, uint32 NumDraws, uint32 NumPrimitives);
+
+	/**
+	 * Trace a signal fence event on a Gpu Queue. 
+	 */
+	RHI_API static void SignalFence(uint32 QueueId, uint64 ResolvedTimestamp, uint64 Value);
+
+	/**
+	 * Trace a wait fence event on a Gpu Queue. 
+	 */
+	RHI_API static void WaitFence(uint32 QueueId, uint64 ResolvedTimestamp, uint32 QueueToWaitForId, uint64 Value);
+
+#else
+	RHI_API static bool IsAvailable() { return false; }
+	RHI_API static void Initialize() {}
+	RHI_API static void InitializeQueue(uint32 QueueId, const TCHAR* Name) {}
+	RHI_API static void FrameBoundary(uint32 QueueId, uint32 FrameId) {}
+	RHI_API static void BeginBreadcrumb(uint32 QueueId, uint64 GPUTimestampTOP, const TCHAR* Name) {}
+	RHI_API static void EndBreadcrumb(uint32 QueueId, uint64 GPUTimestampBOP, const TCHAR* Name) {}
+	RHI_API static void BeginWork(uint32 QueueId, uint64 GPUTimestampTOP, uint64 CPUTimestamp) {}
+	RHI_API static void EndWork(uint32 QueueId, uint64 GPUTimestampBOP) {}
+	RHI_API static void Stats(uint32 QueueId, uint32 NumDraws, uint32 NumPrimitives) {}
+	RHI_API static void SignalFence(uint32 QueueId, uint64 ResolvedTimestamp, uint64 Value) {}
+	RHI_API static void WaitFence(uint32 QueueId, uint64 ResolvedTimestamp, uint32 QueueToWaitForId, uint64 Value) {}
+#endif
+};
+
+}
+
