Adding the Before state of the resource in the in BeginTransition log in order to give more context while debugging a transition.

diff --git a/Engine/Source/Runtime/RHI/Private/RHIValidation.cpp b/Engine/Source/Runtime/RHI/Private/RHIValidation.cpp
--- a/Engine/Source/Runtime/RHI/Private/RHIValidation.cpp
+++ b/Engine/Source/Runtime/RHI/Private/RHIValidation.cpp
@@ -2098,8 +2098,9 @@ namespace RHIValidation
 #endif
 			)
 		{
-			BeginTrace = Log(Resource, SubresourceIndex, CreateTrace, TEXT("Begin"), TEXT("BeginTransition"), *FString::Printf(TEXT("Current: (%s) New: (%s), Flags: %s, Executing Pipeline: %s"),
+			BeginTrace = Log(Resource, SubresourceIndex, CreateTrace, TEXT("Begin"), TEXT("BeginTransition"), *FString::Printf(TEXT("Current: (%s) -> Before(%s) New: (%s), Flags: %s, Executing Pipeline: %s"),
 				*State.Current.ToString(),
+				*CurrentStateFromRHI.ToString(),
 				*TargetState.ToString(),
 				*GetResourceTransitionFlagsName(NewFlags),
 				*GetRHIPipelineName(ExecutingPipeline)
