[HWRT] Add an option to specify maximum size of a scratch buffer used for acceleration structure builds.

#tests EngineTests
#rnx

diff --git a/Engine/Config/BaseDeviceProfiles.ini b/Engine/Config/BaseDeviceProfiles.ini
--- a/Engine/Config/BaseDeviceProfiles.ini
+++ b/Engine/Config/BaseDeviceProfiles.ini
@@ -979,6 +979,7 @@ BaseProfileName=Android_Vulkan_SM5
 [Android_Mali_Vulkan_SM5 DeviceProfile]
 DeviceType=Android
 BaseProfileName=Android_Vulkan_SM5
++CVars=r.RHICmd.MaxAccelerationStructureBuildScratchSize=0
 
 [Android_Xclipse_Vulkan_SM5 DeviceProfile]
 DeviceType=Android
diff --git a/Engine/Source/Runtime/RHI/Private/RHICommandList.cpp b/Engine/Source/Runtime/RHI/Private/RHICommandList.cpp
--- a/Engine/Source/Runtime/RHI/Private/RHICommandList.cpp
+++ b/Engine/Source/Runtime/RHI/Private/RHICommandList.cpp
@@ -85,6 +85,11 @@ static TAutoConsoleVariable<int32> CVarRHICmdBufferWriteLocks(
 	1,
 	TEXT("Only relevant with an RHI thread. Debugging option to diagnose problems with buffered locks."));
 
+static TAutoConsoleVariable<int32> CVarRHICmdMaxAccelerationStructureBuildScratchSize(
+	TEXT("r.RHICmd.MaxAccelerationStructureBuildScratchSize"),
+	-1,
+	TEXT("Set the maximum size in bytes of scratch buffer used for acceleration structures build. Setting it to 0 will serialize all builds. (default 2GB)"));
+
 FAutoConsoleTaskPriority CPrio_SceneRenderingTask(
 	TEXT("TaskGraph.TaskPriorities.SceneRenderingTask"),
 	TEXT("Task and thread priority for various scene rendering tasks."),
@@ -1777,9 +1782,27 @@ void FRHIComputeCommandList::BuildAccelerationStructure(FRHIRayTracingGeometry*
 void FRHIComputeCommandList::BuildAccelerationStructures(const TArrayView<const FRayTracingGeometryBuildParams> Params)
 {
 	// Buffer size is limited to 2Gb, thus split acceleration structure building into pieces to accommodate this limitation
-	const uint64 MaxScratchMemorySize = 2147483647u;
+	const uint64 MaxScratchMemorySize = FMath::Min(uint32(CVarRHICmdMaxAccelerationStructureBuildScratchSize.GetValueOnRenderThread()), 2147483647u);
 	const uint32 ParamTotalCount = Params.Num();
 
+	uint64 TotalScratchMemorySize = 0;
+	uint64 LargestScratchMemorySize = 0;
+	for (const FRayTracingGeometryBuildParams& P : Params)
+	{
+		const uint64 ScratchBufferRequiredSize = P.BuildMode == EAccelerationStructureBuildMode::Update ? P.Geometry->GetSizeInfo().UpdateScratchSize : P.Geometry->GetSizeInfo().BuildScratchSize;
+		TotalScratchMemorySize += ScratchBufferRequiredSize;
+		LargestScratchMemorySize = FMath::Max(ScratchBufferRequiredSize, LargestScratchMemorySize);
+	}
+
+	const uint64 TotalRequiredScratchMemorySize = FMath::Max(LargestScratchMemorySize, FMath::Min(MaxScratchMemorySize, TotalScratchMemorySize));
+
+	FRHIBufferRange ScratchBufferRange{};
+	check(uint32(TotalRequiredScratchMemorySize) == TotalRequiredScratchMemorySize);
+	const FRHIBufferCreateDesc CreateDesc =
+				FRHIBufferCreateDesc::Create(TEXT("RHIScratchBuffer"), TotalRequiredScratchMemorySize, 0, EBufferUsageFlags::StructuredBuffer | EBufferUsageFlags::RayTracingScratch)
+				.SetInitialState(ERHIAccess::UAVCompute);
+	ScratchBufferRange.Buffer = CreateBuffer(CreateDesc);
+
 	uint32 ParamIt = 0;
 	while (ParamIt < ParamTotalCount)
 	{
@@ -1787,14 +1810,14 @@ void FRHIComputeCommandList::BuildAccelerationStructures(const TArrayView<const
 		uint32 ParamCount = 0;
 
 		// Select a sub-range of input params which fits into MaxScratchMemorySize
-		uint64 TotalRequiredScratchMemorySize = 0;
-		for ( ;ParamIt < ParamTotalCount; ++ParamIt)
+		uint64 RequiredScratchMemorySize = 0;
+		for (; ParamIt < ParamTotalCount; ++ParamIt)
 		{
 			const FRayTracingGeometryBuildParams& P = Params[ParamIt];
 			uint64 ScratchBufferRequiredSize = P.BuildMode == EAccelerationStructureBuildMode::Update ? P.Geometry->GetSizeInfo().UpdateScratchSize : P.Geometry->GetSizeInfo().BuildScratchSize;
-			if (ScratchBufferRequiredSize + TotalRequiredScratchMemorySize < MaxScratchMemorySize)
+			if (ScratchBufferRequiredSize + RequiredScratchMemorySize <= TotalRequiredScratchMemorySize)
 			{
-				TotalRequiredScratchMemorySize += ScratchBufferRequiredSize;
+				RequiredScratchMemorySize += ScratchBufferRequiredSize;
 				ParamCount++;
 			}
 			else
@@ -1806,17 +1829,7 @@ void FRHIComputeCommandList::BuildAccelerationStructures(const TArrayView<const
 		// Allocate scratch buffer and build the acceleration structure for the selected sub-range
 		if (ParamCount > 0)
 		{
-			check(uint32(TotalRequiredScratchMemorySize) == TotalRequiredScratchMemorySize);
-
 			const TArrayView<const FRayTracingGeometryBuildParams> EffectiveParams = MakeArrayView<const FRayTracingGeometryBuildParams>(&Params[ParamBegin], ParamCount);
-
-			const FRHIBufferCreateDesc CreateDesc =
-				FRHIBufferCreateDesc::Create(TEXT("RHIScratchBuffer"), TotalRequiredScratchMemorySize, 0, EBufferUsageFlags::StructuredBuffer | EBufferUsageFlags::RayTracingScratch)
-				.SetInitialState(ERHIAccess::UAVCompute);
-
-			FRHIBufferRange ScratchBufferRange{};
-			ScratchBufferRange.Buffer = CreateBuffer(CreateDesc);
-
 			BuildAccelerationStructures(EffectiveParams, ScratchBufferRange);
 		}
 	}
