Fix shadowed variable build errors in GPU profiler

diff --git a/Engine/Source/Runtime/RHI/Private/GPUProfiler.cpp b/Engine/Source/Runtime/RHI/Private/GPUProfiler.cpp
--- a/Engine/Source/Runtime/RHI/Private/GPUProfiler.cpp
+++ b/Engine/Source/Runtime/RHI/Private/GPUProfiler.cpp
@@ -1689,7 +1689,7 @@ namespace UE::RHI::GPUProfiler
 					Current = Current->Parent;
 				}
 
-				void LogTree(FQueue Queue, uint32 FrameNumber) const
+				void LogTree(FQueueState const& QueueState, uint32 FrameNumber) const
 				{
 					FTable Table;
 
@@ -1794,9 +1794,9 @@ namespace UE::RHI::GPUProfiler
 						TEXT("\n")
 						TEXT("%s")
 						, FrameNumber
-						, Queue.GPU
-						, Queue.GetTypeString()
-						, Queue.Index
+						, QueueState.Queue.GPU
+						, QueueState.Queue.GetTypeString()
+						, QueueState.Queue.Index
 						, TEXT("Frame Time")
 						, First ? First->Inclusive.GetBusyMilliseconds() : 0.0
 						, *IConsoleManager::Get().FindConsoleObjectName(GCVarProfileGPU_Root.AsVariable())
@@ -2046,7 +2046,7 @@ namespace UE::RHI::GPUProfiler
 				}
 			}
 
-			void FrameBoundary(FEvent::FFrameBoundary const& Event, FFrameState& FrameState, uint32 ProfileFrameNumber)
+			void FrameBoundary(FEvent::FFrameBoundary const& Event, FFrameState& FrameState, uint32 FrameNumber)
 			{
 				check(!bBusy);
 				Timestamps.CPUFrameBoundary = Event.CPUTimestamp;
@@ -2054,7 +2054,7 @@ namespace UE::RHI::GPUProfiler
 			#if WITH_PROFILEGPU
 				if (Profile.bProfileFrame)
 				{
-					Profile.LogTree(Queue, Event.FrameNumber);
+					Profile.LogTree(*this, Event.FrameNumber);
 					Profile = {};
 				}
 			#endif
@@ -2071,7 +2071,7 @@ namespace UE::RHI::GPUProfiler
 			#endif
 
 			#if WITH_PROFILEGPU
-				if (ProfileFrameNumber == Event.FrameNumber + 1)
+				if (FrameNumber == Event.FrameNumber + 1)
 				{
 					Profile.bProfileFrame = true;
 
@@ -2079,17 +2079,17 @@ namespace UE::RHI::GPUProfiler
 					Profile.PushNode(TEXT("<root>"));
 
 				#if WITH_RHI_BREADCRUMBS
-					auto Recurse = [&](auto& Recurse, FRHIBreadcrumbNode* Breadcrumb) -> void
+					auto Recurse = [&](auto& Recurse, FRHIBreadcrumbNode* Current) -> void
 					{
-						if (!Breadcrumb)
+						if (!Current)
 						{
 							return;
 						}
 
-						Recurse(Recurse, Breadcrumb->GetParent());
+						Recurse(Recurse, Current->GetParent());
 
 						FRHIBreadcrumb::FBuffer Buffer;
-						Profile.PushNode(Breadcrumb->Name.GetTCHAR(Buffer));
+						Profile.PushNode(Current->Name.GetTCHAR(Buffer));
 					};
 					Recurse(Recurse, Event.Breadcrumb);
 				#endif // WITH_RHI_BREADCRUMBS
