Added validation to check that the platform supports lock no-overwrite.

diff --git a/Engine/Source/Runtime/RHI/Private/RHIValidation.cpp b/Engine/Source/Runtime/RHI/Private/RHIValidation.cpp
--- a/Engine/Source/Runtime/RHI/Private/RHIValidation.cpp
+++ b/Engine/Source/Runtime/RHI/Private/RHIValidation.cpp
@@ -952,6 +952,7 @@ void FValidationRHI::LockBufferValidate(class FRHICommandListBase& RHICmdList, F
 	using namespace RHIValidation;
 
 	check(GRHISupportsMultithreadedResources || RHICmdList.IsImmediate());
+	check(LockMode != RLM_WriteOnly_NoOverwrite || GRHIGlobals.SupportsMapWriteNoOverwrite)
 
 	if (!EnumHasAnyFlags(Buffer->GetUsage(), BUF_Volatile) && LockMode == RLM_WriteOnly)
 	{
