Remove FTaskTagScope Scope(ETaskTag::EParallelRenderingThread); because these tasks can be picked up by other threads like RHI as well - see GetDesiredThread() on the task

diff --git a/Engine/Source/Runtime/RHI/Private/PipelineStateCache.cpp b/Engine/Source/Runtime/RHI/Private/PipelineStateCache.cpp
--- a/Engine/Source/Runtime/RHI/Private/PipelineStateCache.cpp
+++ b/Engine/Source/Runtime/RHI/Private/PipelineStateCache.cpp
@@ -3511,7 +3511,6 @@ public:
 
 	void DoTask(ENamedThreads::Type CurrentThread, const FGraphEventRef& MyCompletionGraphEvent)
 	{
-		FTaskTagScope Scope(ETaskTag::EParallelRenderingThread);
 		FRayTracingPipelineState* RayTracingPipeline = static_cast<FRayTracingPipelineState*>(Pipeline);
 		check(!RayTracingPipeline->RHIPipeline.IsValid());
 		RayTracingPipeline->RHIPipeline = RHICreateRayTracingPipelineState(Initializer);
