VulkanRHI: Check if a project's selected feature levels are supported before using them in Linux.

#rnx

diff --git a/Engine/Source/Runtime/RHI/Private/Linux/LinuxDynamicRHI.cpp b/Engine/Source/Runtime/RHI/Private/Linux/LinuxDynamicRHI.cpp
--- a/Engine/Source/Runtime/RHI/Private/Linux/LinuxDynamicRHI.cpp
+++ b/Engine/Source/Runtime/RHI/Private/Linux/LinuxDynamicRHI.cpp
@@ -45,7 +45,7 @@ FDynamicRHI* PlatformCreateDynamicRHI()
 		// OpenGL can only be used for mobile preview.
 		bForceOpenGL = FParse::Param(FCommandLine::Get(), TEXT("opengl"));
 		ERHIFeatureLevel::Type PreviewFeatureLevel;
-		bool bUsePreviewFeatureLevel = RHIGetPreviewFeatureLevel(PreviewFeatureLevel);
+		const bool bUsePreviewFeatureLevel = RHIGetPreviewFeatureLevel(PreviewFeatureLevel);
 		if (bForceOpenGL && !bUsePreviewFeatureLevel)
 		{
 			FMessageDialog::Open(EAppMsgType::Ok, NSLOCTEXT("LinuxDynamicRHI", "OpenGLRemoved", "Warning: OpenGL is no longer supported for desktop platforms. Vulkan will be used instead."));
@@ -80,50 +80,72 @@ FDynamicRHI* PlatformCreateDynamicRHI()
 		TArray<FString> TargetedShaderFormats;
 		GConfig->GetArray(TEXT("/Script/LinuxTargetPlatform.LinuxTargetSettings"), TEXT("TargetedRHIs"), TargetedShaderFormats, GEngineIni);
 
-		// First come first serve
-		for (int32 SfIdx = 0; SfIdx < TargetedShaderFormats.Num(); ++SfIdx)
+		if (!bForceOpenGL)
 		{
-			// If we are not forcing opengl and we havent failed to create a VulkanRHI try to again with a different TargetedRHI
-			if (!bForceOpenGL && !bVulkanFailed && TargetedShaderFormats[SfIdx].StartsWith(TEXT("SF_VULKAN_")))
+			TArray<FString> VulkanTargetedShaderFormats = TargetedShaderFormats.FilterByPredicate([](const FString& ShaderFormatName) { return ShaderFormatName.StartsWith(TEXT("SF_VULKAN_")); });
+
+			if (VulkanTargetedShaderFormats.Num())
 			{
 				DynamicRHIModule = &FModuleManager::LoadModuleChecked<IDynamicRHIModule>(TEXT("VulkanRHI"));
-				if (!DynamicRHIModule->IsSupported())
+				if (DynamicRHIModule->IsSupported())
 				{
-					DynamicRHIModule = nullptr;
-					bVulkanFailed = true;
+					// Sort and start at the end to prefer higher feature levels (SM6 over SM5, for example).
+					VulkanTargetedShaderFormats.Sort();
+					for (int32 ShaderFormatIndex = VulkanTargetedShaderFormats.Num() - 1; ShaderFormatIndex >= 0; --ShaderFormatIndex)
+					{
+						const FName ShaderFormatName(*VulkanTargetedShaderFormats[ShaderFormatIndex]);
+						const EShaderPlatform TargetedPlatform = ShaderFormatToLegacyShaderPlatform(ShaderFormatName);
+						const ERHIFeatureLevel::Type MaxFeatureLevel = GetMaxSupportedFeatureLevel(TargetedPlatform);
+						if (DynamicRHIModule->IsSupported(MaxFeatureLevel))
+						{
+							bVulkanFailed = false;
+							RequestedFeatureLevel = MaxFeatureLevel;
+
+							FString FeatureLevelName;
+							GetFeatureLevelName(RequestedFeatureLevel, FeatureLevelName);
+							FApp::SetGraphicsRHI(FString::Printf(TEXT("Vulkan (%s)"), *FeatureLevelName));
+							FPlatformApplicationMisc::UsingVulkan();
+
+							break;
+						}
+						else
+						{
+							bVulkanFailed = true;
+						}
+					}
+
+					if (bVulkanFailed)
+					{
+						DynamicRHIModule = nullptr;
+					}
 				}
 				else
 				{
-					FPlatformApplicationMisc::UsingVulkan();
-
-					FName ShaderFormatName(*TargetedShaderFormats[SfIdx]);
-					EShaderPlatform TargetedPlatform = ShaderFormatToLegacyShaderPlatform(ShaderFormatName);
-					RequestedFeatureLevel = GetMaxSupportedFeatureLevel(TargetedPlatform);
-
-					FString FeatureLevelName;
-					GetFeatureLevelName(RequestedFeatureLevel, FeatureLevelName);
-					FApp::SetGraphicsRHI(FString::Printf(TEXT("Vulkan (%s)"), *FeatureLevelName));
-
-					break;
+					DynamicRHIModule = nullptr;
+					bVulkanFailed = true;
 				}
 			}
-			else if (!bForceVulkan && !bOpenGLFailed && TargetedShaderFormats[SfIdx].StartsWith(TEXT("GLSL_")))
+		}
+
+		if (!bForceVulkan && (DynamicRHIModule == nullptr))
+		{
+			TArray<FString> OGLTargetedShaderFormats = TargetedShaderFormats.FilterByPredicate([](const FString& ShaderFormatName) { return ShaderFormatName.StartsWith(TEXT("GLSL_")); });
+			if (OGLTargetedShaderFormats.Num())
 			{
 				DynamicRHIModule = &FModuleManager::LoadModuleChecked<IDynamicRHIModule>(TEXT("OpenGLDrv"));
-				if (!DynamicRHIModule->IsSupported())
-				{
-					DynamicRHIModule = nullptr;
-					bOpenGLFailed = true;
-				}
-				else
+				if (DynamicRHIModule->IsSupported())
 				{
 					FApp::SetGraphicsRHI(TEXT("OpenGL"));
 					FPlatformApplicationMisc::UsingOpenGL();
 
-					FName ShaderFormatName(*TargetedShaderFormats[SfIdx]);
+					FName ShaderFormatName(*TargetedShaderFormats[0]);
 					EShaderPlatform TargetedPlatform = ShaderFormatToLegacyShaderPlatform(ShaderFormatName);
 					RequestedFeatureLevel = GetMaxSupportedFeatureLevel(TargetedPlatform);
-					break;
+				}
+				else
+				{
+					DynamicRHIModule = nullptr;
+					bOpenGLFailed = true;
 				}
 			}
 		}
