Fixed a bug in PipelineStateCache::GetAndOrCreateComputePipelineState that would trigger an unnecessary stall on the render thread.
Precached PSOs should not be added as a dispatch prerequisite on the RHI command list, since they aren't used for drawing.

diff --git a/Engine/Source/Runtime/RHI/Private/PipelineStateCache.cpp b/Engine/Source/Runtime/RHI/Private/PipelineStateCache.cpp
--- a/Engine/Source/Runtime/RHI/Private/PipelineStateCache.cpp
+++ b/Engine/Source/Runtime/RHI/Private/PipelineStateCache.cpp
@@ -3396,7 +3396,8 @@ FComputePipelineState* PipelineStateCache::GetAndOrCreateComputePipelineState(FR
 		FGraphEventRef GraphEvent = OutCachedState->GetCompletionEvent();
 		InternalCreateComputePipelineState(ComputeShader, DoAsyncCompile, bFromFileCache, PSOPrecacheResult, OutCachedState, PSOCompilationDebugData, RHICmdList.IsImmediate());
 
-		if (GraphEvent.IsValid())
+		// Don't add precached PSOs as a dispatch prerequisite. We don't need to wait for them to complete before the RHICmdList can be dispatched.
+		if (GraphEvent.IsValid() && !bFromFileCache)
 		{
 			check(DoAsyncCompile);
 			RHICmdList.AddDispatchPrerequisite(GraphEvent);
