Build fixes when new GPU profiler is enabled in configs where RHI breadcrumbs are disabled

diff --git a/Engine/Source/Runtime/RHI/Private/GPUProfiler.cpp b/Engine/Source/Runtime/RHI/Private/GPUProfiler.cpp
--- a/Engine/Source/Runtime/RHI/Private/GPUProfiler.cpp
+++ b/Engine/Source/Runtime/RHI/Private/GPUProfiler.cpp
@@ -1569,8 +1569,10 @@ namespace UE::RHI::GPUProfiler
 			}
 
 			void EmitResults(FQueue Queue
-			#if HAS_GPU_STATS
+			#if STATS && HAS_GPU_STATS
 				, FEndOfPipeStats* Stats
+			#endif
+			#if HAS_GPU_STATS
 				, FGPUStat& GPUStat
 			#endif
 			#if CSV_PROFILER_STATS
@@ -1578,7 +1580,7 @@ namespace UE::RHI::GPUProfiler
 			#endif
 			) const
 			{
-			#if HAS_GPU_STATS
+			#if STATS && HAS_GPU_STATS
 				Stats->AddMessage(GPUStat.GetStatId(Queue, FGPUStat::EType::Busy).GetName(), EStatOperation::Set, FPlatformTime::ToMilliseconds64(BusyCycles));
 				Stats->AddMessage(GPUStat.GetStatId(Queue, FGPUStat::EType::Idle).GetName(), EStatOperation::Set, FPlatformTime::ToMilliseconds64(IdleCycles));
 				Stats->AddMessage(GPUStat.GetStatId(Queue, FGPUStat::EType::Wait).GetName(), EStatOperation::Set, FPlatformTime::ToMilliseconds64(WaitCycles));
@@ -1928,6 +1930,7 @@ namespace UE::RHI::GPUProfiler
 				LastGPUCycles = FMath::Max(LastGPUCycles, Event.GPUTimestampBOP);
 			}
 
+		#if WITH_RHI_BREADCRUMBS
 			void BeginBreadcrumb(FEvent::FBeginBreadcrumb const& Event)
 			{
 				uint64 Busy = Event.GPUTimestampTOP > LastGPUCycles
@@ -1997,6 +2000,7 @@ namespace UE::RHI::GPUProfiler
 				}
 			#endif
 			}
+		#endif
 
 			void Stats(FEvent::FStats const& Event)
 			{
@@ -2326,8 +2330,10 @@ namespace UE::RHI::GPUProfiler
 				for (auto const& [Stat, StatState] : QueueTimestamps.Stats)
 				{
 					StatState.EmitResults(Queue
-					#if HAS_GPU_STATS
+					#if STATS && HAS_GPU_STATS
 						, Stats
+					#endif
+					#if HAS_GPU_STATS
 						, *Stat.GPUStat
 					#endif
 					#if CSV_PROFILER_STATS
@@ -2339,8 +2345,10 @@ namespace UE::RHI::GPUProfiler
 
 				// Set the whole-frame per queue stat
 				QueueTimestamps.WholeQueueStat.EmitResults(Queue
-				#if HAS_GPU_STATS
+				#if STATS && HAS_GPU_STATS
 					, Stats
+				#endif
+				#if HAS_GPU_STATS
 					, GPUStat_Total
 				#endif
 				#if CSV_PROFILER_STATS
diff --git a/Engine/Source/Runtime/RHI/Private/GpuProfilerTrace.cpp b/Engine/Source/Runtime/RHI/Private/GpuProfilerTrace.cpp
--- a/Engine/Source/Runtime/RHI/Private/GpuProfilerTrace.cpp
+++ b/Engine/Source/Runtime/RHI/Private/GpuProfilerTrace.cpp
@@ -368,18 +368,22 @@ void FGpuProfilerTrace::FrameBoundary(uint32 QueueId, uint32 FrameId)
 
 void FGpuProfilerTrace::BeginBreadcrumb(uint32 QueueId, uint64 GPUTimestampTOP, const TCHAR* Name)
 {
+#if WITH_RHI_BREADCRUMBS
 	UE_TRACE_LOG(GpuProfiler, EventBeginBreadcrumb, GpuChannel)
 		<< EventBeginBreadcrumb.QueueId(QueueId)
 		<< EventBeginBreadcrumb.GPUTimestampTOP(GPUTimestampTOP)
 		<< EventBeginBreadcrumb.Name(Name);
+#endif
 }
 
 void FGpuProfilerTrace::EndBreadcrumb(uint32 QueueId, uint64 GPUTimestampTOP, const TCHAR* Name)
 {
+#if WITH_RHI_BREADCRUMBS
 	UE_TRACE_LOG(GpuProfiler, EventEndBreadcrumb, GpuChannel)
 		<< EventEndBreadcrumb.QueueId(QueueId)
 		<< EventEndBreadcrumb.GPUTimestampBOP(GPUTimestampTOP)
 		<< EventEndBreadcrumb.Name(Name); // for debugging
+#endif
 }
 
 void FGpuProfilerTrace::BeginWork(uint32 QueueId, uint64 GPUTimestampTOP, uint64 CPUTimestamp)
