Fixed RHI validation for async compute.

diff --git a/Engine/Source/Runtime/RHI/Private/RHIValidation.cpp b/Engine/Source/Runtime/RHI/Private/RHIValidation.cpp
--- a/Engine/Source/Runtime/RHI/Private/RHIValidation.cpp
+++ b/Engine/Source/Runtime/RHI/Private/RHIValidation.cpp
@@ -955,7 +955,7 @@ void FValidationRHI::LockBufferValidate(class FRHICommandListBase& RHICmdList, F
 	check(GRHISupportsMultithreadedResources || RHICmdList.IsImmediate());
 	check(LockMode != RLM_WriteOnly_NoOverwrite || GRHIGlobals.SupportsMapWriteNoOverwrite)
 
-	if (!EnumHasAnyFlags(Buffer->GetUsage(), BUF_Volatile) && LockMode == RLM_WriteOnly)
+	if (RHICmdList.IsGraphics() && !EnumHasAnyFlags(Buffer->GetUsage(), BUF_Volatile) && LockMode == RLM_WriteOnly)
 	{
 		bool bIsInsideRenderPass;
 		if (RHICmdList.IsTopOfPipe())
@@ -1109,7 +1109,7 @@ void FValidationRHI::RHIBindDebugLabelName(FRHICommandListBase& RHICmdList, FRHI
 	FString NameCopyRT = Name;
 	RHICmdList.EnqueueLambda([Texture, NameCopyRHIT = MoveTemp(NameCopyRT)](FRHICommandListBase& RHICmdList)
 	{
-		((FValidationContext&)RHICmdList.GetContext()).Tracker->Rename(Texture->GetTrackerResource(), *NameCopyRHIT);
+		((FValidationContext&)RHICmdList.GetComputeContext()).Tracker->Rename(Texture->GetTrackerResource(), *NameCopyRHIT);
 	});
 
 	RHI->RHIBindDebugLabelName(RHICmdList, Texture, Name);
@@ -1120,7 +1120,7 @@ void FValidationRHI::RHIBindDebugLabelName(FRHICommandListBase& RHICmdList, FRHI
 	FString NameCopyRT = Name;
 	RHICmdList.EnqueueLambda([Buffer, NameCopyRHIT = MoveTemp(NameCopyRT)](FRHICommandListBase& RHICmdList)
 	{
-		((FValidationContext&)RHICmdList.GetContext()).Tracker->Rename(Buffer, *NameCopyRHIT);
+		((FValidationContext&)RHICmdList.GetComputeContext()).Tracker->Rename(Buffer, *NameCopyRHIT);
 	});
 
 	RHI->RHIBindDebugLabelName(RHICmdList, Buffer, Name);
@@ -1132,7 +1132,7 @@ void FValidationRHI::RHIBindDebugLabelName(FRHICommandListBase& RHICmdList, FRHI
 	FString NameCopyRT = Name;
 	RHICmdList.EnqueueLambda([Resource, NameCopyRHIT = MoveTemp(NameCopyRT)](FRHICommandListBase& RHICmdList)
 	{
-		((FValidationContext&)RHICmdList.GetContext()).Tracker->Rename(Resource, *NameCopyRHIT);
+		((FValidationContext&)RHICmdList.GetComputeContext()).Tracker->Rename(Resource, *NameCopyRHIT);
 	});
 
 	RHI->RHIBindDebugLabelName(RHICmdList, UnorderedAccessViewRHI, Name);
