Fixed D3D11 NoOverwrite lock mode from flushing the RHI thread. Optimized SetTrackedAccess to use a single command.

diff --git a/Engine/Source/Runtime/RHI/Private/RHICommandList.cpp b/Engine/Source/Runtime/RHI/Private/RHICommandList.cpp
--- a/Engine/Source/Runtime/RHI/Private/RHICommandList.cpp
+++ b/Engine/Source/Runtime/RHI/Private/RHICommandList.cpp
@@ -1714,16 +1714,7 @@ void FRHICommandListImmediate::EndFrame()
 void FRHIComputeCommandList::Transition(TArrayView<const FRHITransitionInfo> Infos, ERHITransitionCreateFlags CreateFlags)
 {
 	TransitionInternal(Infos, CreateFlags);
-
-	for (const FRHITransitionInfo& Info : Infos)
-	{
-		ensureMsgf(Info.IsWholeResource(), TEXT("The Transition method only supports whole resource transitions."));
-
-		if (FRHIViewableResource* Resource = GetViewableResource(Info))
-		{
-			SetTrackedAccess({ FRHITrackedAccessInfo(Resource, Info.AccessAfter) });
-		}
-	}
+	SetTrackedAccess(Infos);
 }
 
 void FRHIComputeCommandList::Transition(TArrayView<const FRHITransitionInfo> Infos, ERHIPipeline SrcPipelines, ERHIPipeline DstPipelines, ERHITransitionCreateFlags TransitionCreateFlags)
@@ -1845,7 +1836,7 @@ void* FDynamicRHI::RHILockBuffer(class FRHICommandListBase& RHICmdList, FRHIBuff
 	if (RHICmdList.IsTopOfPipe())
 	{
 		bool bBuffer = CVarRHICmdBufferWriteLocks.GetValueOnRenderThread() > 0;
-		if (!bBuffer || LockMode != RLM_WriteOnly)
+		if (!bBuffer || LockMode == RLM_ReadOnly)
 		{
 			QUICK_SCOPE_CYCLE_COUNTER(STAT_RHIMETHOD_LockBuffer_FlushAndLock);
 			CSV_SCOPED_TIMING_STAT(RHITFlushes, LockBuffer_BottomOfPipe);
diff --git a/Engine/Source/Runtime/RHI/Public/RHICommandList.h b/Engine/Source/Runtime/RHI/Public/RHICommandList.h
--- a/Engine/Source/Runtime/RHI/Public/RHICommandList.h
+++ b/Engine/Source/Runtime/RHI/Public/RHICommandList.h
@@ -3141,13 +3141,23 @@ public:
 
 	FORCEINLINE_DEBUGGABLE void SetTrackedAccess(TArrayView<const FRHITransitionInfo> Infos)
 	{
+		FRHITrackedAccessInfo* TrackedAccessInfos = reinterpret_cast<FRHITrackedAccessInfo*>(Alloc(sizeof(FRHITrackedAccessInfo) * Infos.Num(), alignof(FRHITrackedAccessInfo)));
+		int32 NumTrackedAccessInfos = 0;
+
 		for (const FRHITransitionInfo& Info : Infos)
 		{
+			ensureMsgf(Info.IsWholeResource(), TEXT("The Transition method only supports whole resource transitions."));
+
 			if (FRHIViewableResource* Resource = GetViewableResource(Info))
 			{
-				SetTrackedAccess({ FRHITrackedAccessInfo(Resource, Info.AccessAfter) });
+				new (&TrackedAccessInfos[NumTrackedAccessInfos++]) FRHITrackedAccessInfo(Resource, Info.AccessAfter);
 			}
 		}
+
+		if (NumTrackedAccessInfos > 0)
+		{
+			SetTrackedAccess(TArrayView<FRHITrackedAccessInfo>(TrackedAccessInfos, NumTrackedAccessInfos));
+		}
 	}
 
 	FORCEINLINE_DEBUGGABLE void SetShaderRootConstants(const FUint32Vector4& Constants)
