Temp disable RHI validation per persistent SBTs because record clear operation are not supported yet and validation SBT can end up with dangling resources then

diff --git a/Engine/Source/Runtime/RHI/Private/RHIValidation.cpp b/Engine/Source/Runtime/RHI/Private/RHIValidation.cpp
--- a/Engine/Source/Runtime/RHI/Private/RHIValidation.cpp
+++ b/Engine/Source/Runtime/RHI/Private/RHIValidation.cpp
@@ -356,6 +356,7 @@ namespace RHIValidation
 
 	FShaderBindingTable::FShaderBindingTable(const FRayTracingShaderBindingTableInitializer& InInitializer)
 	{
+		LifeTime = InInitializer.Lifetime;
 		ShaderBindingMode = InInitializer.ShaderBindingMode;
 		HitGroupIndexingMode = InInitializer.HitGroupIndexingMode;
 	}
@@ -460,52 +461,57 @@ namespace RHIValidation
 	void FShaderBindingTable::SetBindingsOnShaderBindingTable(FRayTracingPipelineState* RayTracingPipelineState, uint32 NumBindings, const FRayTracingLocalShaderBindings* Bindings, ERayTracingBindingType BindingType)
 	{		
 		TRACE_CPUPROFILER_EVENT_SCOPE(RHIValidation-SetBindingsOnShaderBindingTable);
-		
-		FGraphEventArray TaskList;
-
-		const uint32 NumWorkerThreads = FTaskGraphInterface::Get().GetNumWorkerThreads();
-		const uint32 MaxTasks = FApp::ShouldUseThreadingForPerformance() ? FMath::Min<uint32>(NumWorkerThreads, MaxBindingWorkers) : 1;
 
-		struct FTaskContext
+		// Disable tracking for persistent SBTs until per record tracking is implemented otherwise
+		// it might end up with dangling SVR/UAV pointers 
+		if (LifeTime != ERayTracingShaderBindingTableLifetime::Persistent)
 		{
-			uint32 WorkerIndex = 0;
-		};
+			FGraphEventArray TaskList;
 
-		TArray<FTaskContext, TInlineAllocator<MaxBindingWorkers>> TaskContexts;
-		for (uint32 WorkerIndex = 0; WorkerIndex < MaxTasks; ++WorkerIndex)
-		{
-			TaskContexts.Add(FTaskContext{WorkerIndex});
-		}
+			const uint32 NumWorkerThreads = FTaskGraphInterface::Get().GetNumWorkerThreads();
+			const uint32 MaxTasks = FApp::ShouldUseThreadingForPerformance() ? FMath::Min<uint32>(NumWorkerThreads, MaxBindingWorkers) : 1;
 
-		auto BindingTask = [this, RayTracingPipelineState, Bindings, BindingType](const FTaskContext& Context, int32 CurrentIndex)
-		{
-			const FRayTracingLocalShaderBindings& Binding = Bindings[CurrentIndex];
+			struct FTaskContext
+			{
+				uint32 WorkerIndex = 0;
+			};
 
-			// only collect shader binding data if RTPSO & hit group indexing mode
-			bool bValidBinding = (Binding.BindingType == ERayTracingLocalShaderBindingType::Persistent || Binding.BindingType == ERayTracingLocalShaderBindingType::Validation);
-			if (HitGroupIndexingMode == ERayTracingHitGroupIndexingMode::Allow && bValidBinding)
+			TArray<FTaskContext, TInlineAllocator<MaxBindingWorkers>> TaskContexts;
+			for (uint32 WorkerIndex = 0; WorkerIndex < MaxTasks; ++WorkerIndex)
 			{
-				if (EnumHasAnyFlags(ShaderBindingMode, ERayTracingShaderBindingMode::RTPSO))
-				{					
-					CollectShaderBindingTableResources(RayTracingPipelineState, this, Binding, BindingType, Context.WorkerIndex);
+				TaskContexts.Add(FTaskContext{ WorkerIndex });
+			}
 
-					// Also add SRV view requirement for all index and vertex buffers used in the SBT
-					const FRayTracingGeometryInitializer& BLASInitializer = Binding.Geometry->GetInitializer();
-					if (BLASInitializer.IndexBuffer)
-					{
-						AddSRV(BLASInitializer.IndexBuffer->GetWholeResourceIdentity(), Context.WorkerIndex);
-					}
-					for (int32 SegmentIndex = 0; SegmentIndex < BLASInitializer.Segments.Num(); ++SegmentIndex)
+			auto BindingTask = [this, RayTracingPipelineState, Bindings, BindingType](const FTaskContext& Context, int32 CurrentIndex)
+				{
+					const FRayTracingLocalShaderBindings& Binding = Bindings[CurrentIndex];
+
+					// only collect shader binding data if RTPSO & hit group indexing mode
+					bool bValidBinding = (Binding.BindingType == ERayTracingLocalShaderBindingType::Persistent || Binding.BindingType == ERayTracingLocalShaderBindingType::Validation);
+					if (HitGroupIndexingMode == ERayTracingHitGroupIndexingMode::Allow && bValidBinding)
 					{
-						AddSRV(BLASInitializer.Segments[SegmentIndex].VertexBuffer->GetWholeResourceIdentity(), Context.WorkerIndex);
+						if (EnumHasAnyFlags(ShaderBindingMode, ERayTracingShaderBindingMode::RTPSO))
+						{
+							CollectShaderBindingTableResources(RayTracingPipelineState, this, Binding, BindingType, Context.WorkerIndex);
+
+							// Also add SRV view requirement for all index and vertex buffers used in the SBT
+							const FRayTracingGeometryInitializer& BLASInitializer = Binding.Geometry->GetInitializer();
+							if (BLASInitializer.IndexBuffer)
+							{
+								AddSRV(BLASInitializer.IndexBuffer->GetWholeResourceIdentity(), Context.WorkerIndex);
+							}
+							for (int32 SegmentIndex = 0; SegmentIndex < BLASInitializer.Segments.Num(); ++SegmentIndex)
+							{
+								AddSRV(BLASInitializer.Segments[SegmentIndex].VertexBuffer->GetWholeResourceIdentity(), Context.WorkerIndex);
+							}
+						}
 					}
-				}
-			}
-		};
+				};
 
-		const int32 ItemsPerTask = 1024;
-		ParallelForWithExistingTaskContext(TEXT("SetRayTracingBindings"), MakeArrayView(TaskContexts), NumBindings, ItemsPerTask, BindingTask);
-				
+			const int32 ItemsPerTask = 1024;
+			ParallelForWithExistingTaskContext(TEXT("SetRayTracingBindings"), MakeArrayView(TaskContexts), NumBindings, ItemsPerTask, BindingTask);
+		}
+		
 		// Mark dirty
 		bIsDirty = true;
 	}
diff --git a/Engine/Source/Runtime/RHI/Public/RHIValidationCommon.h b/Engine/Source/Runtime/RHI/Public/RHIValidationCommon.h
--- a/Engine/Source/Runtime/RHI/Public/RHIValidationCommon.h
+++ b/Engine/Source/Runtime/RHI/Public/RHIValidationCommon.h
@@ -43,6 +43,7 @@ class FRHIRayTracingPipelineState;
 struct FRayTracingShaderBindingTableInitializer;
 struct FRayTracingLocalShaderBindings;
 
+enum class ERayTracingShaderBindingTableLifetime : uint8;
 enum class ERayTracingBindingType : uint8;
 enum class ERayTracingShaderBindingMode : uint8;
 enum class ERayTracingHitGroupIndexingMode : uint8;
@@ -578,7 +579,8 @@ namespace RHIValidation
 		}
 
 	private:
-						
+		
+		ERayTracingShaderBindingTableLifetime LifeTime;
 		ERayTracingShaderBindingMode ShaderBindingMode;
 		ERayTracingHitGroupIndexingMode HitGroupIndexingMode;
 		bool bIsDirty = true;
