Fix build under modular build, RHI_EXECUTE_API macros should be used for RHI execute functions

#rnx

diff --git a/Engine/Source/Runtime/RHI/Public/RHICommandList.h b/Engine/Source/Runtime/RHI/Public/RHICommandList.h
--- a/Engine/Source/Runtime/RHI/Public/RHICommandList.h
+++ b/Engine/Source/Runtime/RHI/Public/RHICommandList.h
@@ -2433,7 +2433,7 @@ FRHICOMMAND_MACRO(FRHICommandSetComputeBudget)
 		: Budget(InBudget)
 	{
 	}
-	RHI_API void Execute(FRHICommandListBase & CmdList);
+	RHI_EXECUTE_API void Execute(FRHICommandListBase & CmdList);
 };
 
 FRHICOMMAND_MACRO(FRHICommandCopyToStagingBuffer)
