Fixed RHI validation for clearing UAV's on async compute without UAV overlap enabled.

diff --git a/Engine/Source/Runtime/RHI/Public/RHIValidationContext.h b/Engine/Source/Runtime/RHI/Public/RHIValidationContext.h
--- a/Engine/Source/Runtime/RHI/Public/RHIValidationContext.h
+++ b/Engine/Source/Runtime/RHI/Public/RHIValidationContext.h
@@ -139,7 +139,8 @@ public:
 
 	virtual void RHIClearUAVUint(FRHIUnorderedAccessView* UnorderedAccessViewRHI, const FUintVector4& Values) override final
 	{
-		Tracker->Assert(UnorderedAccessViewRHI->GetViewIdentity(), ERHIAccess::UAVCompute);
+		// @todo should we assert here? If the base RHI uses a compute shader via
+		// FRHICommandList_RecursiveHazardous then we might double-assert which breaks the tracking
 		RHIContext->RHIClearUAVUint(UnorderedAccessViewRHI, Values);
 	}
 
