Fixed pipeline state cache cleanup on non-rendering threads/tasks
- Fixes rare assert on GPU crash handling

diff --git a/Engine/Source/Runtime/RHI/Private/PipelineStateCache.cpp b/Engine/Source/Runtime/RHI/Private/PipelineStateCache.cpp
--- a/Engine/Source/Runtime/RHI/Private/PipelineStateCache.cpp
+++ b/Engine/Source/Runtime/RHI/Private/PipelineStateCache.cpp
@@ -1900,6 +1900,28 @@ private:
 		Completed.Reset();
 	}
 
+	template<typename F>
+	void ExecuteImmediateCommand(F&& Functor)
+	{
+		if (IsInRenderingThread())
+		{
+			FRHICommandListImmediate::Get().EnqueueLambda([Functor](FRHICommandListImmediate& RHICmdList) mutable
+			{
+				Functor(RHICmdList);
+			});
+		}
+		else
+		{
+			FFunctionGraphTask::CreateAndDispatchWhenReady([Functor]() mutable
+			{
+				FRHICommandListImmediate::Get().EnqueueLambda([Functor = MoveTemp(Functor)](FRHICommandListImmediate& RHICmdList) mutable
+				{
+					Functor(RHICmdList);
+				});
+			}, TStatId{}, nullptr, ENamedThreads::ActualRenderingThread);
+		}
+	}
+
 	void ProcessDelayedCleanup()
 	{
 		if (DeleteArray.IsEmpty())
@@ -1907,7 +1929,7 @@ private:
 			return;
 		}
 
-		FRHICommandListImmediate::Get().EnqueueLambda([DeleteArray = MoveTemp(DeleteArray)](FRHICommandListImmediate& RHICmdList) mutable
+		ExecuteImmediateCommand([DeleteArray = MoveTemp(DeleteArray)](FRHICommandListImmediate& RHICmdList) mutable
 		{
 			for (TMyValue& OldPipelineState : DeleteArray)
 			{
@@ -1934,7 +1956,7 @@ private:
 		int32 Discarded = InOutBackfillMap->Num();
 		if (Discarded > 0)
 		{
-			FRHICommandListImmediate::Get().EnqueueLambda([DiscardMap = MoveTemp(*InOutBackfillMap)](FRHICommandListImmediate& RHICmdList) mutable
+			ExecuteImmediateCommand([DiscardMap = MoveTemp(*InOutBackfillMap)](FRHICommandListImmediate& RHICmdList) mutable
 			{
 				for (const auto& DiscardIterator : DiscardMap)
 				{
