Using SetNumUninitialized on an array of FRHITransitionInfo and then setting the elements with operator= is bad, because FRHITransitionInfo is not a POD.

This caused bad codegen with some compilers when copying the CommitInfo field, which is a TOptional. The bIsSet field wasn't copied, so regular resources appeared to have commit requests, leading to crashes later on. This change also modifies TOptional::Destroy to explicitly destroy the value storage, to make it clear that we're not invalidating the entire Value object.

Although not related to the original problem, I've made the FRHICommitResourceInfo constructor explicit, since we don't want ints to silently convert to that type.

#rnx

diff --git a/Engine/Source/Runtime/Core/Public/Misc/Optional.h b/Engine/Source/Runtime/Core/Public/Misc/Optional.h
--- a/Engine/Source/Runtime/Core/Public/Misc/Optional.h
+++ b/Engine/Source/Runtime/Core/Public/Misc/Optional.h
@@ -409,7 +409,14 @@ private:
 	 */
 	FORCEINLINE void DestroyValue()
 	{
-		DestructItem((OptionalType*)&Value);
+		if constexpr (bUsingIntrusiveUnsetState)
+		{
+			DestructItem((OptionalType*)&Value);
+		}
+		else
+		{
+			DestructItem((OptionalType*)Value.Storage);
+		}
 	}
 
 	using ValueStorageType = std::conditional_t<bUsingIntrusiveUnsetState, uint8[sizeof(OptionalType)], UE::Core::Private::TNonIntrusiveOptionalStorage<sizeof(OptionalType)>>;
diff --git a/Engine/Source/Runtime/RHI/Public/RHITransition.h b/Engine/Source/Runtime/RHI/Public/RHITransition.h
--- a/Engine/Source/Runtime/RHI/Public/RHITransition.h
+++ b/Engine/Source/Runtime/RHI/Public/RHITransition.h
@@ -112,7 +112,7 @@ struct FRHISubresourceRange
 struct FRHICommitResourceInfo
 {
 	uint64 SizeInBytes = 0;
-	FRHICommitResourceInfo(uint64 InSizeInBytes) : SizeInBytes(InSizeInBytes) {}
+	explicit FRHICommitResourceInfo(uint64 InSizeInBytes) : SizeInBytes(InSizeInBytes) {}
 };
 
 struct FRHITransitionInfo : public FRHISubresourceRange
@@ -763,17 +763,8 @@ public:
 		DstPipelines = CreationInfo.DstPipelines;
 		m_Flags = CreationInfo.Flags;
 
-		TransitionInfos.SetNumUninitialized(CreationInfo.TransitionInfos.Num());
-		for (int32 i = 0; i < CreationInfo.TransitionInfos.Num(); i++)
-		{
-			TransitionInfos[i] = CreationInfo.TransitionInfos[i];
-		}
-
-		AliasingInfos.SetNumUninitialized(CreationInfo.AliasingInfos.Num());
-		for (int32 i = 0; i < CreationInfo.AliasingInfos.Num(); i++)
-		{
-			AliasingInfos[i] = CreationInfo.AliasingInfos[i];
-		}
+		TransitionInfos = CreationInfo.TransitionInfos;
+		AliasingInfos = CreationInfo.AliasingInfos;
 
 		// Deep copy of AlisianInfo used for validation
 		if (GRHIValidationEnabled)
