VulkanRHI: extend lifetime of breadcrumb allocator for FFramBoundary also.

#rnx

diff --git a/Engine/Source/Runtime/RHI/Public/GPUProfiler.h b/Engine/Source/Runtime/RHI/Public/GPUProfiler.h
--- a/Engine/Source/Runtime/RHI/Public/GPUProfiler.h
+++ b/Engine/Source/Runtime/RHI/Public/GPUProfiler.h
@@ -412,12 +412,16 @@ namespace UE::RHI::GPUProfiler
 		#if WITH_RHI_BREADCRUMBS
 			if constexpr (
 				std::is_same_v<UE::RHI::GPUProfiler::FEvent::FBeginBreadcrumb, TEventType> ||
-				std::is_same_v<UE::RHI::GPUProfiler::FEvent::FEndBreadcrumb  , TEventType>
+				std::is_same_v<UE::RHI::GPUProfiler::FEvent::FEndBreadcrumb  , TEventType> ||
+				std::is_same_v<UE::RHI::GPUProfiler::FEvent::FFrameBoundary  , TEventType>
 				)
 			{
-				// Attach the breadcrumb allocator for begin/end breadcrumb events.
-				// This keeps the breadcrumbs alive until the events have been consumed by the profilers.
-				Current->Header.BreadcrumbAllocators.AddUnique(Data.Breadcrumb->Allocator);
+				if (Data.Breadcrumb)
+				{
+					// Attach the breadcrumb allocator for begin/end breadcrumb events.
+					// This keeps the breadcrumbs alive until the events have been consumed by the profilers.
+					Current->Header.BreadcrumbAllocators.AddUnique(Data.Breadcrumb->Allocator);
+				}
 			}
 		#endif
 
