Remove unrequired check to see if SBT is dirty - code is safe to handle non dirty sbt commits

diff --git a/Engine/Source/Runtime/RHI/Private/RHIValidation.cpp b/Engine/Source/Runtime/RHI/Private/RHIValidation.cpp
--- a/Engine/Source/Runtime/RHI/Private/RHIValidation.cpp
+++ b/Engine/Source/Runtime/RHI/Private/RHIValidation.cpp
@@ -517,9 +517,7 @@ namespace RHIValidation
 	}
 
 	void FShaderBindingTable::Commit()
-	{		
-		ensureMsgf(bIsDirty, TEXT("bIsDirty should be checked before calling Commit()"));
-		
+	{				
 		// Merge all data from worker threads into the main set
 		for (uint32 WorkerIndex = 1; WorkerIndex < MaxBindingWorkers; ++WorkerIndex)
 		{
