Improved RHI validation around the EResourceTransitionFlags::IgnoreAfterState case:
- assert in case of using the flag on all pipes
- do not replicate on all pipes the state so the assert can still work in that case
- test for sequential paired calls to EResourceTransitionFlags::IgnoreAfterState

diff --git a/Engine/Source/Runtime/RHI/Private/RHIValidation.cpp b/Engine/Source/Runtime/RHI/Private/RHIValidation.cpp
--- a/Engine/Source/Runtime/RHI/Private/RHIValidation.cpp
+++ b/Engine/Source/Runtime/RHI/Private/RHIValidation.cpp
@@ -1986,6 +1986,50 @@ namespace RHIValidation
 			*GetRHIPipelineName(RequiredState.Pipelines));
 	}
 
+	static inline FString GetReasonString_IgnoreAfterStateAllPipes(
+		FResource* Resource, FSubresourceIndex const& SubresourceIndex,
+		const FState& PendingState,
+		const FState& TargetState)
+	{
+		FString DebugName = GetResourceDebugName(Resource, SubresourceIndex);
+		return FString::Printf(
+			BARRIER_TRACKER_LOG_PREFIX_RESNAME
+			TEXT("Attempted to begin a resource transition for resource %s on All pipes. Transitionwith  EResourceTransitionFlags::IgnoreAfterState on All pipes are not supported.\n\n")
+			TEXT("    --- Pending access states for this resource are:              %s\n")
+			TEXT("    --- Attempted access states for the current transition are: %s\n")
+			TEXT("    --- Pending pipelines for this resource are:                  %s\n")
+			TEXT("    --- Attempted pipelines for the current transition are:     %s\n")
+			BARRIER_TRACKER_LOG_SUFFIX,
+			*DebugName,
+			*DebugName,
+			*GetRHIAccessName(PendingState.Access),
+			*GetRHIAccessName(TargetState.Access),
+			*GetRHIPipelineName(PendingState.Pipelines),
+			*GetRHIPipelineName(TargetState.Pipelines));
+	}
+
+	static inline FString GetReasonString_MismatchedIgnoreAfterState(
+			FResource* Resource, FSubresourceIndex const& SubresourceIndex,
+			const FState& PendingState,
+			const FState& TargetState)
+	{
+		FString DebugName = GetResourceDebugName(Resource, SubresourceIndex);
+		return FString::Printf(
+			BARRIER_TRACKER_LOG_PREFIX_RESNAME
+			TEXT("Attempted to begin a resource transition for resource %s whilst not having the flag EResourceTransitionFlags::IgnoreAfterState matching the previous transiton. Transition with EResourceTransitionFlags::IgnoreAfterState always needs to be done in pairs.\n\n")
+			TEXT("    --- Pending access states for this resource are:              %s\n")
+			TEXT("    --- Attempted access states for the current transition are: %s\n")
+			TEXT("    --- Pending pipelines for this resource are:                  %s\n")
+			TEXT("    --- Attempted pipelines for the current transition are:     %s\n")
+			BARRIER_TRACKER_LOG_SUFFIX,
+			*DebugName,
+			*DebugName,
+			*GetRHIAccessName(PendingState.Access),
+			*GetRHIAccessName(TargetState.Access),
+			*GetRHIPipelineName(PendingState.Pipelines),
+			*GetRHIPipelineName(TargetState.Pipelines));
+	}
+
 	static inline void* Log(FResource* Resource, FSubresourceIndex const& SubresourceIndex, void* CreateTrace, const TCHAR* TracePrefix, const TCHAR* Type, const TCHAR* LogStr)
 	{
 		void* Trace = CaptureBacktrace();
@@ -2187,39 +2231,55 @@ namespace RHIValidation
 				return Previous == Next;
 			};
 
+
+			bool bHasMatchingPipelines = true;
+
+			// We do not check pipelines for IgnoreAfterState since we do not replicate those transition on all pipelines.
 			if (EnumHasAnyFlags(NewFlags, EResourceTransitionFlags::IgnoreAfterState) == false)
+			{ 
+				bHasMatchingPipelines = HasMatchingPipelines(State.Previous.Pipelines, CurrentStateFromRHI.Pipelines);
+			}
+
+			if (CurrentStateFromRHI.Access == ERHIAccess::Unknown)
 			{
-				if (CurrentStateFromRHI.Access == ERHIAccess::Unknown)
-				{
-					RHI_VALIDATION_CHECK(Resource->TrackedAccess == State.Previous.Access && HasMatchingPipelines(State.Previous.Pipelines, CurrentStateFromRHI.Pipelines),
-						*GetReasonString_IncorrectPreviousTrackedState(Resource, SubresourceIndex, State.Previous, CurrentStateFromRHI.Pipelines));
-				}
-				else
-				{
-					// Check the current RHI state passed in matches the tracked state for the resource.
-					RHI_VALIDATION_CHECK(CurrentStateFromRHI.Access == State.Previous.Access && HasMatchingPipelines(State.Previous.Pipelines, CurrentStateFromRHI.Pipelines),
-						*GetReasonString_IncorrectPreviousExplicitState(Resource, SubresourceIndex, State.Previous, CurrentStateFromRHI));
-				}
+				RHI_VALIDATION_CHECK(Resource->TrackedAccess == State.Previous.Access && bHasMatchingPipelines,
+					*GetReasonString_IncorrectPreviousTrackedState(Resource, SubresourceIndex, State.Previous, CurrentStateFromRHI.Pipelines));
+			}
+			else
+			{
+				// Check the current RHI state passed in matches the tracked state for the resource.
+				RHI_VALIDATION_CHECK(CurrentStateFromRHI.Access == State.Previous.Access && bHasMatchingPipelines,
+					*GetReasonString_IncorrectPreviousExplicitState(Resource, SubresourceIndex, State.Previous, CurrentStateFromRHI));
 			}
 		}
 
+		const bool bIgnoreAfterStateAllPipes = ((EnumHasAnyFlags(NewFlags, EResourceTransitionFlags::IgnoreAfterState)) && (CurrentStateFromRHI.Pipelines == ERHIPipeline::All));
+		RHI_VALIDATION_CHECK(bIgnoreAfterStateAllPipes == false,
+			*GetReasonString_IgnoreAfterStateAllPipes(Resource, SubresourceIndex, State.Current, TargetState));
+
+		const bool bRegularTransitionWhileIgnoring = (State.bIgnoringAfterState && (EnumHasAnyFlags(NewFlags, EResourceTransitionFlags::IgnoreAfterState) == false));
+
+		RHI_VALIDATION_CHECK(bRegularTransitionWhileIgnoring == false,
+				*GetReasonString_MismatchedIgnoreAfterState(Resource, SubresourceIndex, State.Current, TargetState));
+
 		// Check for unnecessary transitions
 		// @todo: this check is not particularly useful at the moment, as there are many unnecessary resource transitions.
 		//RHI_VALIDATION_CHECK(CurrentState != TargetState, *GetReasonString_UnnecessaryTransition(Resource, SubresourceIndex, CurrentState));
 
-		if (EnumHasAnyFlags(NewFlags, EResourceTransitionFlags::IgnoreAfterState) == false)
-		{
-			State.bTransitioning = true;
-			State.Flags = NewFlags;
-			State.BeginTransitionBacktrace = BeginTrace;
-			// Update the tracked state once all pipes have begun.
-			State.Previous = TargetState;
-			State.Current = TargetState;
-			State.CreateTransitionBacktrace = CreateTrace;
-			State.bUsedWithAllUAVsOverlap = false;
-			State.bUsedWithExplicitUAVsOverlap = false;
 
+		State.bTransitioning = true;
+		State.Flags = NewFlags;
+		State.BeginTransitionBacktrace = BeginTrace;
+		// Update the tracked state once all pipes have begun.
+		State.Previous = TargetState;
+		State.Current = TargetState;
+		State.CreateTransitionBacktrace = CreateTrace;
+		State.bUsedWithAllUAVsOverlap = false;
+		State.bUsedWithExplicitUAVsOverlap = false;
 
+		// Do not replicate the state for EResourceTransitionFlags::IgnoreAfterState transition
+		if (EnumHasAnyFlags(NewFlags, EResourceTransitionFlags::IgnoreAfterState) == false)
+		{
 			// Replicate the state to other pipes that are not part of the begin pipe mask.
 			for (ERHIPipeline OtherPipeline : MakeFlagsRange(ERHIPipeline::All & ~CurrentStateFromRHI.Pipelines))
 			{
@@ -2244,18 +2304,25 @@ namespace RHIValidation
 			));
 		}
 
-		if (EnumHasAnyFlags(NewFlags, EResourceTransitionFlags::IgnoreAfterState) == false)
+		FPipelineState& State = States[ExecutingPipeline];
+
+		// Set bIgnoringAfterState to true in case of EResourceTransitionFlags::IgnoreAfterState so on the next transition we can check that EResourceTransitionFlags::IgnoreAfterState is used as well
+		if ( EnumHasAnyFlags(NewFlags, EResourceTransitionFlags::IgnoreAfterState))
 		{
-			FPipelineState& State = States[ExecutingPipeline];
+			State.bIgnoringAfterState = !State.bIgnoringAfterState;
+		}
 
-			// Check that we aren't ending a transition that never began.
-			RHI_VALIDATION_CHECK(State.bTransitioning, TEXT("Unsolicited resource end transition call."));
-			State.bTransitioning = false;
-			State.BeginTransitionBacktrace = nullptr;
+		// Check that we aren't ending a transition that never began.
+		RHI_VALIDATION_CHECK(State.bTransitioning, TEXT("Unsolicited resource end transition call."));
+		State.bTransitioning = false;
+		State.BeginTransitionBacktrace = nullptr;
 
-			// Check that the end matches the begin.
-			RHI_VALIDATION_CHECK(TargetState == State.Current, *GetReasonString_MismatchedEndTransition(Resource, SubresourceIndex, State.Current, TargetState));
+		// Check that the end matches the begin.
+		RHI_VALIDATION_CHECK(TargetState == State.Current, *GetReasonString_MismatchedEndTransition(Resource, SubresourceIndex, State.Current, TargetState));
 
+		// Do not replicate the state for EResourceTransitionFlags::IgnoreAfterState transitions
+		if (EnumHasAnyFlags(NewFlags, EResourceTransitionFlags::IgnoreAfterState) == false)
+		{
 			// Replicate the state to other pipes that are not part of the end pipe mask.
 			for (ERHIPipeline OtherPipeline : MakeFlagsRange(ERHIPipeline::All))
 			{
diff --git a/Engine/Source/Runtime/RHI/Public/RHIValidationCommon.h b/Engine/Source/Runtime/RHI/Public/RHIValidationCommon.h
--- a/Engine/Source/Runtime/RHI/Public/RHIValidationCommon.h
+++ b/Engine/Source/Runtime/RHI/Public/RHIValidationCommon.h
@@ -184,6 +184,9 @@ namespace RHIValidation
 			// True when a BeginTransition has been issued, and false when the transition has been ended.
 			bool bTransitioning = false;
 
+			// True when a transition with EResourceTransitionFlags::IgnoreAfterState happened, another transition with EResourceTransitionFlags::IgnoreAfterState needs to happen before a regular one
+			bool bIgnoringAfterState = false;
+
 			// True when the resource has been used within a Begin/EndUAVOverlap region.
 			bool bUsedWithAllUAVsOverlap = false;
 
