Make sure the shader binding layout is kept alive for async RTPSO compile jobs

diff --git a/Engine/Source/Runtime/RHI/Private/PipelineStateCache.cpp b/Engine/Source/Runtime/RHI/Private/PipelineStateCache.cpp
--- a/Engine/Source/Runtime/RHI/Private/PipelineStateCache.cpp
+++ b/Engine/Source/Runtime/RHI/Private/PipelineStateCache.cpp
@@ -3528,6 +3528,13 @@ public:
 		Initializer.SetMissShaderTable(MissTable, InInitializer.GetRayMissHash());
 		Initializer.SetHitGroupTable(HitGroupTable, InInitializer.GetHitGroupHash());
 		Initializer.SetCallableTable(CallableTable, InInitializer.GetCallableHash());
+
+		// Also copy over the shader binding layout and update the reference to make sure the binding layout is kept alive
+		if (InInitializer.ShaderBindingLayout)
+		{
+			ShaderBindingLayout = *InInitializer.ShaderBindingLayout;
+			Initializer.ShaderBindingLayout = &ShaderBindingLayout;
+		}		
 	}
 
 	static ESubsequentsMode::Type GetSubsequentsMode() { return ESubsequentsMode::TrackSubsequents; }
@@ -3590,6 +3597,8 @@ private:
 	TArray<FRHIRayTracingShader*> MissTable;
 	TArray<FRHIRayTracingShader*> HitGroupTable;
 	TArray<FRHIRayTracingShader*> CallableTable;
+
+	FRHIShaderBindingLayout ShaderBindingLayout;
 };
 #endif // RHI_RAYTRACING
 
