Workaround for D3D12 not being able to transition from All -> AsyncCompute with ERHIAccess::Discard.
- Discard operations involving All -> AsyncCompute will now occur in the forking graphics pass of the async compute pass that is doing the acquire operation.
- Refactored the transition logic to allow for an immediate epilogue transition for discard operations perfomed in the same pass.
- D3D12 opts into the workaround while other RHI's retain the old behavior.

diff --git a/Engine/Source/Runtime/D3D12RHI/Private/D3D12RHI.cpp b/Engine/Source/Runtime/D3D12RHI/Private/D3D12RHI.cpp
--- a/Engine/Source/Runtime/D3D12RHI/Private/D3D12RHI.cpp
+++ b/Engine/Source/Runtime/D3D12RHI/Private/D3D12RHI.cpp
@@ -254,6 +254,7 @@ FD3D12DynamicRHI::FD3D12DynamicRHI(const TArray<TSharedPtr<FD3D12Adapter>>& Chos
 
 	GRHIGlobals.NeedExtraTransitions = true;
 	GRHIGlobals.EnableHighLevelTransitionPatching = true;
+	GRHIGlobals.NeedsTransientDiscardOnGraphicsWorkaround = true;
 }
 
 void FD3D12DynamicRHI::PostInit()
diff --git a/Engine/Source/Runtime/RHI/Public/RHIGlobals.h b/Engine/Source/Runtime/RHI/Public/RHIGlobals.h
--- a/Engine/Source/Runtime/RHI/Public/RHIGlobals.h
+++ b/Engine/Source/Runtime/RHI/Public/RHIGlobals.h
@@ -717,6 +717,9 @@ struct FRHIGlobals
 	/** True if RHI needs the transitions from the high level for COPYSRC/COPYDEST*/
 	bool NeedExtraTransitions = false;
 
+	/** True if RHI is unable to service transitioning to ERHIAccess::Discard from all pipes to async compute. */
+	bool NeedsTransientDiscardOnGraphicsWorkaround = false;
+
 	/** True if RHI always needs the state of both planes (depth|stencil) for transitions */
 	bool NeedsAllPlanes = false;
 
diff --git a/Engine/Source/Runtime/RHI/Public/RHITransientResourceAllocator.h b/Engine/Source/Runtime/RHI/Public/RHITransientResourceAllocator.h
--- a/Engine/Source/Runtime/RHI/Public/RHITransientResourceAllocator.h
+++ b/Engine/Source/Runtime/RHI/Public/RHITransientResourceAllocator.h
@@ -11,6 +11,106 @@
 class FRHITransientHeap;
 class FRHITransientPagePool;
 
+/** This data structure contains fence values used for allocating / deallocating transient memory regions for transient resources. A memory region can
+ *  be re-used if the deallocation fences from the discarding resource and the allocation fences for the acquiring resource are not executing simultaneously
+ *  on both the graphics | async compute pipe on the GPU timeline.
+ *
+ *  Allocation events are always on a single pipeline, while deallocation events can happen on multiple pipelines at the same time. Async compute is represented
+ *  using three fence values: one for the async compute pipe, and two for the fork / join points on the graphics pipe. If fences are active on both pipes at the
+ *  same time, the graphics fence must be contained within the async compute fork / join region.
+ */
+class FRHITransientAllocationFences
+{
+public:
+	// Returns the fence at which the Acquire operation can occur for the given pair of resources transitioning from Discard -> Acquire.
+	static uint32 GetAcquireFence(const FRHITransientAllocationFences& Discard, const FRHITransientAllocationFences& Acquire)
+	{
+		check(Acquire.IsSinglePipeline());
+
+		if (Discard.Graphics != Invalid)
+		{
+			// Graphics -> Graphics | AsyncCompute
+			if (Discard.AsyncCompute == Invalid)
+			{
+				return Discard.Graphics;
+			}
+
+			// All -> AsyncCompute
+			if (Acquire.AsyncCompute != Invalid)
+			{
+				// All -> AsyncCompute - The acquire graphics fork pass is used because a fence from Graphics -> AsyncCompute after the discard's graphics pass must be present.
+				return Acquire.GraphicsForkJoin.Min;
+			}
+		}
+
+		// AsyncCompute -> AsyncCompute
+		if (Acquire.AsyncCompute != Invalid)
+		{
+			return Discard.AsyncCompute;
+		}
+
+		// All | AsyncCompute -> Graphics - The discard graphics fork pass is used because because a fence from AsyncCompute -> Graphics after the discard's async compute pass must be present.
+		return Discard.GraphicsForkJoin.Max;
+	}
+
+	// Returns whether two regions described by the discard and acquire fences contain each other. If they do, that means the memory would being used by both pipes simultaneously and cannot be aliased.
+	static bool Contains(const FRHITransientAllocationFences& Discard, const FRHITransientAllocationFences& Acquire)
+	{
+		return Contains(Discard.GraphicsForkJoin, Acquire.Graphics) || Contains(Acquire.GraphicsForkJoin, Discard.Graphics);
+	}
+
+	FRHITransientAllocationFences() = default;
+
+	// Expects the bitmask of the pipes this transient allocation was accessed on, which can be different from the fences themselves.
+	FRHITransientAllocationFences(ERHIPipeline InPipelines)
+		: Pipelines(InPipelines)
+	{}
+
+	void SetGraphics(uint32 InGraphics)
+	{
+		check(EnumHasAnyFlags(Pipelines, ERHIPipeline::Graphics));
+		check(!GraphicsForkJoin.IsValid() || Contains(GraphicsForkJoin, InGraphics));
+		Graphics = InGraphics;
+	}
+
+	void SetAsyncCompute(uint32 InAsyncCompute, TInterval<uint32> InGraphicsForkJoin)
+	{
+		check(EnumHasAnyFlags(Pipelines, ERHIPipeline::AsyncCompute));
+		check(InGraphicsForkJoin.IsValid() && Contains(InGraphicsForkJoin, InAsyncCompute));
+		check(Graphics == Invalid || Contains(InGraphicsForkJoin, Graphics));
+		AsyncCompute = InAsyncCompute;
+		GraphicsForkJoin = InGraphicsForkJoin;
+	}
+
+	ERHIPipeline GetPipelines() const
+	{
+		return Pipelines;
+	}
+
+	uint32 GetSinglePipeline() const
+	{
+		check(IsSinglePipeline());
+		return Graphics != Invalid ? Graphics : AsyncCompute;
+	}
+
+	bool IsSinglePipeline() const
+	{
+		return Pipelines != ERHIPipeline::All;
+	}
+
+private:
+	static bool Contains(TInterval<uint32> Interval, uint32 Element)
+	{
+		return Interval.IsValid() && Element > Interval.Min && Element < Interval.Max;
+	}
+
+	static const uint32 Invalid = std::numeric_limits<uint32>::max();
+	uint32 Graphics = Invalid;
+	uint32 AsyncCompute = Invalid;
+	TInterval<uint32> GraphicsForkJoin;
+	ERHIPipeline Pipelines = ERHIPipeline::None;
+};
+
 struct FRHITransientPageSpan
 {
 	// Offset of the span in the page pool in pages. 
@@ -112,20 +212,26 @@
 	//////////////////////////////////////////////////////////////////////////////////////////////////
 	//! Internal Allocator API
 
-	void Acquire(const TCHAR* InName, uint32 InAcquirePassIndex, uint64 InAllocatorCycle)
+	void Acquire(const TCHAR* InName, uint32 InAcquirePassIndex, ERHIPipeline InAcquirePipeline, uint64 InAllocatorCycle)
 	{
 		Name = InName;
 		AcquirePasses = TInterval<uint32>(0, InAcquirePassIndex);
 		DiscardPass = kInvalidPassIndex;
 		bAcquired = true;
+		AcquirePipeline = InAcquirePipeline;
 		AcquireCycle = InAllocatorCycle;
 		AcquireCount++;
 		AliasingOverlaps.Reset();
 	}
 
-	void Discard()
+	void Discard(const FRHITransientAllocationFences& Fences)
 	{
 		bAcquired = false;
+
+		if (GRHIGlobals.NeedsTransientDiscardOnGraphicsWorkaround && !Fences.IsSinglePipeline())
+		{
+			bDiscardOnGraphicsWorkaround = true;
+		}
 	}
 
 	void AddAliasingOverlap(FRHITransientResource* InBeforeResource, uint32 InAcquirePassIndex)
@@ -138,9 +244,17 @@
 			AliasingOverlaps.Emplace(InBeforeResource->GetRHI(), InBeforeResource->IsTexture() ? FRHITransientAliasingOverlap::EType::Texture : FRHITransientAliasingOverlap::EType::Buffer);
 		}
 
-		InBeforeResource->DiscardPass = FMath::Min(InBeforeResource->DiscardPass, AcquirePasses.Max);
 		AcquirePasses.Min = FMath::Max(AcquirePasses.Min, InAcquirePassIndex);
 
+		if (AcquirePipeline == ERHIPipeline::AsyncCompute && InBeforeResource->bDiscardOnGraphicsWorkaround)
+		{
+			InBeforeResource->DiscardPass = FMath::Min(InBeforeResource->DiscardPass, AcquirePasses.Min);
+		}
+		else
+		{
+			InBeforeResource->DiscardPass = FMath::Min(InBeforeResource->DiscardPass, AcquirePasses.Max);
+		}
+
 		check(AcquirePasses.Min <= AcquirePasses.Max);
 	}
 
@@ -266,7 +380,9 @@
 	// Start -> End split pass index intervals for acquire / discard operations.
 	TInterval<uint32> AcquirePasses = TInterval<uint32>(0, 0);
 	uint32 DiscardPass = 0;
+	ERHIPipeline AcquirePipeline = ERHIPipeline::None;
 	bool bAcquired = false;
+	bool bDiscardOnGraphicsWorkaround = false;
 
 	ERHITransientAllocationType AllocationType;
 	ERHITransientResourceType ResourceType;
@@ -400,93 +516,6 @@
 
 ENUM_CLASS_FLAGS(FRHITransientAllocationStats::EMemoryRangeFlags);
 
-/** This data structure contains fence values used for allocating / deallocating transient memory regions for transient resources. A memory region can
- *  be re-used if the deallocation fences from the discarding resource and the allocation fences for the acquiring resource are not executing simultaneously
- *  on both the graphics | async compute pipe on the GPU timeline.
- *
- *  Allocation events are always on a single pipeline, while deallocation events can happen on multiple pipelines at the same time. Async compute is represented
- *  using three fence values: one for the async compute pipe, and two for the fork / join points on the graphics pipe. If fences are active on both pipes at the
- *  same time, the graphics fence must be contained within the async compute fork / join region.
- */
-class FRHITransientAllocationFences
-{
-public:
-	// Returns the fence at which the Acquire operation can occur for the given pair of resources transitioning from Discard -> Acquire.
-	static uint32 GetAcquireFence(const FRHITransientAllocationFences& Discard, const FRHITransientAllocationFences& Acquire)
-	{
-		check(Acquire.IsSinglePipeline());
-
-		if (Discard.Graphics != Invalid)
-		{
-			// Graphics -> Graphics | AsyncCompute
-			if (Discard.AsyncCompute == Invalid)
-			{
-				return Discard.Graphics;
-			}
-
-			// All -> AsyncCompute
-			if (Acquire.AsyncCompute != Invalid)
-			{
-				// All -> AsyncCompute - The acquire graphics fork pass is used because a fence from Graphics -> AsyncCompute after the discard's graphics pass must be present.
-				return Acquire.GraphicsForkJoin.Min;
-			}
-		}
-
-		// AsyncCompute -> AsyncCompute
-		if (Acquire.AsyncCompute != Invalid)
-		{
-			return Discard.AsyncCompute;
-		}
-
-		// All | AsyncCompute -> Graphics - The discard graphics fork pass is used because because a fence from AsyncCompute -> Graphics after the discard's async compute pass must be present.
-		return Discard.GraphicsForkJoin.Max;
-	}
-
-	// Returns whether two regions described by the discard and acquire fences contain each other. If they do, that means the memory would being used by both pipes simultaneously and cannot be aliased.
-	static bool Contains(const FRHITransientAllocationFences& Discard, const FRHITransientAllocationFences& Acquire)
-	{
-		return Contains(Discard.GraphicsForkJoin, Acquire.Graphics) || Contains(Acquire.GraphicsForkJoin, Discard.Graphics);
-	}
-
-	FRHITransientAllocationFences() = default;
-
-	void SetGraphics(uint32 InGraphics)
-	{
-		check(!GraphicsForkJoin.IsValid() || Contains(GraphicsForkJoin, InGraphics));
-		Graphics = InGraphics;
-	}
-
-	void SetAsyncCompute(uint32 InAsyncCompute, TInterval<uint32> InGraphicsForkJoin)
-	{
-		check(InGraphicsForkJoin.IsValid() && Contains(InGraphicsForkJoin, InAsyncCompute));
-		check(Graphics == Invalid || Contains(InGraphicsForkJoin, Graphics));
-		AsyncCompute = InAsyncCompute;
-		GraphicsForkJoin = InGraphicsForkJoin;
-	}
-
-	uint32 GetSinglePipeline() const
-	{
-		check(IsSinglePipeline());
-		return Graphics == Invalid ? AsyncCompute : Graphics;
-	}
-
-	bool IsSinglePipeline() const
-	{
-		return (Graphics == Invalid) ^ (AsyncCompute == Invalid);
-	}
-
-private:
-	static bool Contains(TInterval<uint32> Interval, uint32 Element)
-	{
-		return Interval.IsValid() && Element > Interval.Min && Element < Interval.Max;
-	}
-
-	static const uint32 Invalid = std::numeric_limits<uint32>::max();
-	uint32 Graphics = Invalid;
-	uint32 AsyncCompute = Invalid;
-	TInterval<uint32> GraphicsForkJoin;
-};
-
 enum class ERHITransientResourceCreateMode
 {
 	// Transient resources are always created inline inside of the Create call.
diff --git a/Engine/Source/Runtime/RHICore/Private/RHICoreTransientResourceAllocator.cpp b/Engine/Source/Runtime/RHICore/Private/RHICoreTransientResourceAllocator.cpp
--- a/Engine/Source/Runtime/RHICore/Private/RHICoreTransientResourceAllocator.cpp
+++ b/Engine/Source/Runtime/RHICore/Private/RHICoreTransientResourceAllocator.cpp
@@ -10,13 +10,6 @@
 #include "RHICoreNvidiaAftermath.h"
 #include "RHICore.h"
 
-static bool GRHITransientAllocatorSeparateAsyncComputeHeap = false;
-static FAutoConsoleVariableRef CVarRHITransientAllocatorSeparateAsyncComputeHeap(
-	TEXT("RHI.TransientAllocator.SeparateAsyncComputeHeap"),
-	GRHITransientAllocatorSeparateAsyncComputeHeap,
-	TEXT("If enabled, async compute only allocations are put in their own separate heap."),
-	ECVF_RenderThreadSafe);
-
 static bool GRHITransientAllocatorParallelResourceCreation = true;
 static FAutoConsoleVariableRef CVarRHITransientAllocatorParallelResourceCreation(
 	TEXT("RHI.TransientAllocator.ParallelResourceCreation"),
@@ -386,7 +379,7 @@ FRHITransientTexture* FRHITransientHeap::CreateTexture(
 	});
 
 	check(Texture);
-	Texture->Acquire(DebugName, Fences.GetSinglePipeline(), CurrentAllocatorCycle);
+	Texture->Acquire(DebugName, Fences.GetSinglePipeline(), Fences.GetPipelines(), CurrentAllocatorCycle);
 	AllocateMemoryInternal(Texture, Allocation);
 	Stats.AllocateTexture(Allocation.Size);
 	return Texture;
@@ -422,7 +415,7 @@ FRHITransientBuffer* FRHITransientHeap::CreateBuffer(
 	});
 
 	check(Buffer);
-	Buffer->Acquire(DebugName, Fences.GetSinglePipeline(), CurrentAllocatorCycle);
+	Buffer->Acquire(DebugName, Fences.GetSinglePipeline(), Fences.GetPipelines(), CurrentAllocatorCycle);
 	AllocateMemoryInternal(Buffer, Allocation);
 	Stats.AllocateBuffer(Allocation.Size);
 	return Buffer;
@@ -449,7 +442,7 @@ void FRHITransientHeap::AllocateMemoryInternal(FRHITransientResource* Resource,
 
 void FRHITransientHeap::DeallocateMemoryInternal(FRHITransientResource* Resource, const FRHITransientAllocationFences& Fences)
 {
-	Resource->Discard();
+	Resource->Discard(Fences);
 	Allocator.Deallocate(Resource, Fences);
 }
 
@@ -1419,7 +1412,7 @@ FRHITransientTexture* FRHITransientResourcePageAllocator::CreateTexture(
 	const float FastPoolPercentageRequested = bFastPool ? CreateInfo.FastVRAMPercentage / 255.f : 0.f;
 
 	check(Texture);
-	Texture->Acquire(DebugName, Fences.GetSinglePipeline(), CurrentCycle);
+	Texture->Acquire(DebugName, Fences.GetSinglePipeline(), Fences.GetPipelines(), CurrentCycle);
 	AllocateMemoryInternal(Texture, DebugName, Fences, bFastPool, FastPoolPercentageRequested);
 	Stats.AllocateTexture(Texture->GetSize());
 	IF_RHICORE_TRANSIENT_ALLOCATOR_DEBUG(ActiveResources.Emplace(Texture));
@@ -1438,7 +1431,7 @@ FRHITransientBuffer* FRHITransientResourcePageAllocator::CreateBuffer(
 	});
 
 	check(Buffer);
-	Buffer->Acquire(DebugName, Fences.GetSinglePipeline(), CurrentCycle);
+	Buffer->Acquire(DebugName, Fences.GetSinglePipeline(), Fences.GetPipelines(), CurrentCycle);
 	AllocateMemoryInternal(Buffer, DebugName, Fences, EnumHasAnyFlags(CreateInfo.Usage, EBufferUsageFlags::FastVRAM), false);
 	Stats.AllocateBuffer(Buffer->GetSize());
 	IF_RHICORE_TRANSIENT_ALLOCATOR_DEBUG(ActiveResources.Emplace(Buffer));
@@ -1484,7 +1477,7 @@ void FRHITransientResourcePageAllocator::AllocateMemoryInternal(FRHITransientRes
 
 void FRHITransientResourcePageAllocator::DeallocateMemoryInternal(FRHITransientResource* Resource, const FRHITransientAllocationFences& Fences)
 {
-	Resource->Discard();
+	Resource->Discard(Fences);
 
 	for (const FRHITransientPagePoolAllocation& Allocation : Resource->GetPageAllocation().PoolAllocations)
 	{
diff --git a/Engine/Source/Runtime/RenderCore/Private/RenderGraphBuilder.cpp b/Engine/Source/Runtime/RenderCore/Private/RenderGraphBuilder.cpp
--- a/Engine/Source/Runtime/RenderCore/Private/RenderGraphBuilder.cpp
+++ b/Engine/Source/Runtime/RenderCore/Private/RenderGraphBuilder.cpp
@@ -4045,6 +4045,7 @@ void FRDGBuilder::AddLastTextureTransition(FRDGTexture* Texture)
 	if (Texture->DiscardPass.IsValid())
 	{
 		SubresourceStateAfter.SetPass(GetPassPipeline(Texture->DiscardPass), Texture->DiscardPass);
+		SubresourceStateAfter.BarrierLocation = ERDGBarrierLocation::Epilogue;
 		SubresourceStateAfter.Access = ERHIAccess::Discard;
 	}
 	else
@@ -4133,8 +4134,8 @@ void FRDGBuilder::AddLastBufferTransition(FRDGBuffer* Buffer)
 	// Texture is using the RHI transient allocator. Transition it back to Discard in the final pass it is used.
 	if (Buffer->DiscardPass.IsValid())
 	{
-		const FRDGPassHandle MaxDiscardPass(FMath::Min<uint32>(Buffer->TransientBuffer->GetDiscardPass(), GetEpiloguePassHandle().GetIndex()));
-		StateAfter->SetPass(GetPassPipeline(MaxDiscardPass), MaxDiscardPass);
+		StateAfter->SetPass(GetPassPipeline(Buffer->DiscardPass), Buffer->DiscardPass);
+		StateAfter->BarrierLocation = ERDGBarrierLocation::Epilogue;
 		StateAfter->Access = ERHIAccess::Discard;
 	}
 	else
@@ -4351,32 +4352,43 @@ void FRDGBuilder::AddTransition(
 	const FRDGPassHandlesByPipeline& PassesBefore = StateBefore.LastPass;
 	const FRDGPassHandlesByPipeline& PassesAfter = StateAfter.FirstPass;
 
-	// 1-to-1 or 1-to-N pipe transition.
-	if (PipelinesBefore != ERHIPipeline::All)
+	// 1-to-1 same-pipe transition
+	if (PipelinesBefore == PipelinesAfter)
 	{
-		const FRDGPassHandle BeginPassHandle = StateBefore.GetLastPass();
-		const FRDGPassHandle FirstEndPassHandle = StateAfter.GetFirstPass();
+		const FRDGPassHandle BeginPassHandle = StateBefore.LastPass[PipelinesAfter];
+		const FRDGPassHandle EndPassHandle   = StateAfter.FirstPass[PipelinesAfter];
 
-		FRDGPass* BeginPass = nullptr;
-		FRDGBarrierBatchBegin* BarriersToBegin = nullptr;
-
-		// Issue the begin in the epilogue of the begin pass if the barrier is being split across multiple passes or the barrier end is in the epilogue.
-		if (BeginPassHandle < FirstEndPassHandle)
+		// Split the transition from the epilogue of the begin pass to the prologue of the end pass.
+		if (BeginPassHandle < EndPassHandle)
+		{
+			FRDGPass* BeginPass = GetEpilogueBarrierPass(BeginPassHandle);
+			FRDGBarrierBatchBegin* BarriersToBegin = &BeginPass->GetEpilogueBarriersToBeginFor(Allocators.Transition, TransitionCreateQueue, PipelinesAfter);
+			BarriersToBegin->AddTransition(Resource, TransitionInfo);
+			AddToPrologueBarriersToEnd(EndPassHandle, *BarriersToBegin);
+		}
+		// This is an immediate transition in the same pass on the same pipe done in the epilogue of the pass.
+		else if (StateAfter.BarrierLocation == ERDGBarrierLocation::Epilogue)
 		{
-			BeginPass = GetEpilogueBarrierPass(BeginPassHandle);
-			BarriersToBegin = &BeginPass->GetEpilogueBarriersToBeginFor(Allocators.Transition, TransitionCreateQueue, PipelinesAfter);
+			FRDGPass* BeginPass = GetEpilogueBarrierPass(BeginPassHandle);
+			FRDGBarrierBatchBegin* BarriersToBegin = &BeginPass->GetEpilogueBarriersToBeginFor(Allocators.Transition, TransitionCreateQueue, PipelinesAfter);
+			BarriersToBegin->AddTransition(Resource, TransitionInfo);
+			AddToEpilogueBarriersToEnd(EndPassHandle, *BarriersToBegin);
 		}
-		// This is an immediate prologue transition in the same pass. Issue the begin in the prologue.
+		// This is an immediate transition in the same pass on the same pipe done in the prologue of the pass.
 		else
 		{
-			checkf(PipelinesAfter == ERHIPipeline::Graphics,
-				TEXT("Attempted to queue an immediate async pipe transition for %s. Pipelines: %s. Async transitions must be split."),
-				Resource->Name, *GetRHIPipelineName(PipelinesAfter));
-
-			BeginPass = GetPrologueBarrierPass(BeginPassHandle);
-			BarriersToBegin = &BeginPass->GetPrologueBarriersToBegin(Allocators.Transition, TransitionCreateQueue);
+			FRDGPass* BeginPass = GetPrologueBarrierPass(BeginPassHandle);
+			FRDGBarrierBatchBegin* BarriersToBegin = &BeginPass->GetPrologueBarriersToBegin(Allocators.Transition, TransitionCreateQueue);
+			BarriersToBegin->AddTransition(Resource, TransitionInfo);
+			AddToPrologueBarriersToEnd(EndPassHandle, *BarriersToBegin);
 		}
-
+	}
+	// 1-to-1 or 1-to-N cross-pipe transition.
+	else if (PipelinesBefore != ERHIPipeline::All)
+	{
+		const FRDGPassHandle BeginPassHandle = StateBefore.LastPass[PipelinesBefore];
+		FRDGPass* BeginPass = GetEpilogueBarrierPass(BeginPassHandle);
+		FRDGBarrierBatchBegin* BarriersToBegin = &BeginPass->GetEpilogueBarriersToBeginFor(Allocators.Transition, TransitionCreateQueue, PipelinesAfter);
 		BarriersToBegin->AddTransition(Resource, TransitionInfo);
 
 		for (ERHIPipeline Pipeline : MakeFlagsRange(ERHIPipeline::All))
@@ -4407,12 +4419,11 @@ void FRDGBuilder::AddTransition(
 			}
 		}
 	}
-	// N-to-1 or N-to-N transition.
+	// N-to-1
 	else
 	{
-		checkf(StateBefore.GetLastPass() != StateAfter.GetFirstPass(),
-			TEXT("Attempted to queue a transition for resource '%s' from '%s' to '%s', but previous and next passes are the same on one pipe."),
-			Resource->Name, *GetRHIPipelineName(PipelinesBefore), *GetRHIPipelineName(PipelinesAfter));
+		check(PipelinesAfter != ERHIPipeline::All);
+		FRDGPassHandle PassAfter = PassesAfter[PipelinesAfter];
 
 		FRDGBarrierBatchBeginId Id;
 		Id.PipelinesAfter = PipelinesAfter;
@@ -4440,9 +4451,15 @@ void FRDGBuilder::AddTransition(
 
 		BarriersToBegin->AddTransition(Resource, TransitionInfo);
 
-		for (ERHIPipeline Pipeline : MakeFlagsRange(PipelinesAfter))
+		// If the end pass is the same as the begin pass on one pipe, end it in the epilogue instead.
+		if (PassesBefore[PipelinesAfter] == PassAfter)
 		{
-			AddToPrologueBarriersToEnd(PassesAfter[Pipeline], *BarriersToBegin);
+			check(StateAfter.BarrierLocation == ERDGBarrierLocation::Epilogue);
+			AddToEpilogueBarriersToEnd(PassAfter, *BarriersToBegin);
+		}
+		else
+		{
+			AddToPrologueBarriersToEnd(PassAfter, *BarriersToBegin);
 		}
 	}
 }
@@ -4479,11 +4496,13 @@ void FRDGBuilder::AddAliasingTransition(FRDGPassHandle BeginPassHandle, FRDGPass
 
 FRHITransientAllocationFences FRDGBuilder::GetAllocateFences(FRDGViewableResource* Resource) const
 {
-	FRHITransientAllocationFences Fences;
-
 	FRDGPassHandle FirstPassHandle = Resource->FirstPass;
 
-	if (GetPassPipeline(FirstPassHandle) == ERHIPipeline::Graphics)
+	ERHIPipeline Pipeline = GetPassPipeline(FirstPassHandle);
+
+	FRHITransientAllocationFences Fences(Pipeline);
+
+	if (Pipeline == ERHIPipeline::Graphics)
 	{
 		Fences.SetGraphics(FirstPassHandle.GetIndex());
 	}
@@ -4507,8 +4526,12 @@ FRHITransientAllocationFences FRDGBuilder::GetDeallocateFences(FRDGViewableResou
 	FRDGPassHandle GraphicsForkPass;
 	FRDGPassHandle GraphicsJoinPass;
 
+	ERHIPipeline Pipelines = GraphicsPassHandle.IsValid() ? ERHIPipeline::Graphics : ERHIPipeline::None;
+
 	if (AsyncComputePassHandle.IsValid())
 	{
+		Pipelines |= ERHIPipeline::AsyncCompute;
+
 		const FRDGPass* Pass = Passes[AsyncComputePassHandle];
 		GraphicsForkPass = Pass->GraphicsForkPass;
 		GraphicsJoinPass = Pass->GraphicsJoinPass;
@@ -4528,7 +4551,7 @@ FRHITransientAllocationFences FRDGBuilder::GetDeallocateFences(FRDGViewableResou
 		}
 	}
 
-	FRHITransientAllocationFences Fences;
+	FRHITransientAllocationFences Fences(Pipelines);
 
 	if (GraphicsPassHandle.IsValid())
 	{
diff --git a/Engine/Source/Runtime/RenderCore/Private/RenderGraphResourcePool.cpp b/Engine/Source/Runtime/RenderCore/Private/RenderGraphResourcePool.cpp
--- a/Engine/Source/Runtime/RenderCore/Private/RenderGraphResourcePool.cpp
+++ b/Engine/Source/Runtime/RenderCore/Private/RenderGraphResourcePool.cpp
@@ -378,7 +378,7 @@ void FRDGTransientResourceAllocator::ReleasePendingDeallocations()
 
 		for (FRDGTransientRenderTarget* RenderTarget : PendingDeallocationList)
 		{
-			FRHITransientAllocationFences Fences;
+			FRHITransientAllocationFences Fences(ERHIPipeline::Graphics);
 			Fences.SetGraphics(0);
 			Allocator->DeallocateMemory(RenderTarget->Texture, Fences);
 			Transitions.Emplace(RenderTarget->Texture->GetRHI(), ERHIAccess::Unknown, ERHIAccess::Discard);
diff --git a/Engine/Source/Runtime/RenderCore/Public/RenderGraphPass.h b/Engine/Source/Runtime/RenderCore/Public/RenderGraphPass.h
--- a/Engine/Source/Runtime/RenderCore/Public/RenderGraphPass.h
+++ b/Engine/Source/Runtime/RenderCore/Public/RenderGraphPass.h
@@ -54,20 +54,6 @@ struct FRDGBarrierBatchBeginId
 	ERHIPipeline PipelinesAfter = ERHIPipeline::None;
 };
 
-/** Barrier location controls where the barrier is 'Ended' relative to the pass lambda being executed.
- *  Most barrier locations are done in the prologue prior to the executing lambda. But certain cases
- *  like an aliasing discard operation need to be done *after* the pass being invoked. Therefore, when
- *  adding a transition the user can specify where to place the barrier.
- */
-enum class ERDGBarrierLocation : uint8
-{
-	/** The barrier occurs in the prologue of the pass (before execution). */
-	Prologue,
-
-	/** The barrier occurs in the epilogue of the pass (after execution). */
-	Epilogue
-};
-
 struct FRDGTransitionInfo
 {
 	static_assert((int32)ERHIAccess::Last <= (1 << 20) && (int32)ERDGViewableResourceType::MAX <= 3 && (int32)EResourceTransitionFlags::Last == (1 << 2), "FRDGTransitionInfo packing is no longer correct.");
diff --git a/Engine/Source/Runtime/RenderCore/Public/RenderGraphResources.h b/Engine/Source/Runtime/RenderCore/Public/RenderGraphResources.h
--- a/Engine/Source/Runtime/RenderCore/Public/RenderGraphResources.h
+++ b/Engine/Source/Runtime/RenderCore/Public/RenderGraphResources.h
@@ -48,6 +48,20 @@ struct FRDGProducerState
 	FRDGViewHandle NoUAVBarrierHandle;
 };
 
+/** Barrier location controls where the barrier is 'Ended' relative to the pass lambda being executed.
+ *  Most barrier locations are done in the prologue prior to the executing lambda. But certain cases
+ *  like an aliasing discard operation need to be done *after* the pass being invoked. Therefore, when
+ *  adding a transition the user can specify where to place the barrier.
+ */
+enum class ERDGBarrierLocation : uint8
+{
+	/** The barrier occurs in the prologue of the pass (before execution). */
+	Prologue,
+
+	/** The barrier occurs in the epilogue of the pass (after execution). */
+	Epilogue
+};
+
 using FRDGProducerStatesByPipeline = TRHIPipelineArray<FRDGProducerState>;
 
 /** Used for tracking the state of an individual subresource during execution. */
@@ -105,6 +119,9 @@ struct FRDGSubresourceState
 
 	/** The last used transition flags on the pass. */
 	EResourceTransitionFlags Flags = EResourceTransitionFlags::None;
+
+	/** Used to specify whether the state is applied during the prologue or epilogue of the pass. This is only used when transitioning on the same pass / pipe. */
+	ERDGBarrierLocation BarrierLocation = ERDGBarrierLocation::Prologue;
 };
 
 using FRDGTextureSubresourceState = TRDGTextureSubresourceArray<FRDGSubresourceState*, FRDGArrayAllocator>;
