Added FRHICommandListScopedFence object to use a single fence when issuing multiple commands.

diff --git a/Engine/Source/Runtime/Engine/Private/GPUSkinCache.cpp b/Engine/Source/Runtime/Engine/Private/GPUSkinCache.cpp
--- a/Engine/Source/Runtime/Engine/Private/GPUSkinCache.cpp
+++ b/Engine/Source/Runtime/Engine/Private/GPUSkinCache.cpp
@@ -1606,10 +1606,14 @@ struct FGPUSkinCache::FTaskData
 	void Begin()
 	{
 		OriginalPipeline = RHICmdList->SwitchPipeline(Pipeline);
+
+		RHICmdListScopedFence.Emplace(*RHICmdList);
 	}
 
 	void End()
 	{
+		RHICmdListScopedFence.Reset();
+
 		if (bAsyncCommandList)
 		{
 			RHICmdList->FinishRecording();
@@ -1622,6 +1626,7 @@ struct FGPUSkinCache::FTaskData
 
 	TArray<FDispatchEntry> TangentDispatches;
 	TArray<FSortedDispatchEntry> SortedDispatches;
+	TOptional<FRHICommandListScopedFence> RHICmdListScopedFence;
 	FRHICommandList* RHICmdList = nullptr;
 
 	struct
diff --git a/Engine/Source/Runtime/Engine/Private/Materials/MaterialRenderProxy.cpp b/Engine/Source/Runtime/Engine/Private/Materials/MaterialRenderProxy.cpp
--- a/Engine/Source/Runtime/Engine/Private/Materials/MaterialRenderProxy.cpp
+++ b/Engine/Source/Runtime/Engine/Private/Materials/MaterialRenderProxy.cpp
@@ -128,17 +128,21 @@ public:
 			{
 				TRACE_CPUPROFILER_EVENT_SCOPE(FUniformExpressionCacheAsyncUpdater::Update);
 				FTaskTagScope Scope(ETaskTag::EParallelRenderingThread);
-				FMemMark Mark(FMemStack::Get());
 
-				for (const FItem& Item : Items)
 				{
-					uint8* TempBuffer = FMemStack::Get().PushBytes(Item.UniformBufferLayout->ConstantBufferSize, SHADER_PARAMETER_STRUCT_ALIGNMENT);
+					FRHICommandListScopedFence ScopedFence(RHICmdList);
+					FMemMark Mark(FMemStack::Get());
 
-					FMaterialRenderContext Context(Item.MaterialRenderProxy, *Item.Material, Item.bShowSelection);
+					for (const FItem& Item : Items)
+					{
+						uint8* TempBuffer = FMemStack::Get().PushBytes(Item.UniformBufferLayout->ConstantBufferSize, SHADER_PARAMETER_STRUCT_ALIGNMENT);
+
+						FMaterialRenderContext Context(Item.MaterialRenderProxy, *Item.Material, Item.bShowSelection);
 
-					Item.UniformExpressionSet->FillUniformBuffer(Context, Item.AllocatedVTs, Item.UniformBufferLayout, TempBuffer, Item.UniformBufferLayout->ConstantBufferSize);
+						Item.UniformExpressionSet->FillUniformBuffer(Context, Item.AllocatedVTs, Item.UniformBufferLayout, TempBuffer, Item.UniformBufferLayout->ConstantBufferSize);
 
-					RHICmdList.UpdateUniformBuffer(Item.UniformBuffer, TempBuffer);
+						RHICmdList.UpdateUniformBuffer(Item.UniformBuffer, TempBuffer);
+					}
 				}
 
 				RHICmdList.FinishRecording();
diff --git a/Engine/Source/Runtime/Engine/Private/SkeletalMeshUpdater.cpp b/Engine/Source/Runtime/Engine/Private/SkeletalMeshUpdater.cpp
--- a/Engine/Source/Runtime/Engine/Private/SkeletalMeshUpdater.cpp
+++ b/Engine/Source/Runtime/Engine/Private/SkeletalMeshUpdater.cpp
@@ -288,10 +288,14 @@ struct FSkeletalMeshUpdater::FTaskData
 		{
 			RHICmdList->SwitchPipeline(ERHIPipeline::Graphics);
 		}
+
+		RHICmdListScopedFence.Emplace(*RHICmdList);
 	}
 
 	void End()
 	{
+		RHICmdListScopedFence.Reset();
+
 		if (bAsyncCommandList)
 		{
 			RHICmdList->FinishRecording();
@@ -342,6 +346,7 @@ struct FSkeletalMeshUpdater::FTaskData
 	}
 
 	TArray<TUniquePtr<FSkeletalMeshUpdatePacket>, FConcurrentLinearArrayAllocator> Packets;
+	TOptional<FRHICommandListScopedFence> RHICmdListScopedFence;
 	FRHICommandList* RHICmdList;
 	const ERHIPipeline GPUSkinCachePipeline;
 	const bool bAsyncCommandList;
diff --git a/Engine/Source/Runtime/RHI/Private/RHICommandList.cpp b/Engine/Source/Runtime/RHI/Private/RHICommandList.cpp
--- a/Engine/Source/Runtime/RHI/Private/RHICommandList.cpp
+++ b/Engine/Source/Runtime/RHI/Private/RHICommandList.cpp
@@ -189,6 +189,14 @@ void FRHICommandListBase::AddDispatchPrerequisite(const FGraphEventRef& Prereq)
 
 void FRHICommandListBase::FinishRecording()
 {
+	checkf(IsImmediate() || PersistentState.CurrentFenceScope == nullptr, TEXT("Finished recording with an open RHI fence scope."));
+
+	if (PersistentState.CurrentFenceScope && PersistentState.CurrentFenceScope->bFenceRequested)
+	{
+		PersistentState.CurrentFenceScope->bFenceRequested = false;
+		RHIThreadFence(true);
+	}
+
 	if (PendingBufferUploads.Num() > 0)
 	{
 		FStringBuilderBase BufferList;
@@ -1408,6 +1416,12 @@ FGraphEventRef FRHICommandListBase::RHIThreadFence(bool bSetLockFence)
 {
 	checkf(IsTopOfPipe() || Bypass(), TEXT("RHI thread fences only work when recording RHI commands (or in bypass mode)."));
 
+	if (PersistentState.CurrentFenceScope && bSetLockFence)
+	{
+		PersistentState.CurrentFenceScope->bFenceRequested = true;
+		return nullptr;
+	}
+
 	bUsesLockFence |= bSetLockFence;
 
 	if (IsRunningRHIInSeparateThread())
diff --git a/Engine/Source/Runtime/RHI/Public/RHICommandList.h b/Engine/Source/Runtime/RHI/Public/RHICommandList.h
--- a/Engine/Source/Runtime/RHI/Public/RHICommandList.h
+++ b/Engine/Source/Runtime/RHI/Public/RHICommandList.h
@@ -439,6 +439,20 @@ struct TRHILambdaCommandMultiPipe final : public FRHICommandBase
 // This controls if the cmd list bypass can be toggled at runtime. It is quite expensive to have these branches in there.
 #define CAN_TOGGLE_COMMAND_LIST_BYPASS (!UE_BUILD_SHIPPING && !UE_BUILD_TEST)
 
+// Issues a single fence at the end of the scope if an RHI fence is requested by commands within the scope. Can
+// reduce overhead of RHIThreadFence when batch updating resources that would otherwise issue N fences.
+class FRHICommandListScopedFence
+{
+	friend class FRHICommandListBase;
+	FRHICommandListBase& RHICmdList;
+    FRHICommandListScopedFence* Previous;
+    bool bFenceRequested = false;
+
+public:
+	FRHICommandListScopedFence(FRHICommandListBase& RHICmdList);
+    ~FRHICommandListScopedFence();
+};
+
 class FRHICommandListScopedPipelineGuard
 {
 	FRHICommandListBase& RHICmdList;
@@ -1454,6 +1468,7 @@ protected:
 		FBoundShaderStateInput BoundShaderInput;
 		FRHIComputeShader* BoundComputeShaderRHI = nullptr;
 		FRHIWorkGraphShader* BoundWorkGraphShaderRHI = nullptr;
+		FRHICommandListScopedFence* CurrentFenceScope = nullptr;
 
 #if WITH_RHI_BREADCRUMBS
 		FRHIBreadcrumbNode* LocalBreadcrumb = FRHIBreadcrumbNode::Sentinel;
@@ -1540,6 +1555,7 @@ private:
 	friend class FRHICommandListExecutor;
 	friend class FRHICommandListIterator;
 	friend class FRHICommandListScopedFlushAndExecute;
+	friend class FRHICommandListScopedFence;
 	friend class FRHIComputeCommandList;
 	friend class FRHICommandListImmediate;
 	friend class FRHICommandList_RecursiveHazardous;
diff --git a/Engine/Source/Runtime/RHI/Public/RHICommandList.inl b/Engine/Source/Runtime/RHI/Public/RHICommandList.inl
--- a/Engine/Source/Runtime/RHI/Public/RHICommandList.inl
+++ b/Engine/Source/Runtime/RHI/Public/RHICommandList.inl
@@ -111,6 +111,24 @@ private:
 #endif
 };
 
+inline FRHICommandListScopedFence::FRHICommandListScopedFence(FRHICommandListBase& RHICmdList)
+    : RHICmdList(RHICmdList)
+    , Previous(RHICmdList.PersistentState.CurrentFenceScope)
+{
+    RHICmdList.PersistentState.CurrentFenceScope = this;
+}
+
+inline FRHICommandListScopedFence::~FRHICommandListScopedFence()
+{
+    if (bFenceRequested)
+    {
+        RHICmdList.PersistentState.CurrentFenceScope = nullptr;
+        RHICmdList.RHIThreadFence(true);
+    }
+
+    RHICmdList.PersistentState.CurrentFenceScope = Previous;
+}
+
 inline FRHICommandListScopedPipelineGuard::FRHICommandListScopedPipelineGuard(FRHICommandListBase& RHICmdList)
 	: RHICmdList(RHICmdList)
 {
diff --git a/Engine/Source/Runtime/RenderCore/Private/RenderGraphBuilder.cpp b/Engine/Source/Runtime/RenderCore/Private/RenderGraphBuilder.cpp
--- a/Engine/Source/Runtime/RenderCore/Private/RenderGraphBuilder.cpp
+++ b/Engine/Source/Runtime/RenderCore/Private/RenderGraphBuilder.cpp
@@ -2657,6 +2657,7 @@ void FRDGBuilder::CompilePassOps(FRDGPass* Pass)
 void FRDGBuilder::SubmitBufferUploads(FRHICommandListBase& RHICmdListUpload, UE::Tasks::FTaskEvent* AllocateUploadBuffersTask)
 {
 	SCOPED_NAMED_EVENT_TEXT("FRDGBuilder::SubmitBufferUploads", FColor::Magenta);
+	FRHICommandListScopedFence ScopedFence(RHICmdListUpload);
 
 	{
 		SCOPED_NAMED_EVENT_TEXT("Allocate", FColor::Magenta);
diff --git a/Engine/Source/Runtime/Renderer/Private/RendererScene.cpp b/Engine/Source/Runtime/Renderer/Private/RendererScene.cpp
--- a/Engine/Source/Runtime/Renderer/Private/RendererScene.cpp
+++ b/Engine/Source/Runtime/Renderer/Private/RendererScene.cpp
@@ -6851,6 +6851,7 @@ void FScene::Update(FRDGBuilder& GraphBuilder, const FUpdateParameters& Paramete
 		GraphBuilder.AddCommandListSetupTask([this, ProxiesToUpdate = MoveTemp(ProxiesToUpdate)](FRHICommandList& RHICmdList)
 		{
 			SCOPED_NAMED_EVENT(AsyncUpdateUniformBuffers, FColor::Emerald);
+			FRHICommandListScopedFence ScopedFence(RHICmdList);
 
 			for (FPrimitiveSceneProxy* Proxy : ProxiesToUpdate)
 			{
