New RHI breadcrumb metadata API
- Refactored RHI breadcrumb system to allow tagging of vararg values with RHI_BREADCRUMB_FIELD. This provides a name to individual varargs, which will be exposed to Unreal Insights as additional metadata on the breadcrumb markers.
- Add "_F" versions of the breadcrumb macros, and deprecated the use of varargs in the non-"_F" versions. Unreal Insights requires a stable, human readable name per event, so the user is now requried to provide this alongside the more detailed format string. E.g. the frame marker has a fixed "Frame" name, and a "Frame %d" format string.
- Breadcrumbs now use a static descriptor object, which contains the static name, format string, and array of field names. This will also allow us to register events with Insights once on first-use, and then emit only the captured values and event ID for each subsequent use.

Modified RenderGraph's breadcrumb support to avoid the allocation of an FRDGEventName if the varargs of an RDG_EVENT_SCOPE are directly compatible with RHI breadcrumbs
- If they are compatible, the varargs are forwarded as-is, otherwise an RDG specific TValue<> specialization performs the snprintf like before.
- This also gives us an ability to deprecate the non-compatible use cases in future.

Modify CSV profiler support in the new GPU profiler to emit a single "GPU" category
- Values emitted are the busy time for the graphics pipe, plus time the graphics pipe spent waiting on other queues. This is equivalent to the stats we had with the old profiler.

Simplified RHI breadcrumb log/check macros
- These just forward to the standard check/UE_LOG macros, with an appended breadcrumb string.
- When WITH_RHI_BREADCRUMBS is zero, fall back to standard checks/logs.

Fixup deprecated use of RHI_BREADCRUMB_EVENT macros with varargs
- Replaced with _F versions that require both a static name and a format string.

diff --git a/Engine/Plugins/Experimental/DynamicMaterial/Source/DynamicMaterialEditor/Private/Utils/DMMaterialSnapshotLibrary.cpp b/Engine/Plugins/Experimental/DynamicMaterial/Source/DynamicMaterialEditor/Private/Utils/DMMaterialSnapshotLibrary.cpp
--- a/Engine/Plugins/Experimental/DynamicMaterial/Source/DynamicMaterialEditor/Private/Utils/DMMaterialSnapshotLibrary.cpp
+++ b/Engine/Plugins/Experimental/DynamicMaterial/Source/DynamicMaterialEditor/Private/Utils/DMMaterialSnapshotLibrary.cpp
@@ -39,7 +39,7 @@ void FDMMaterialShapshotLibrary::RenderMaterialToRenderTarget(UMaterialInterface
 	Canvas->Init(InRenderTarget->SizeX, InRenderTarget->SizeY, nullptr, &RenderCanvas);
 
 	{
-		RHI_BREADCRUMB_EVENT_GAMETHREAD("DrawMaterialToRenderTarget: %s", InRenderTarget->GetFName());
+		RHI_BREADCRUMB_EVENT_GAMETHREAD_F("DrawMaterialToRenderTarget", "DrawMaterialToRenderTarget: %s", InRenderTarget->GetFName());
 
 		ENQUEUE_RENDER_COMMAND(FlushDeferredResourceUpdateCommand)(
 			[RenderTargetResource](FRHICommandListImmediate& RHICmdList)
diff --git a/Engine/Plugins/Media/ImgMedia/Source/ImgMedia/Private/Readers/ExrImgMediaReaderGpu.cpp b/Engine/Plugins/Media/ImgMedia/Source/ImgMedia/Private/Readers/ExrImgMediaReaderGpu.cpp
--- a/Engine/Plugins/Media/ImgMedia/Source/ImgMedia/Private/Readers/ExrImgMediaReaderGpu.cpp
+++ b/Engine/Plugins/Media/ImgMedia/Source/ImgMedia/Private/Readers/ExrImgMediaReaderGpu.cpp
@@ -180,7 +180,7 @@ FExrImgMediaReader::EReadResult FExrImgMediaReaderGpu::ReadMip
 		{
 			ENQUEUE_RENDER_COMMAND(CopyFromUploadBuffer)([SampleConverter, BufferData, BufferRegionsToCopy, FrameId = ConverterParams.FrameId](FRHICommandListImmediate& RHICmdList)
 				{
-					RHI_BREADCRUMB_EVENT_STAT(RHICmdList, ExrImgMediaReaderGpu_CopyUploadBuffer, "ExrReaderGpu.StartCopy %d", FrameId);
+					RHI_BREADCRUMB_EVENT_STAT_F(RHICmdList, ExrImgMediaReaderGpu_CopyUploadBuffer, "ExrReaderGpu.StartCopy", "ExrReaderGpu.StartCopy %d", FrameId);
 					SCOPED_GPU_STAT(RHICmdList, ExrImgMediaReaderGpu_CopyUploadBuffer);
 
 					if (BufferRegionsToCopy.IsEmpty())
@@ -468,7 +468,7 @@ void FExrImgMediaReaderGpu::CreateSampleConverterCallback(TSharedPtr<FExrMediaTe
 {
 	auto RenderThreadSwizzler = [] (FRHICommandListImmediate& RHICmdList, FTextureRHIRef RenderTargetTextureRHI, TMap<int32, FStructuredBufferPoolItemSharedPtr>& MipBuffers, const FSampleConverterParameters ConverterParams)->bool
 	{
-		RHI_BREADCRUMB_EVENT_STAT(RHICmdList, ExrImgMediaReaderGpu, "ExrReaderGpu.Convert %d", ConverterParams.FrameId);
+		RHI_BREADCRUMB_EVENT_STAT_F(RHICmdList, ExrImgMediaReaderGpu, "ExrReaderGpu.Convert", "ExrReaderGpu.Convert %d", ConverterParams.FrameId);
 		SCOPED_GPU_STAT(RHICmdList, ExrImgMediaReaderGpu);
 
 		auto RenderMip = []
diff --git a/Engine/Source/Runtime/Apple/MetalRHI/Private/MetalRHI.cpp b/Engine/Source/Runtime/Apple/MetalRHI/Private/MetalRHI.cpp
--- a/Engine/Source/Runtime/Apple/MetalRHI/Private/MetalRHI.cpp
+++ b/Engine/Source/Runtime/Apple/MetalRHI/Private/MetalRHI.cpp
@@ -1341,7 +1341,7 @@ void FMetalRHICommandContext::RHIBeginBreadcrumbGPU(FRHIBreadcrumbNode* Breadcru
 	{
 		if (!NameStr)
 		{
-			NameStr = Breadcrumb->Name.GetTCHAR(Buffer);
+			NameStr = Breadcrumb->GetTCHAR(Buffer);
 		}
 		return NameStr;
 	};
diff --git a/Engine/Source/Runtime/D3D12RHI/Private/D3D12Buffer.cpp b/Engine/Source/Runtime/D3D12RHI/Private/D3D12Buffer.cpp
--- a/Engine/Source/Runtime/D3D12RHI/Private/D3D12Buffer.cpp
+++ b/Engine/Source/Runtime/D3D12RHI/Private/D3D12Buffer.cpp
@@ -91,8 +91,8 @@ void FD3D12CommandContext::CopyBufferRegionChecked(
 )
 {
 #if ENABLE_COPY_BUFFER_REGION_CHECK
-	RHI_BREADCRUMB_CHECK_SHIPPINGF(this, DestOffset + ByteCount <= DestResource->GetDesc().Width, TEXT("Dest byte range out of bounds for: '%s'"), *DestName.ToString());
-	RHI_BREADCRUMB_CHECK_SHIPPINGF(this, SourceOffset + ByteCount <= SourceResource->GetDesc().Width, TEXT("Source byte range out of bounds for: '%s'"), *SourceName.ToString());
+	RHI_BREADCRUMB_CHECK_SHIPPINGF(*this, DestOffset + ByteCount <= DestResource->GetDesc().Width, TEXT("Dest byte range out of bounds for: '%s'"), *DestName.ToString());
+	RHI_BREADCRUMB_CHECK_SHIPPINGF(*this, SourceOffset + ByteCount <= SourceResource->GetDesc().Width, TEXT("Source byte range out of bounds for: '%s'"), *SourceName.ToString());
 #endif // ENABLE_COPY_BUFFER_REGION_CHECK
 	
 	// Pass down callchain
diff --git a/Engine/Source/Runtime/D3D12RHI/Private/D3D12CommandContext.cpp b/Engine/Source/Runtime/D3D12RHI/Private/D3D12CommandContext.cpp
--- a/Engine/Source/Runtime/D3D12RHI/Private/D3D12CommandContext.cpp
+++ b/Engine/Source/Runtime/D3D12RHI/Private/D3D12CommandContext.cpp
@@ -162,7 +162,7 @@ void FD3D12ContextCommon::BindDiagnosticBuffer(FD3D12RootSignature const* RootSi
 		if (DiagBuffer && UE::RHI::UseGPUCrashBreadcrumbs())
 		{
 			D3D12_GPU_VIRTUAL_ADDRESS Marker = DiagBuffer->GetGPUQueueMarkerIn();
-			WriteMarker(Marker, Breadcrumb->Name.ID, EMarkerType::In);
+			WriteMarker(Marker, Breadcrumb->ID, EMarkerType::In);
 		}
 
 	#if NV_AFTERMATH
@@ -178,7 +178,7 @@ void FD3D12ContextCommon::BindDiagnosticBuffer(FD3D12RootSignature const* RootSi
 		{
 			if (!NameStr)
 			{
-				NameStr = Breadcrumb->Name.GetTCHAR(Buffer);
+				NameStr = Breadcrumb->GetTCHAR(Buffer);
 			}
 			return NameStr;
 		};
@@ -277,7 +277,7 @@ void FD3D12ContextCommon::BindDiagnosticBuffer(FD3D12RootSignature const* RootSi
 		if (DiagBuffer && UE::RHI::UseGPUCrashBreadcrumbs())
 		{
 			D3D12_GPU_VIRTUAL_ADDRESS Marker = DiagBuffer->GetGPUQueueMarkerOut();
-			WriteMarker(Marker, Breadcrumb->Name.ID, EMarkerType::Out);
+			WriteMarker(Marker, Breadcrumb->ID, EMarkerType::Out);
 		}
 	}
 #endif // WITH_RHI_BREADCRUMBS
diff --git a/Engine/Source/Runtime/D3D12RHI/Private/D3D12CommandContext.h b/Engine/Source/Runtime/D3D12RHI/Private/D3D12CommandContext.h
--- a/Engine/Source/Runtime/D3D12RHI/Private/D3D12CommandContext.h
+++ b/Engine/Source/Runtime/D3D12RHI/Private/D3D12CommandContext.h
@@ -481,18 +481,6 @@ protected:
 	FRHIGPUMask PhysicalGPUMask;
 };
 
-#if WITH_RHI_BREADCRUMBS
-inline FRHIBreadcrumbNode* GetTopBreadcrumbNode(FD3D12CommandContextBase& CommandContext)
-{
-	return CommandContext.GetExecutingCommandList().GetCurrentBreadcrumbRef();
-}
-
-inline FRHIBreadcrumbNode* GetTopBreadcrumbNode(FD3D12CommandContextBase* CommandContext)
-{
-	return CommandContext->GetExecutingCommandList().GetCurrentBreadcrumbRef();
-}
-#endif // WITH_RHI_BREADCRUMBS
-
 // RHI Context type used for graphics and async compute command lists.
 class FD3D12CommandContext : public FD3D12ContextCommon, public FD3D12CommandContextBase, public FD3D12DeviceChild
 {
diff --git a/Engine/Source/Runtime/D3D12RHI/Private/D3D12Texture.cpp b/Engine/Source/Runtime/D3D12RHI/Private/D3D12Texture.cpp
--- a/Engine/Source/Runtime/D3D12RHI/Private/D3D12Texture.cpp
+++ b/Engine/Source/Runtime/D3D12RHI/Private/D3D12Texture.cpp
@@ -1174,7 +1174,7 @@ static D3D12_BOX GetSubresourceCopyBox(const D3D12_RESOURCE_DESC& Desc, const D3
 	};
 }
 
-static uint64 GetPlacedBufferRequiredSize(FD3D12CommandContext* Context, const D3D12_SUBRESOURCE_FOOTPRINT& Footprint, const FPixelFormatInfo& Format, const FName& DebugName)
+static uint64 GetPlacedBufferRequiredSize(FD3D12CommandContext& Context, const D3D12_SUBRESOURCE_FOOTPRINT& Footprint, const FPixelFormatInfo& Format, const FName& DebugName)
 {
 	RHI_BREADCRUMB_CHECK_SHIPPINGF(Context, (Footprint.Width % Format.BlockSizeX) == 0, TEXT("Width not aligned to block size for: '%s'"), *DebugName.ToString());
 	RHI_BREADCRUMB_CHECK_SHIPPINGF(Context, (Footprint.Height % Format.BlockSizeY) == 0, TEXT("Height not aligned to block size for: '%s'"), *DebugName.ToString());
@@ -1211,24 +1211,24 @@ FORCENOINLINE void FD3D12CommandContext::CopyTextureRegionChecked(
 
 			const D3D12_BOX CopyBox = SourceBox ? *SourceBox : GetSubresourceCopyBox(SourceDesc, DestCopyLocation, Indices);
 
-			RHI_BREADCRUMB_CHECK_SHIPPINGF(this, (CopyBox.left % PixelFormatInfo.BlockSizeX) == 0 && (CopyBox.right % PixelFormatInfo.BlockSizeX) == 0, TEXT("Width not aligned to block size for: '%s'"), *DebugName.ToString());
-			RHI_BREADCRUMB_CHECK_SHIPPINGF(this, (CopyBox.top % PixelFormatInfo.BlockSizeY) == 0 && (CopyBox.bottom % PixelFormatInfo.BlockSizeY) == 0, TEXT("Height not aligned to block size for: '%s'"), *DebugName.ToString());
+			RHI_BREADCRUMB_CHECK_SHIPPINGF(*this, (CopyBox.left % PixelFormatInfo.BlockSizeX) == 0 && (CopyBox.right % PixelFormatInfo.BlockSizeX) == 0, TEXT("Width not aligned to block size for: '%s'"), *DebugName.ToString());
+			RHI_BREADCRUMB_CHECK_SHIPPINGF(*this, (CopyBox.top % PixelFormatInfo.BlockSizeY) == 0 && (CopyBox.bottom % PixelFormatInfo.BlockSizeY) == 0, TEXT("Height not aligned to block size for: '%s'"), *DebugName.ToString());
 				
 			const uint32 SubresourceWidth  = AlignArbitrary<uint32>(FMath::Max<uint32>(1u, SourceDesc.Width >> Indices.MipIndex), PixelFormatInfo.BlockSizeX);
 			const uint32 SubresourceHeight = AlignArbitrary<uint32>(FMath::Max<uint32>(1u, SourceDesc.Height >> Indices.MipIndex), PixelFormatInfo.BlockSizeY);
 			const uint32 SubresourceDepth  = AlignArbitrary<uint32>(SourceDesc.Dimension == D3D12_RESOURCE_DIMENSION_TEXTURE3D ? FMath::Max<uint32>(1u, SourceDesc.DepthOrArraySize >> Indices.MipIndex) : 1, PixelFormatInfo.BlockSizeZ);
 
-			RHI_BREADCRUMB_CHECK_SHIPPINGF(this, CopyBox.left <= CopyBox.right && CopyBox.right <= SubresourceWidth, TEXT("Source width out of bounds for: '%s'"), *DebugName.ToString());
-			RHI_BREADCRUMB_CHECK_SHIPPINGF(this, CopyBox.top <= CopyBox.bottom && CopyBox.bottom <= SubresourceHeight, TEXT("Source height out of bounds for: '%s'"), *DebugName.ToString());
-			RHI_BREADCRUMB_CHECK_SHIPPINGF(this, CopyBox.front <= CopyBox.back && CopyBox.back <= SubresourceDepth, TEXT("Source depth out of bounds for: '%s'"), *DebugName.ToString());
+			RHI_BREADCRUMB_CHECK_SHIPPINGF(*this, CopyBox.left <= CopyBox.right && CopyBox.right <= SubresourceWidth, TEXT("Source width out of bounds for: '%s'"), *DebugName.ToString());
+			RHI_BREADCRUMB_CHECK_SHIPPINGF(*this, CopyBox.top <= CopyBox.bottom && CopyBox.bottom <= SubresourceHeight, TEXT("Source height out of bounds for: '%s'"), *DebugName.ToString());
+			RHI_BREADCRUMB_CHECK_SHIPPINGF(*this, CopyBox.front <= CopyBox.back && CopyBox.back <= SubresourceDepth, TEXT("Source depth out of bounds for: '%s'"), *DebugName.ToString());
 			break;
 		}
 
 		// Copy from buffer
 		case D3D12_TEXTURE_COPY_TYPE_PLACED_FOOTPRINT:
 		{
-			const uint32 RequiredPlacedSize = GetPlacedBufferRequiredSize(this, SourceCopyLocation->PlacedFootprint.Footprint, PixelFormatInfo, DebugName);
-			RHI_BREADCRUMB_CHECK_SHIPPINGF(this, SourceCopyLocation->PlacedFootprint.Offset + RequiredPlacedSize <= SourceDesc.Width, TEXT("Source placed buffer width out of bounds for: '%s'"), *DebugName.ToString());
+			const uint32 RequiredPlacedSize = GetPlacedBufferRequiredSize(*this, SourceCopyLocation->PlacedFootprint.Footprint, PixelFormatInfo, DebugName);
+			RHI_BREADCRUMB_CHECK_SHIPPINGF(*this, SourceCopyLocation->PlacedFootprint.Offset + RequiredPlacedSize <= SourceDesc.Width, TEXT("Source placed buffer width out of bounds for: '%s'"), *DebugName.ToString());
 			break;
 		}
 	}
@@ -1242,24 +1242,24 @@ FORCENOINLINE void FD3D12CommandContext::CopyTextureRegionChecked(
 
 			const D3D12_BOX CopyBox = SourceBox ? *SourceBox : GetSubresourceCopyBox(DestDesc, SourceCopyLocation, Indices);
 				
-			RHI_BREADCRUMB_CHECK_SHIPPINGF(this, (CopyBox.left % PixelFormatInfo.BlockSizeX) == 0 && (CopyBox.right % PixelFormatInfo.BlockSizeX) == 0, TEXT("Width not aligned to block size for: '%s'"), *DebugName.ToString());
-			RHI_BREADCRUMB_CHECK_SHIPPINGF(this, (CopyBox.top % PixelFormatInfo.BlockSizeY) == 0 && (CopyBox.bottom % PixelFormatInfo.BlockSizeY) == 0, TEXT("Height not aligned to block size for: '%s'"), *DebugName.ToString());
+			RHI_BREADCRUMB_CHECK_SHIPPINGF(*this, (CopyBox.left % PixelFormatInfo.BlockSizeX) == 0 && (CopyBox.right % PixelFormatInfo.BlockSizeX) == 0, TEXT("Width not aligned to block size for: '%s'"), *DebugName.ToString());
+			RHI_BREADCRUMB_CHECK_SHIPPINGF(*this, (CopyBox.top % PixelFormatInfo.BlockSizeY) == 0 && (CopyBox.bottom % PixelFormatInfo.BlockSizeY) == 0, TEXT("Height not aligned to block size for: '%s'"), *DebugName.ToString());
 				
 			const uint32 SubresourceWidth  = AlignArbitrary<uint32>(FMath::Max<uint32>(1u, DestDesc.Width >> Indices.MipIndex), PixelFormatInfo.BlockSizeX);
 			const uint32 SubresourceHeight = AlignArbitrary<uint32>(FMath::Max<uint32>(1u, DestDesc.Height >> Indices.MipIndex), PixelFormatInfo.BlockSizeY);
 			const uint32 SubresourceDepth  = AlignArbitrary<uint32>(DestDesc.Dimension == D3D12_RESOURCE_DIMENSION_TEXTURE3D ? FMath::Max<uint32>(1u, DestDesc.DepthOrArraySize >> Indices.MipIndex) : 1, PixelFormatInfo.BlockSizeZ);
 
-			RHI_BREADCRUMB_CHECK_SHIPPINGF(this, DestX + (CopyBox.right - CopyBox.left) <= SubresourceWidth, TEXT("Dest width out of bounds for: '%s'"), *DebugName.ToString());
-			RHI_BREADCRUMB_CHECK_SHIPPINGF(this, DestY + (CopyBox.bottom - CopyBox.top) <= SubresourceHeight, TEXT("Dest height out of bounds for: '%s'"), *DebugName.ToString());
-			RHI_BREADCRUMB_CHECK_SHIPPINGF(this, DestZ + (CopyBox.back - CopyBox.front) <= SubresourceDepth, TEXT("Dest depth out of bounds for: '%s'"), *DebugName.ToString());
+			RHI_BREADCRUMB_CHECK_SHIPPINGF(*this, DestX + (CopyBox.right - CopyBox.left) <= SubresourceWidth, TEXT("Dest width out of bounds for: '%s'"), *DebugName.ToString());
+			RHI_BREADCRUMB_CHECK_SHIPPINGF(*this, DestY + (CopyBox.bottom - CopyBox.top) <= SubresourceHeight, TEXT("Dest height out of bounds for: '%s'"), *DebugName.ToString());
+			RHI_BREADCRUMB_CHECK_SHIPPINGF(*this, DestZ + (CopyBox.back - CopyBox.front) <= SubresourceDepth, TEXT("Dest depth out of bounds for: '%s'"), *DebugName.ToString());
 			break;
 		}
 
 		// Copy to buffer
 		case D3D12_TEXTURE_COPY_TYPE_PLACED_FOOTPRINT:
 		{
-			const uint32 RequiredPlacedSize = GetPlacedBufferRequiredSize(this, DestCopyLocation->PlacedFootprint.Footprint, PixelFormatInfo, DebugName);
-			RHI_BREADCRUMB_CHECK_SHIPPINGF(this, DestCopyLocation->PlacedFootprint.Offset + RequiredPlacedSize <= DestDesc.Width, TEXT("Dest placed buffer width out of bounds for: '%s'"), *DebugName.ToString());
+			const uint32 RequiredPlacedSize = GetPlacedBufferRequiredSize(*this, DestCopyLocation->PlacedFootprint.Footprint, PixelFormatInfo, DebugName);
+			RHI_BREADCRUMB_CHECK_SHIPPINGF(*this, DestCopyLocation->PlacedFootprint.Offset + RequiredPlacedSize <= DestDesc.Width, TEXT("Dest placed buffer width out of bounds for: '%s'"), *DebugName.ToString());
 			break;
 		}
 	}
diff --git a/Engine/Source/Runtime/Engine/Private/GPUSortManager.cpp b/Engine/Source/Runtime/Engine/Private/GPUSortManager.cpp
--- a/Engine/Source/Runtime/Engine/Private/GPUSortManager.cpp
+++ b/Engine/Source/Runtime/Engine/Private/GPUSortManager.cpp
@@ -569,7 +569,7 @@ bool FGPUSortManager::TestBatchFlags(EGPUSortFlags BatchFlags, EGPUSortFlags Tas
 	return EnumHasAnyFlags(BatchFlags, TaskFlags & EGPUSortFlags::AnyKeyPrecision) && EnumHasAnyFlags(BatchFlags, TaskFlags & EGPUSortFlags::AnySortLocation);
 }
 
-auto FGPUSortManager::GetPrecisionString(EGPUSortFlags BatchFlags) -> TCHAR const(*)[1]
+auto FGPUSortManager::GetPrecisionString(EGPUSortFlags BatchFlags) -> TCHAR const(&)[1]
 {
 	return RHI_BREADCRUMB_FORCE_STRING_LITERAL(EnumHasAnyFlags(BatchFlags, EGPUSortFlags::LowPrecisionKeys) ? TEXT("LowPrecision") : TEXT("HighPrecision"));
 }
diff --git a/Engine/Source/Runtime/Engine/Private/KismetRenderingLibrary.cpp b/Engine/Source/Runtime/Engine/Private/KismetRenderingLibrary.cpp
--- a/Engine/Source/Runtime/Engine/Private/KismetRenderingLibrary.cpp
+++ b/Engine/Source/Runtime/Engine/Private/KismetRenderingLibrary.cpp
@@ -208,7 +208,7 @@ void UKismetRenderingLibrary::DrawMaterialToRenderTarget(UObject* WorldContextOb
 		Canvas->Init(TextureRenderTarget->SizeX, TextureRenderTarget->SizeY, nullptr, &RenderCanvas);
 
 		{
-			RHI_BREADCRUMB_EVENT_GAMETHREAD("DrawMaterialToRenderTarget: %s", TextureRenderTarget->GetFName());
+			RHI_BREADCRUMB_EVENT_GAMETHREAD_F("DrawMaterialToRenderTarget", "DrawMaterialToRenderTarget: %s", TextureRenderTarget->GetFName());
 
 			ENQUEUE_RENDER_COMMAND(FlushDeferredResourceUpdateCommand)(
 				[RenderTargetResource](FRHICommandListImmediate& RHICmdList)
@@ -797,7 +797,7 @@ void UKismetRenderingLibrary::BeginDrawCanvasToRenderTarget(UObject* WorldContex
 				](FRHICommandListImmediate& RHICmdList)
 			{
 #if  WANTS_DRAW_MESH_EVENTS
-				Breadcrumb->Emplace(RHICmdList, FRHIBreadcrumbData(__FILE__, __LINE__), TEXT("DrawCanvasToTarget: %s"), Name);
+				Breadcrumb->Emplace(RHICmdList, RHI_BREADCRUMB_DESC_FORWARD_VALUES(TEXT("DrawCanvasToTarget"), TEXT("DrawCanvasToTarget: %s"), RHI_GPU_STAT_ARGS_NONE)(Name));
 #endif
 				RenderTargetResource->FlushDeferredResourceUpdate(RHICmdList);
 			});
diff --git a/Engine/Source/Runtime/Engine/Private/SkeletalRenderGPUSkin.cpp b/Engine/Source/Runtime/Engine/Private/SkeletalRenderGPUSkin.cpp
--- a/Engine/Source/Runtime/Engine/Private/SkeletalRenderGPUSkin.cpp
+++ b/Engine/Source/Runtime/Engine/Private/SkeletalRenderGPUSkin.cpp
@@ -1519,7 +1519,7 @@ void FSkeletalMeshObjectGPUSkin::FSkeletalMeshObjectLOD::UpdateMorphVertexBuffer
 	const bool bUseGPUMorphTargets = UseGPUMorphTargets(FeatureLevel);
 	InMorphVertexBuffer.RecreateResourcesIfRequired(RHICmdList, bUseGPUMorphTargets);
 
-	RHI_BREADCRUMB_EVENT_STAT(RHICmdList, MorphTargets, "MorphUpdate%s_%s_LOD%d LodVertices=%d Batches=%d"
+	RHI_BREADCRUMB_EVENT_STAT_F(RHICmdList, MorphTargets, "MorphUpdate", "MorphUpdate%s_%s_LOD%d LodVertices=%d Batches=%d"
 		, RHI_BREADCRUMB_FORCE_STRING_LITERAL(Mode == EGPUSkinCacheEntryMode::RayTracing ? TEXT("[RT]") : TEXT(""))
 		, OwnerName
 		, LODIndex
diff --git a/Engine/Source/Runtime/Engine/Public/GPUSortManager.h b/Engine/Source/Runtime/Engine/Public/GPUSortManager.h
--- a/Engine/Source/Runtime/Engine/Public/GPUSortManager.h
+++ b/Engine/Source/Runtime/Engine/Public/GPUSortManager.h
@@ -470,7 +470,7 @@ private:
 	/** Test whether a batch flags are compatible with a given task. Used to know if the task can be merged in the batch.  */
 	static FORCEINLINE bool TestBatchFlags(EGPUSortFlags BatchFlags, EGPUSortFlags TaskFlags);
 	/** Convert the precision flags into a string, used for GPU markers. */
-	static FORCEINLINE auto GetPrecisionString(EGPUSortFlags BatchFlags) -> TCHAR const(*)[1];
+	static FORCEINLINE auto GetPrecisionString(EGPUSortFlags BatchFlags) -> TCHAR const(&)[1];
 
 	/**
 	 * Find an unused buffer to assign to a newly created FSortBatch. 
diff --git a/Engine/Source/Runtime/Landscape/Private/LandscapeBlueprintBrushBase.cpp b/Engine/Source/Runtime/Landscape/Private/LandscapeBlueprintBrushBase.cpp
--- a/Engine/Source/Runtime/Landscape/Private/LandscapeBlueprintBrushBase.cpp
+++ b/Engine/Source/Runtime/Landscape/Private/LandscapeBlueprintBrushBase.cpp
@@ -160,7 +160,7 @@ UTextureRenderTarget2D* ALandscapeBlueprintBrushBase::Execute(const FLandscapeBr
 	}
 	UTextureRenderTarget2D* Result = nullptr;
 	{
-		RHI_BREADCRUMB_EVENT_GAMETHREAD("BP Render %s (%s): %s", GetActorNameOrLabel(), UEnum::GetDisplayValueAsText(InParameters.LayerType).ToString(), LayerDetailString);
+		RHI_BREADCRUMB_EVENT_GAMETHREAD_F("BP Render", "BP Render %s (%s): %s", GetActorNameOrLabel(), UEnum::GetDisplayValueAsText(InParameters.LayerType).ToString(), LayerDetailString);
 
 		TGuardValue<bool> AutoRestore(GAllowActorScriptExecutionInEditor, true);
 		Result = RenderLayer(InParameters);
@@ -394,11 +394,11 @@ bool ALandscapeBlueprintBrushBase::RenderLayer(UE::Landscape::EditLayers::FRende
 	for (int32 TargetLayerIndexInGroup = 0; TargetLayerIndexInGroup < NumTargetLayersInGroup; ++TargetLayerIndexInGroup)
 	{
 		const FName TargetLayerName = RenderParams.TargetLayerGroupLayerNames[TargetLayerIndexInGroup];
-		RHI_BREADCRUMB_EVENT_GAMETHREAD("Render %s", TargetLayerName);
+		RHI_BREADCRUMB_EVENT_GAMETHREAD_F("Render Layer", "Render %s", TargetLayerName);
 
 		// If necessary, copy from the source render target (slice) to the scratch render target 2D : 
 		{
-			RHI_BREADCRUMB_EVENT_GAMETHREAD("Copy Source (slice %i) -> %s", TargetLayerIndexInGroup, CurrentLayerReadRT2D->GetDebugName());
+			RHI_BREADCRUMB_EVENT_GAMETHREAD_F("Copy Source", "Copy Source (slice %i) -> %s", TargetLayerIndexInGroup, CurrentLayerReadRT2D->GetDebugName());
 			ULandscapeScratchRenderTarget::FCopyFromScratchRenderTargetParams CopyParams(CurrentLayerReadRT);
 			// It's important to copy only the needed size, because CurrentLayerReadRT2D comes from ScratchTexture, whose size may be larger than RenderParams.RenderAreaSectionRect)
 			CopyParams.CopySize = RenderParams.RenderAreaSectionRect.Size();
@@ -432,7 +432,7 @@ bool ALandscapeBlueprintBrushBase::RenderLayer(UE::Landscape::EditLayers::FRende
 
 		// Resolve back to the write RT 
 		{
-			RHI_BREADCRUMB_EVENT_GAMETHREAD("Resolve BP Render Result -> %s (slice %i)", WriteRT->GetDebugName(), TargetLayerIndexInGroup);
+			RHI_BREADCRUMB_EVENT_GAMETHREAD_F("Resolve BP Render Result", "Resolve BP Render Result -> %s (slice %i)", WriteRT->GetDebugName(), TargetLayerIndexInGroup);
 
 			ENQUEUE_RENDER_COMMAND(ResolveLayerData)(
 				[ InputResource = OutputRT2D->GetResource()
diff --git a/Engine/Source/Runtime/Landscape/Private/LandscapeEditLayerMergeRenderContext.cpp b/Engine/Source/Runtime/Landscape/Private/LandscapeEditLayerMergeRenderContext.cpp
--- a/Engine/Source/Runtime/Landscape/Private/LandscapeEditLayerMergeRenderContext.cpp
+++ b/Engine/Source/Runtime/Landscape/Private/LandscapeEditLayerMergeRenderContext.cpp
@@ -835,7 +835,7 @@ void FMergeRenderContext::Render(TFunction<void(const FOnRenderBatchTargetGroupD
 	{
 		const FMergeRenderBatch& RenderBatch = RenderBatches[CurrentRenderBatchIndex];
 		FString RenderBatchDebugName = FString::Format(TEXT("Render Batch [{0}] : ({1},{2})->({3},{4})"), { CurrentRenderBatchIndex, RenderBatch.SectionRect.Min.X, RenderBatch.SectionRect.Min.Y, RenderBatch.SectionRect.Max.X, RenderBatch.SectionRect.Max.Y });
-		RHI_BREADCRUMB_EVENT_GAMETHREAD("%s", RenderBatchDebugName);
+		RHI_BREADCRUMB_EVENT_GAMETHREAD_F("Render Batch", "%s", RenderBatchDebugName);
 		TRACE_CPUPROFILER_EVENT_SCOPE(RenderBatch);
 
 		checkf((RenderBatch.RenderSteps.Num() >= 1) && (RenderBatch.RenderSteps.Last().Type == FMergeRenderStep::EType::SignalBatchMergeGroupDone),
@@ -1037,7 +1037,7 @@ void FMergeRenderContext::Render(TFunction<void(const FOnRenderBatchTargetGroupD
 				    else
 				    {
 					    checkf(RDGBuilderRecorder.IsEmpty() && !RDGBuilderRecorder.IsRecording(), TEXT("(Render/Blend)Layer_Immediate should be preceded by a EndRenderCommandRecorder which should finalize a command recorder and destroy it"));
-					        RHI_BREADCRUMB_EVENT_GAMETHREAD("%s", RenderStepProfilingEventName);
+					        RHI_BREADCRUMB_EVENT_GAMETHREAD_F("Step", "%s", RenderStepProfilingEventName);
 					    RenderOrBlend();
 				    }
     
@@ -1047,7 +1047,7 @@ void FMergeRenderContext::Render(TFunction<void(const FOnRenderBatchTargetGroupD
 			    case FMergeRenderStep::EType::SignalBatchMergeGroupDone:
 			    {
 				    TRACE_CPUPROFILER_EVENT_SCOPE(MergeGroupDone);
-				    RHI_BREADCRUMB_EVENT_GAMETHREAD("Step [%i] (%s) : Render Group Done", RenderLayerStepIndex, ConvertTargetLayerNamesToString(TargetLayerGroupLayerNames));
+				    RHI_BREADCRUMB_EVENT_GAMETHREAD_F("Step", "Step [%i] (%s) : Render Group Done", RenderLayerStepIndex, ConvertTargetLayerNamesToString(TargetLayerGroupLayerNames));
     
 				    checkf(RDGBuilderRecorder.IsEmpty() && !RDGBuilderRecorder.IsRecording(), TEXT("SignalBatchMergeGroupDone should be preceded by a EndRenderCommandRecorder which should finalize a command recorder and destroy it"));
 
diff --git a/Engine/Source/Runtime/Landscape/Private/LandscapeEditLayers.cpp b/Engine/Source/Runtime/Landscape/Private/LandscapeEditLayers.cpp
--- a/Engine/Source/Runtime/Landscape/Private/LandscapeEditLayers.cpp
+++ b/Engine/Source/Runtime/Landscape/Private/LandscapeEditLayers.cpp
@@ -2034,7 +2034,8 @@ private:
 	void CopyInternal(FRHICommandListImmediate& InRHICmdList)
 	{
 		// TODO [jonathan.bard] : make those perf tags optional : with the amount of textures we copy, it slows down texture copies quite a bit : 
-		RHI_BREADCRUMB_EVENT_STAT(InRHICmdList, LandscapeLayers_CopyTexture
+		RHI_BREADCRUMB_EVENT_STAT_F(InRHICmdList, LandscapeLayers_CopyTexture
+			, "LandscapeLayers_Copy"
 			, "LandscapeLayers_Copy %s -> %s, Mip (%d -> %d), Array Index (%d -> %d)"
 			, Params.SourceResourceDebugName
 			, Params.DestResourceDebugName
@@ -2083,7 +2084,8 @@ private:
 		const int32 NumChannelsDest = GPixelFormats[Params.DestResource->TextureRHI->GetDesc().Format].NumComponents;
 
 		// TODO [jonathan.bard] : make those perf tags optional : with the amount of textures we copy, it slows down texture copies quite a bit : 
-		RHI_BREADCRUMB_EVENT_STAT(InRHICmdList, LandscapeLayers_CopyTexturePS
+		RHI_BREADCRUMB_EVENT_STAT_F(InRHICmdList, LandscapeLayers_CopyTexturePS
+			, "LandscapeLayers_CopyPS"
 			, "LandscapeLayers_CopyPS %s -> %s, Mip (%d -> %d), Array Index (%d -> %d), [%s]"
 			, Params.SourceResourceDebugName
 			, Params.DestResourceDebugName
@@ -2157,7 +2159,7 @@ public:
 
 	void Clear(FRHICommandListImmediate& InRHICmdList)
 	{
-		RHI_BREADCRUMB_EVENT_STAT(InRHICmdList, LandscapeLayers_Clear, "LandscapeLayers_Clear %s", DebugName);
+		RHI_BREADCRUMB_EVENT_STAT_F(InRHICmdList, LandscapeLayers_Clear, "LandscapeLayers_Clear", "LandscapeLayers_Clear %s", DebugName);
 		SCOPED_GPU_STAT(InRHICmdList, LandscapeLayers_Clear);
 		TRACE_CPUPROFILER_EVENT_SCOPE(LandscapeLayersWeightmapClear_RenderThread::Clear);
 
@@ -2201,7 +2203,7 @@ public:
 
 	void Render(FRHICommandListImmediate& InRHICmdList, bool InClearRT)
 	{
-		RHI_BREADCRUMB_EVENT_STAT(InRHICmdList, LandscapeLayers_Render, "LandscapeLayers_Render %s", DebugName);
+		RHI_BREADCRUMB_EVENT_STAT_F(InRHICmdList, LandscapeLayers_Render, "LandscapeLayers_Render", "LandscapeLayers_Render %s", DebugName);
 		SCOPED_GPU_STAT(InRHICmdList, LandscapeLayers_Render);
 		INC_DWORD_STAT(STAT_LandscapeLayersRegenerateDrawCalls);
 		TRACE_CPUPROFILER_EVENT_SCOPE(FLandscapeLayersRender_RenderThread::Render);
@@ -3143,7 +3145,7 @@ struct FLandscapeLayersCopyReadbackTextureParams
 void ExecuteCopyToReadbackTexture(TArray<FLandscapeLayersCopyReadbackTextureParams>& InParams)
 {
 	TRACE_CPUPROFILER_EVENT_SCOPE(ExecuteCopyToReadbackTexture);
-	RHI_BREADCRUMB_EVENT_GAMETHREAD("Copy to readback textures (%d copies)", InParams.Num());
+	RHI_BREADCRUMB_EVENT_GAMETHREAD_F("Copy to readback textures", "Copy to readback textures (%d copies)", InParams.Num());
 	if (!FApp::CanEverRender())
 	{
 		return;
@@ -3220,7 +3222,7 @@ void ALandscape::CopyTexturePS(const FString& InSourceDebugName, FTextureResourc
 		[InSourceDebugName, InSourceResource, InDestDebugName, InDestResource](FRHICommandListImmediate& RHICmdList)
 	{
 		TRACE_CPUPROFILER_EVENT_SCOPE(LandscapeLayers_RT_CopyTexturePS);
-		RHI_BREADCRUMB_EVENT_STAT(RHICmdList, LandscapeLayers_CopyTexturePS, "LandscapeLayers_CopyTexturePS %s -> %s", InSourceDebugName, InDestDebugName);
+		RHI_BREADCRUMB_EVENT_STAT_F(RHICmdList, LandscapeLayers_CopyTexturePS, "LandscapeLayers_CopyTexturePS", "LandscapeLayers_CopyTexturePS %s -> %s", InSourceDebugName, InDestDebugName);
 		SCOPED_GPU_STAT(RHICmdList, LandscapeLayers_CopyTexturePS);
 
 		check(InSourceResource->GetSizeX() == InDestResource->GetSizeX());
@@ -7724,7 +7726,7 @@ int32 ALandscape::PerformLayersHeightmapsBatchedMerge(const FUpdateLayersContent
 		{
 			TRACE_CPUPROFILER_EVENT_SCOPE(GenerateMips);
 			const int32 NumMips = (int32)FMath::CeilLogTwo(Landscape->SubsectionSizeQuads) + 1;
-			RHI_BREADCRUMB_EVENT_GAMETHREAD("Generate %i remaining mips", NumMips - 1);
+			RHI_BREADCRUMB_EVENT_GAMETHREAD_F("Generate remaining mips", "Generate %i remaining mips", NumMips - 1);
 
 			FIntPoint CurrentMipResolution = RenderBatch->GetRenderTargetResolution(/*bInWithDuplicateBorders = */true); // Mips are generated after the borders have been duplicated
 			FIntPoint CurrentMipSubsectionSize = Landscape->SubsectionSizeQuads + 1;
@@ -7743,7 +7745,7 @@ int32 ALandscape::PerformLayersHeightmapsBatchedMerge(const FUpdateLayersContent
 				check(CurrentMipSubsectionSize.X > 0 && CurrentMipSubsectionSize.Y > 0);
 
 				{
-					RHI_BREADCRUMB_EVENT_GAMETHREAD("Generate mip %i", MipIndex);
+					RHI_BREADCRUMB_EVENT_GAMETHREAD_F("Generate mip", "Generate mip %i", MipIndex);
 
 					ENQUEUE_RENDER_COMMAND(LandscapeLayers_Cmd_HeightmapsGenerateMips)(
 						[ OutputResource = WriteRT->GetRenderTarget2D()->GetResource()
@@ -7774,7 +7776,7 @@ int32 ALandscape::PerformLayersHeightmapsBatchedMerge(const FUpdateLayersContent
 				// Then copy the appropriate regions to the destination texture mips : 
 				// TODO [jonathan.bard] : add this when we don't auto-transition to SRV in the copy texture thing : WriteRT->TransitionTo(ERHIAccess::CopySrc);
 				{
-					RHI_BREADCRUMB_EVENT_GAMETHREAD("Copy mip %i", MipIndex);
+					RHI_BREADCRUMB_EVENT_GAMETHREAD_F("Copy mip", "Copy mip %i", MipIndex);
 
 					WriteRT->TransitionTo(ERHIAccess::SRVMask, RDGBuilderRecorder);
 					TArray<FLandscapeLayersCopyTextureParams> DeferredCopyTextures;
@@ -10118,7 +10120,7 @@ int32 ALandscape::PerformLayersWeightmapsGlobalMerge(FUpdateLayersContentContext
 					int32 LayerIndex = LayerInfoObject.Value;
 					ULandscapeLayerInfoObject* LayerInfoObj = LayerInfoObject.Key;
 
-					RHI_BREADCRUMB_EVENT_GAMETHREAD("LS Weight: %s PaintLayer: %s", Layer.Name, LayerInfoObj->LayerName);
+					RHI_BREADCRUMB_EVENT_GAMETHREAD_F("Layer", "LS Weight: %s PaintLayer: %s", Layer.Name, LayerInfoObj->LayerName);
 
 					// Copy the layer we are working on
 					SourceDebugName = FString::Printf(TEXT("Weight: %s PaintLayer: %s, CurrentProcLayerWeightmapAllLayersResource"), *Layer.Name.ToString(), *LayerInfoObj->LayerName.ToString());
diff --git a/Engine/Source/Runtime/Landscape/Public/LandscapeUtils.h b/Engine/Source/Runtime/Landscape/Public/LandscapeUtils.h
--- a/Engine/Source/Runtime/Landscape/Public/LandscapeUtils.h
+++ b/Engine/Source/Runtime/Landscape/Public/LandscapeUtils.h
@@ -10,6 +10,7 @@
 #include "Materials/Material.h"
 #include "RenderGraphDefinitions.h"
 #include "RenderGraphEvent.h"
+#include "RenderGraphBuilder.h"
 
 enum EShaderPlatform : uint16;
 class FRDGBuilder;
@@ -184,33 +185,31 @@ public:
 	class FScopedBreadcrumbEvent final
 	{
 	public:
-		template<typename... TArgs>
-		FScopedBreadcrumbEvent(FRDGBuilderRecorder& InRecorder, TArgs&&... InArgs)
+		FScopedBreadcrumbEvent(FRDGBuilderRecorder& InRecorder, TCHAR const* StaticName, FRDGEventName&& EventName)
 			: Recorder(&InRecorder)
 			, RDGEvent(MakeShared<TOptional<TRDGEventScopeGuard<FRDGScope_RHI>>>())
 		{
 			//  We use a shared ptr to create a TOptional<TRDGEventScopeGuard> immediately, then capture this shared ptr in these additional operations' lambdas (so that the object
 			//  continues to live until the closing tag operation)
-			Recorder->EnqueueRDGCommand([
+			Recorder->EnqueueRDGCommand(
+			[
 				RDGEvent = RDGEvent, 
-				Values = std::make_tuple(std::forward<TArgs>(InArgs) ...)
-				](FRDGBuilder& GraphBuilder) mutable
-				{
-					std::apply([RDGEvent, &GraphBuilder](auto&&... InnerArgs)
-						{
-							// Allocate the TOptional<TRDGEventScopeGuard> now in order to insert the tag : 
-							RDGEvent->Emplace(GraphBuilder, ERDGScopeFlags::None, FRHIBreadcrumbData(__FILE__, __LINE__), std::forward<TArgs>(InnerArgs)...);
-						}, Values);
-				});
+				StaticName = StaticName,
+				EventName = MoveTemp(EventName)
+			](FRDGBuilder& GraphBuilder) mutable
+			{
+				// Allocate the TOptional<TRDGEventScopeGuard> now in order to insert the tag : 
+				RDG_EVENT_SCOPE_CONSTRUCT(*RDGEvent, GraphBuilder, true, ERDGScopeFlags::None, RHI_GPU_STAT_ARGS_NONE, StaticName, MoveTemp(EventName));
+			});
 		}
 
 		~FScopedBreadcrumbEvent()
 		{
 			Recorder->EnqueueRDGCommand([RDGEvent = RDGEvent](FRDGBuilder& GraphBuilder)
-				{
-					// Reset the TOptional in order to delete the TRDGEventScopeGuard, which will remove the tag : 
-					RDGEvent->Reset();
-				});
+			{
+				// Reset the TOptional in order to delete the TRDGEventScopeGuard, which will remove the tag : 
+				RDGEvent->Reset();
+			});
 		}
 
 	private:
@@ -221,7 +220,7 @@ public:
 };
 
 #if RDG_EVENTS
-#define RDG_RENDER_COMMAND_RECORDER_BREADCRUMB_EVENT(Recorder, Format, ...) FRDGBuilderRecorder::FScopedBreadcrumbEvent ANONYMOUS_VARIABLE(BreadcrumbEvent)(Recorder, RDG_EVENT_NAME(Format, ##__VA_ARGS__))
+#define RDG_RENDER_COMMAND_RECORDER_BREADCRUMB_EVENT(Recorder, Format, ...) FRDGBuilderRecorder::FScopedBreadcrumbEvent ANONYMOUS_VARIABLE(BreadcrumbEvent)(Recorder, TEXT(Format), RDG_EVENT_NAME(Format, ##__VA_ARGS__))
 #else // RDG_EVENTS
 #define RDG_RENDER_COMMAND_RECORDER_BREADCRUMB_EVENT(Recorder, Format, ...) do { } while(0)
 #endif // !RDG_EVENTS
diff --git a/Engine/Source/Runtime/Launch/Private/LaunchEngineLoop.cpp b/Engine/Source/Runtime/Launch/Private/LaunchEngineLoop.cpp
--- a/Engine/Source/Runtime/Launch/Private/LaunchEngineLoop.cpp
+++ b/Engine/Source/Runtime/Launch/Private/LaunchEngineLoop.cpp
@@ -5228,12 +5228,17 @@ static inline void BeginFrameRenderThread(FRHICommandListImmediate& RHICmdList,
 #if CSV_PROFILER
 	if (FCsvProfiler::Get()->IsCapturing_Renderthread())
 	{
-		GRHIFrameBreadcrumb.Emplace(RHICmdList, FRHIBreadcrumbData(__FILE__, __LINE__), TEXT("CsvFrame %d"), FCsvProfiler::Get()->GetCaptureFrameNumberRT());
+		GRHIFrameBreadcrumb.Emplace(RHICmdList, RHI_BREADCRUMB_DESC_FORWARD_VALUES(TEXT("Frame"), TEXT("Frame %d (CsvFrame %d)"), RHI_GPU_STAT_ARGS_NONE)(
+			  RHI_BREADCRUMB_FIELD("Frame Number", CurrentFrameCounter)
+			, RHI_BREADCRUMB_FIELD("Csv Frame Number", FCsvProfiler::Get()->GetCaptureFrameNumberRT())
+		));
 	}
 	else
 #endif
 	{
-		GRHIFrameBreadcrumb.Emplace(RHICmdList, FRHIBreadcrumbData(__FILE__, __LINE__), TEXT("Frame %d"), CurrentFrameCounter);
+		GRHIFrameBreadcrumb.Emplace(RHICmdList, RHI_BREADCRUMB_DESC_FORWARD_VALUES(TEXT("Frame"), TEXT("Frame %d"), RHI_GPU_STAT_ARGS_NONE)(
+			RHI_BREADCRUMB_FIELD("Frame Number", CurrentFrameCounter)
+		));
 	}
 #endif
 
diff --git a/Engine/Source/Runtime/OpenGLDrv/Private/OpenGLDrv.cpp b/Engine/Source/Runtime/OpenGLDrv/Private/OpenGLDrv.cpp
--- a/Engine/Source/Runtime/OpenGLDrv/Private/OpenGLDrv.cpp
+++ b/Engine/Source/Runtime/OpenGLDrv/Private/OpenGLDrv.cpp
@@ -129,7 +129,7 @@ void FOpenGLDynamicRHI::RHIEGLTerminateContext()
 		{
 			if (!NameStr)
 			{
-				NameStr = Breadcrumb->Name.GetTCHAR(Buffer);
+				NameStr = Breadcrumb->GetTCHAR(Buffer);
 			}
 			return NameStr;
 		};
diff --git a/Engine/Source/Runtime/RHI/Private/GPUProfiler.cpp b/Engine/Source/Runtime/RHI/Private/GPUProfiler.cpp
--- a/Engine/Source/Runtime/RHI/Private/GPUProfiler.cpp
+++ b/Engine/Source/Runtime/RHI/Private/GPUProfiler.cpp
@@ -1322,20 +1322,20 @@ namespace UE::RHI::GPUProfiler
 #if HAS_GPU_STATS
 	// Per queue GPU stats
 
-	// Total busy time on the current queue.
-	static FGPUStat GPUStat_Total(TEXT("Queue Total"));
+	// Total busy time on the current queue. StatName == "Unaccounted" is used by the Csv profiler
+	static FGPUStat GPUStat_Total(TEXT("Unaccounted"), TEXT("Queue Total"));
 #endif
 
-	TMap<FQueue, TUniquePtr<FGPUStat::FStatCategory>> FGPUStat::Categories;
-
-	FGPUStat::FStatCategory& FGPUStat::GetCategory(FQueue Queue)
+#if STATS
+	TCHAR const* FGPUStat::GetTypeString(EType Type)
 	{
-		TUniquePtr<FStatCategory>& Category = Categories.FindOrAdd(Queue);
-		if (!Category)
+		switch (Type)
 		{
-			Category = MakeUnique<FStatCategory>(Queue);
+		default: checkNoEntry(); [[fallthrough]];
+		case EType::Busy: return TEXT("Busy");
+		case EType::Wait: return TEXT("Wait");
+		case EType::Idle: return TEXT("Idle");
 		}
-		return *Category;
 	}
 
 	FString FGPUStat::GetIDString(FQueue Queue, bool bFriendly)
@@ -1358,52 +1358,47 @@ namespace UE::RHI::GPUProfiler
 		}
 	}
 
-	FGPUStat::FStatCategory_Stats::FStatCategory_Stats(FQueue Queue)
-	#if STATS
-		: GroupName(FString::Printf(TEXT("STATGROUP_%s"), *GetIDString(Queue, false)))
-		, GroupDesc(FString::Printf(TEXT("%s Timing"), *GetIDString(Queue, true)))
-	#endif
-	{}
-
-	FGPUStat::FStatCategory_Csv::FStatCategory_Csv(FQueue Queue)
-	#if CSV_PROFILER_STATS
-		: Csv(*GetIDString(Queue, false), true)
-	#endif
-	{}
-
-	FGPUStat::FStatCategory::FStatCategory(FQueue Queue)
-		: FStatCategory_Stats(Queue)
-		, FStatCategory_Csv(Queue)
-	{}
-
-#if STATS
-	TStatId FGPUStat::GetStatId(FQueue Queue, EType Type)
+	FGPUStat::FStatInstance::FInner& FGPUStat::GetStatInstance(FQueue Queue, EType Type)
 	{
 		FStatInstance& Instance = Instances.FindOrAdd(Queue);
 
-		FStatInstance::FInner* Inner;
 		switch (Type)
 		{
 		default: checkNoEntry(); [[fallthrough]];
-		case EType::Busy: Inner = &Instance.Busy; break;
-		case EType::Wait: Inner = &Instance.Wait; break;
-		case EType::Idle: Inner = &Instance.Idle; break;
+		case EType::Busy: return Instance.Busy; break;
+		case EType::Wait: return Instance.Wait; break;
+		case EType::Idle: return Instance.Idle; break;
 		}
+	}
+
+	TMap<FQueue, TUniquePtr<FGPUStat::FStatCategory>> FGPUStat::FStatCategory::Categories;
 
-		if (!Inner->Stat)
+	FGPUStat::FStatCategory::FStatCategory(FQueue Queue)
+		: GroupName(FString::Printf(TEXT("STATGROUP_%s"), *GetIDString(Queue, false)))
+		, GroupDesc(FString::Printf(TEXT("%s Timing"), *GetIDString(Queue, true)))
+	{}
+
+	TStatId FGPUStat::GetStatId(FQueue Queue, EType Type)
+	{
+		FStatInstance::FInner& Instance = GetStatInstance(Queue, Type);
+
+		if (!Instance.Stat)
 		{
-			FStatCategory& Category = GetCategory(Queue);
+			TUniquePtr<FStatCategory>& Category = FStatCategory::Categories.FindOrAdd(Queue);
+			if (!Category)
+			{
+				Category = MakeUnique<FStatCategory>(Queue);
+			}
 
 			// Encode the stat type in the FName number
-			int32 Number = int32(Type);
-			Inner->StatName = FName(*FString::Printf(TEXT("STAT_%s_%s"), *GetIDString(Queue, false), Name), Number);
+			Instance.StatName = FName(*FString::Printf(TEXT("STAT_%s_%s"), *GetIDString(Queue, false), DisplayName), int32(Type));
 
-			Inner->Stat = MakeUnique<FDynamicStat>(
-				Inner->StatName,
-				Name,
-				*Category.GroupName,
+			Instance.Stat = MakeUnique<FDynamicStat>(
+				Instance.StatName,
+				DisplayName,
+				*Category->GroupName,
 				FStatNameAndInfo::GpuStatCategory,
-				*Category.GroupDesc,
+				*Category->GroupDesc,
 				true, // IsDefaultEnabled
 				true, // IsClearEveryFrame
 				EStatDataType::ST_double,
@@ -1413,23 +1408,9 @@ namespace UE::RHI::GPUProfiler
 			);
 		}
 
-		return Inner->Stat->GetStatId();
+		return Instance.Stat->GetStatId();
 	}
-#endif
-
-#if CSV_PROFILER_STATS
-	FCsvDeclaredStat& FGPUStat::GetCsvStat(FQueue Queue)
-	{
-		FStatInstance& Instance = Instances.FindOrAdd(Queue);
-
-		if (!Instance.Busy.CsvStat.IsSet())
-		{
-			FStatCategory& Category = GetCategory(Queue);
-			Instance.Busy.CsvStat.Emplace(Name, Category.Csv.Index);
-		}
 
-		return *Instance.Busy.CsvStat;
-	}
 #endif
 
 	// Handles computing the "stat unit" GPU time, "stat gpu" stats, and "profilegpu".
@@ -1544,9 +1525,19 @@ namespace UE::RHI::GPUProfiler
 
 		struct FStatState
 		{
-			uint64 BusyCycles = 0;
-			uint64 IdleCycles = 0;
-			uint64 WaitCycles = 0;
+			struct
+			{
+				uint64 BusyCycles = 0;
+				uint64 IdleCycles = 0;
+				uint64 WaitCycles = 0;
+
+				void Accumulate(uint64 Busy, uint64 Wait, uint64 Idle)
+				{
+					BusyCycles += Busy;
+					IdleCycles += Idle;
+					WaitCycles += Wait;
+				}
+			} Exclusive, Inclusive;
 
 			FStatState() = default;
 			FStatState(FStatState const&) = default;
@@ -1554,16 +1545,8 @@ namespace UE::RHI::GPUProfiler
 			FStatState(FStatState&& Other)
 				: FStatState(Other)
 			{
-				Other.BusyCycles = 0;
-				Other.IdleCycles = 0;
-				Other.WaitCycles = 0;
-			}
-
-			void Accumulate(uint64 Busy, uint64 Wait, uint64 Idle)
-			{
-				BusyCycles += Busy;
-				IdleCycles += Idle;
-				WaitCycles += Wait;
+				Other.Exclusive = {};
+				Other.Inclusive = {};
 			}
 
 			void EmitResults(FQueue Queue
@@ -1579,16 +1562,35 @@ namespace UE::RHI::GPUProfiler
 			) const
 			{
 			#if STATS && HAS_GPU_STATS
-				Stats->AddMessage(GPUStat.GetStatId(Queue, FGPUStat::EType::Busy).GetName(), EStatOperation::Set, FPlatformTime::ToMilliseconds64(BusyCycles));
-				Stats->AddMessage(GPUStat.GetStatId(Queue, FGPUStat::EType::Idle).GetName(), EStatOperation::Set, FPlatformTime::ToMilliseconds64(IdleCycles));
-				Stats->AddMessage(GPUStat.GetStatId(Queue, FGPUStat::EType::Wait).GetName(), EStatOperation::Set, FPlatformTime::ToMilliseconds64(WaitCycles));
+				Stats->AddMessage(GPUStat.GetStatId(Queue, FGPUStat::EType::Busy).GetName(), EStatOperation::Set, FPlatformTime::ToMilliseconds64(Inclusive.BusyCycles));
+				Stats->AddMessage(GPUStat.GetStatId(Queue, FGPUStat::EType::Idle).GetName(), EStatOperation::Set, FPlatformTime::ToMilliseconds64(Inclusive.IdleCycles));
+				Stats->AddMessage(GPUStat.GetStatId(Queue, FGPUStat::EType::Wait).GetName(), EStatOperation::Set, FPlatformTime::ToMilliseconds64(Inclusive.WaitCycles));
 			#endif
 
 			#if CSV_PROFILER_STATS
-				if (CsvProfiler)
+				if (CsvProfiler && Queue.Type == FQueue::EType::Graphics && Queue.Index == 0)
 				{
-					FCsvDeclaredStat& CsvStat = GPUStat.GetCsvStat(Queue);
-					CsvProfiler->RecordEndOfPipeCustomStat(CsvStat.Name, CsvStat.CategoryIndex, FPlatformTime::ToMilliseconds64(BusyCycles), ECsvCustomStatOp::Set);
+					if (!GPUStat.CsvStat.IsSet())
+					{
+						static TArray<TUniquePtr<FCsvCategory>> CsvGPUCategories;
+						if (!CsvGPUCategories.IsValidIndex(Queue.GPU))
+						{
+							CsvGPUCategories.SetNum(Queue.GPU + 1);
+						}
+
+						TUniquePtr<FCsvCategory>& Category = CsvGPUCategories[Queue.GPU];
+						if (!Category)
+						{
+							Category = Queue.GPU > 0
+								? MakeUnique<FCsvCategory>(*FString::Printf(TEXT("GPU%d"), Queue.GPU + 1), true)
+								: MakeUnique<FCsvCategory>(TEXT("GPU"), true);
+						}
+
+						GPUStat.CsvStat.Emplace(GPUStat.StatName, Category->Index);
+					}
+
+					uint64 TotalCycles = Exclusive.BusyCycles + Exclusive.WaitCycles;
+					CsvProfiler->RecordEndOfPipeCustomStat(GPUStat.CsvStat->Name, GPUStat.CsvStat->CategoryIndex, FPlatformTime::ToMilliseconds64(TotalCycles), ECsvCustomStatOp::Set);
 				}
 			#endif
 			}
@@ -1651,6 +1653,7 @@ namespace UE::RHI::GPUProfiler
 
 		#if WITH_RHI_BREADCRUMBS
 			TMap<FRHIBreadcrumbData_Stats, int32> ActiveStats;
+			TArray<FRHIBreadcrumbData_Stats> ActiveStatsStack;
 			FRHIBreadcrumbNode* Breadcrumb = nullptr;
 		#endif
 
@@ -1878,15 +1881,26 @@ namespace UE::RHI::GPUProfiler
 
 			void AccumulateTime(uint64 Busy, uint64 Wait, uint64 Idle)
 			{
-				Timestamps.WholeQueueStat.Accumulate(Busy, Wait, Idle);
-
 			#if WITH_RHI_BREADCRUMBS
 				// Apply the timings to all active stats
 				for (auto const& [Stat, RefCount] : ActiveStats)
 				{
-					Timestamps.Stats.FindChecked(Stat).Accumulate(Busy, Wait, Idle);
+					FStatState& State = Timestamps.Stats.FindChecked(Stat);
+					State.Inclusive.Accumulate(Busy, Wait, Idle);
+
+					if (ActiveStatsStack.Num() > 0 && ActiveStatsStack.Last() == Stat)
+					{
+						State.Exclusive.Accumulate(Busy, Wait, Idle);
+					}
 				}
+
+				if (ActiveStatsStack.Num() == 0)
 			#endif
+				{
+					Timestamps.WholeQueueStat.Exclusive.Accumulate(Busy, Wait, Idle);
+				}
+
+				Timestamps.WholeQueueStat.Inclusive.Accumulate(Busy, Wait, Idle);
 
 			#if WITH_PROFILEGPU
 				for (FNode* Node = Profile.Current; Node; Node = Node->Parent)
@@ -1941,7 +1955,7 @@ namespace UE::RHI::GPUProfiler
 				AccumulateTime(Busy, 0, 0);
 				LastGPUCycles = FMath::Max(LastGPUCycles, Event.GPUTimestampTOP);
 
-				FRHIBreadcrumbData_Stats const& Stat = Event.Breadcrumb->Name.Data;
+				FRHIBreadcrumbData_Stats const& Stat = Event.Breadcrumb->Data;
 				if (Stat.ShouldComputeStat())
 				{
 					// Disregard the stat if it is nested within itself (i.e. its already in the ActiveStats map with a non-zero ref count).
@@ -1951,12 +1965,14 @@ namespace UE::RHI::GPUProfiler
 					{
 						Timestamps.Stats.FindOrAdd(Stat);
 					}
+
+					ActiveStatsStack.Add(Stat);
 				}
 
 				Breadcrumb = Event.Breadcrumb;
 
 				FRHIBreadcrumb::FBuffer Buffer;
-				const TCHAR* Name = Event.Breadcrumb->Name.GetTCHAR(Buffer);
+				const TCHAR* Name = Event.Breadcrumb->GetTCHAR(Buffer);
 				FGpuProfilerTrace::BeginBreadcrumb(Queue.Value, Event.GPUTimestampTOP, Name);
 
 			#if WITH_PROFILEGPU
@@ -1977,7 +1993,7 @@ namespace UE::RHI::GPUProfiler
 				AccumulateTime(Busy, 0, 0);
 				LastGPUCycles = FMath::Max(LastGPUCycles, Event.GPUTimestampBOP);
 
-				FRHIBreadcrumbData_Stats const& Stat = Event.Breadcrumb->Name.Data;
+				FRHIBreadcrumbData_Stats const& Stat = Event.Breadcrumb->Data;
 				if (Stat.ShouldComputeStat())
 				{
 					// Pop the stat when the refcount hits zero.
@@ -1986,12 +2002,15 @@ namespace UE::RHI::GPUProfiler
 					{
 						ActiveStats.FindAndRemoveChecked(Stat);
 					}
+
+					check(ActiveStatsStack.Last() == Stat);
+					ActiveStatsStack.RemoveAt(ActiveStatsStack.Num() - 1, EAllowShrinking::No);
 				}
 
 				Breadcrumb = Event.Breadcrumb->GetParent();
 
 				FRHIBreadcrumb::FBuffer Buffer;
-				const TCHAR* Name = Event.Breadcrumb->Name.GetTCHAR(Buffer);
+				const TCHAR* Name = Event.Breadcrumb->GetTCHAR(Buffer);
 				FGpuProfilerTrace::EndBreadcrumb(Queue.Value, Event.GPUTimestampBOP, Name);
 
 			#if WITH_PROFILEGPU
@@ -2104,7 +2123,7 @@ namespace UE::RHI::GPUProfiler
 						Recurse(Recurse, Current->GetParent());
 
 						FRHIBreadcrumb::FBuffer Buffer;
-						Profile.PushNode(Current->Name.GetTCHAR(Buffer));
+						Profile.PushNode(Current->GetTCHAR(Buffer));
 					};
 					Recurse(Recurse, Event.Breadcrumb);
 				#endif // WITH_RHI_BREADCRUMBS
diff --git a/Engine/Source/Runtime/RHI/Private/RHIBreadcrumbs.cpp b/Engine/Source/Runtime/RHI/Private/RHIBreadcrumbs.cpp
--- a/Engine/Source/Runtime/RHI/Private/RHIBreadcrumbs.cpp
+++ b/Engine/Source/Runtime/RHI/Private/RHIBreadcrumbs.cpp
@@ -9,17 +9,21 @@
 #if WITH_RHI_BREADCRUMBS
 
 // Constructor for the sentinel node
-FRHIBreadcrumbNode::FRHIBreadcrumbNode()
-	: Name([]() -> FRHIBreadcrumb const&
-	{
-		static TRHIBreadcrumb<> Breadcrumb(FRHIBreadcrumbData(__FILE__, __LINE__), TEXT("Sentinel"));
-		return Breadcrumb;
-	}())
+FRHIBreadcrumbNode::FRHIBreadcrumbNode(FRHIBreadcrumbData const& Data)
+	: Data(Data)
+	, CPUTraceMarkerID(0)
+	, ID(0)
 {}
 
-RHI_API std::atomic<uint32> FRHIBreadcrumb::NextID;
-FRHIBreadcrumbNode FRHIBreadcrumbNode::SentinelNode;
-RHI_API FRHIBreadcrumbNode* const FRHIBreadcrumbNode::Sentinel = &FRHIBreadcrumbNode::SentinelNode;
+RHI_API FRHIBreadcrumbNode* const FRHIBreadcrumbNode::Sentinel = []()
+{
+	using namespace UE::RHI::Breadcrumbs::Private;
+	static TRHIBreadcrumbDesc<0> StaticData(FRHIBreadcrumbData(TEXT("Sentinel"), __FILE__, __LINE__, RHI_GPU_STAT_ARGS_NONE), nullptr);
+	static TRHIBreadcrumb<TRHIBreadcrumbDesc<0>> SentinelNode(StaticData);
+	return &SentinelNode;
+}();
+
+RHI_API std::atomic<uint32> FRHIBreadcrumbNode::NextID;
 
 RHI_API uint32 FRHIBreadcrumbNode::GetLevel(FRHIBreadcrumbNode const* Node)
 {
@@ -70,7 +74,7 @@ RHI_API FString FRHIBreadcrumbNode::GetFullPath() const
 
 		Recurse(Recurse, Current->Parent);
 
-		Result += Current->Name.GetTCHAR(Buffer);
+		Result += Current->GetTCHAR(Buffer);
 		if (Current != this)
 		{
 			Result += TEXT("/");
@@ -91,7 +95,7 @@ RHI_API void FRHIBreadcrumbNode::WriteCrashData(struct FCrashContextExtendedWrit
 	FRHIBreadcrumb::FBuffer Buffer;
 	for (FRHIBreadcrumbNode const* Breadcrumb = this; Breadcrumb && Length < UE_ARRAY_COUNT(String); Breadcrumb = Breadcrumb->Parent)
 	{
-		TCHAR const* Str = Breadcrumb->Name.GetTCHAR(Buffer);
+		TCHAR const* Str = Breadcrumb->GetTCHAR(Buffer);
 		Length += FCString::Snprintf(&String[Length], UE_ARRAY_COUNT(String) - Length, TEXT(" - %s\n"), Str);
 	}
 
@@ -188,14 +192,14 @@ RHI_API void FRHIBreadcrumbState::DumpActiveBreadcrumbs(TMap<FQueueID, TArray<FR
 				}
 
 				// Scan for the MarkerOut. Everything before this marker has been completed by the GPU.
-				if (Node->Name.ID == Data.MarkerOut)
+				if (Node->ID == Data.MarkerOut)
 				{
 					check(ActiveRange.First == SearchRange.First);
 					ActiveRange.First = Node;
 				}
 
 				// Scan for the MarkerIn. Everything after this marker has not been started by the GPU.
-				if (Node->Name.ID == Data.MarkerIn)
+				if (Node->ID == Data.MarkerIn)
 				{
 					check(ActiveRange.Last == SearchRange.Last);
 					ActiveRange.Last = Node;
@@ -264,7 +268,7 @@ RHI_API void FRHIBreadcrumbState::DumpActiveBreadcrumbs(TMap<FQueueID, TArray<FR
 				{
 					Tabs += TEXT("\t");
 				}
-				TCHAR const* Name = Node->Name.GetTCHAR(Buffer);
+				TCHAR const* Name = Node->GetTCHAR(Buffer);
 				EState NodeState = NodeStates.FindChecked(Node);
 
 				TCHAR const* StateStr; 
@@ -276,7 +280,7 @@ RHI_API void FRHIBreadcrumbState::DumpActiveBreadcrumbs(TMap<FQueueID, TArray<FR
 				case EState::Finished  : StateStr = TEXT("   Finished"); break;
 				}
 
-				Tree += FString::Printf(TEXT("\r\n\t\t(ID: 0x%08x) [%s]\t%s%s"), Node->Name.ID, StateStr, *Tabs, Name);
+				Tree += FString::Printf(TEXT("\r\n\t\t(ID: 0x%08x) [%s]\t%s%s"), Node->ID, StateStr, *Tabs, Name);
 
 				while (LastLevel >= Level)
 				{
diff --git a/Engine/Source/Runtime/RHI/Private/RHIValidation.cpp b/Engine/Source/Runtime/RHI/Private/RHIValidation.cpp
--- a/Engine/Source/Runtime/RHI/Private/RHIValidation.cpp
+++ b/Engine/Source/Runtime/RHI/Private/RHIValidation.cpp
@@ -2507,11 +2507,11 @@ namespace RHIValidation
 				Output += TEXT("\t");
 			}
 			FRHIBreadcrumb::FBuffer Buffer;
-			const TCHAR* Str = Node->Name.GetTCHAR(Buffer);
+			const TCHAR* Str = Node->GetTCHAR(Buffer);
 			Output += Str;
 			UE_LOG(LogRHI, Display, TEXT(" ## BC (0x%016p, 0x%08x) [%12s] [%s]: %s")
 				, Node
-				, Node->Name.ID
+				, Node->ID
 				, *GetRHIPipelineName(Pipeline)
 				, bBegin ? TEXT("BEGIN") : TEXT(" END ")
 				, *Output
diff --git a/Engine/Source/Runtime/RHI/Public/GPUProfiler.h b/Engine/Source/Runtime/RHI/Public/GPUProfiler.h
--- a/Engine/Source/Runtime/RHI/Public/GPUProfiler.h
+++ b/Engine/Source/Runtime/RHI/Public/GPUProfiler.h
@@ -518,33 +518,27 @@ namespace UE::RHI::GPUProfiler
 			Idle
 		};
 
-		static FString GetIDString(FQueue Queue, bool bFriendly);
+		TCHAR const* const StatName;
+		TCHAR const* const DisplayName;
+
+	#if CSV_PROFILER_STATS
+		TOptional<FCsvDeclaredStat> CsvStat;
+	#endif
 
 	private:
-		struct FStatCategory_Stats
+	#if STATS
+		static FString GetIDString(FQueue Queue, bool bFriendly);
+		static TCHAR const* GetTypeString(EType Type);
+
+		struct FStatCategory
 		{
-		#if STATS
 			FAnsiString const GroupName;
 			FString     const GroupDesc;
-		#endif
-
-			FStatCategory_Stats(FQueue Queue);
-		};
-
-		struct FStatCategory_Csv
-		{
-		#if CSV_PROFILER_STATS
-			FCsvCategory const Csv;
-		#endif
-
-			FStatCategory_Csv(FQueue Queue);
-		};
 
-		struct FStatCategory
-			: public FStatCategory_Stats
-			, public FStatCategory_Csv
-		{
 			FStatCategory(FQueue Queue);
+
+			static TMap<FQueue, TUniquePtr<FStatCategory>> Categories;
+			static FStatCategory& GetCategory(FQueue Queue);
 		};
 
 		struct FStatInstance
@@ -555,40 +549,32 @@ namespace UE::RHI::GPUProfiler
 				FName StatName;
 				TUniquePtr<FDynamicStat> Stat;
 			#endif
-			#if CSV_PROFILER_STATS
-				TOptional<FCsvDeclaredStat> CsvStat;
-			#endif
 			};
 
 			FInner Busy, Wait, Idle;
 		};
 
-		static TMap<FQueue, TUniquePtr<FStatCategory>> Categories;
-		static FStatCategory& GetCategory(FQueue Queue);
-
 		TMap<FQueue, FStatInstance> Instances;
 
-	public:
-		TCHAR const* const Name;
+		FStatInstance::FInner& GetStatInstance(FQueue Queue, EType Type);
+	#endif
 
-		FGPUStat(TCHAR const* Name)
-			: Name(Name)
+	public:
+		FGPUStat(TCHAR const* StatName, TCHAR const* DisplayName)
+			: StatName   (StatName)
+			, DisplayName(DisplayName)
 		{}
 
 	#if STATS
 		TStatId GetStatId(FQueue Queue, EType Type);
 	#endif
-
-	#if CSV_PROFILER_STATS
-		FCsvDeclaredStat& GetCsvStat(FQueue Queue);
-	#endif
 	};
 
 	template <typename TNameProvider>
 	struct TGPUStat : public FGPUStat
 	{
 		TGPUStat()
-			: FGPUStat(TNameProvider::GetName())
+			: FGPUStat(TNameProvider::GetStatName(), TNameProvider::GetDisplayName())
 		{}
 	};
 
diff --git a/Engine/Source/Runtime/RHI/Public/RHIBreadcrumbs.h b/Engine/Source/Runtime/RHI/Public/RHIBreadcrumbs.h
--- a/Engine/Source/Runtime/RHI/Public/RHIBreadcrumbs.h
+++ b/Engine/Source/Runtime/RHI/Public/RHIBreadcrumbs.h
@@ -13,6 +13,8 @@
 #include "RHIPipeline.h"
 #include "MultiGPU.h"
 
+#include <array>
+
 #ifndef WITH_RHI_BREADCRUMBS
 #define WITH_RHI_BREADCRUMBS ((UE_BUILD_DEBUG || UE_BUILD_DEVELOPMENT) || WITH_PROFILEGPU)
 #endif
@@ -79,7 +81,7 @@ namespace UE::RHI::GPUProfiler
 			return PointerHash(Stats.GPUStat);
 		}
 
-#else
+#elif HAS_GPU_STATS
 
 	#if STATS
 		TStatId StatId {};
@@ -131,6 +133,25 @@ namespace UE::RHI::GPUProfiler
 	#endif
 		}
 
+#else
+
+		FRHIBreadcrumbData_Stats() = default;
+
+		bool ShouldComputeStat() const
+		{
+			return false;
+		}
+
+		bool operator == (FRHIBreadcrumbData_Stats const& RHS) const
+		{
+			return true;
+		}
+
+		friend uint32 GetTypeHash(FRHIBreadcrumbData_Stats const& Stats)
+		{
+			return 0;
+		}
+
 #endif
 	};
 
@@ -143,24 +164,15 @@ namespace UE::RHI::GPUProfiler
 		, public FRHIBreadcrumbData_Stats
 	{
 	public:
+		TCHAR const* const StaticName;
 
-	#if RHI_NEW_GPU_PROFILER && HAS_GPU_STATS
-		FRHIBreadcrumbData(ANSICHAR const* InFile, uint32 InLine, UE::RHI::GPUProfiler::FGPUStat* InGPUStat = nullptr)
-			: FRHIBreadcrumbData_Location(InFile, InLine)
-			, FRHIBreadcrumbData_Stats(InGPUStat)
-		{}
-
-	#else
-
-		FRHIBreadcrumbData(ANSICHAR const* File, uint32 Line, TStatId StatId = TStatId(), FName CsvStat = NAME_None)
+		FRHIBreadcrumbData(TCHAR const* StaticName, ANSICHAR const* File, uint32 Line, FRHIBreadcrumbData_Stats&& Stats)
 			: FRHIBreadcrumbData_Location(File, Line)
-			, FRHIBreadcrumbData_Stats(StatId, CsvStat)
+			, FRHIBreadcrumbData_Stats(MoveTemp(Stats))
+			, StaticName(StaticName)
 		{}
-
-	#endif
 	};
 
-	class FRHIBreadcrumb;
 	class FRHIBreadcrumbAllocator;
 	struct FRHIBreadcrumbRange;
 
@@ -204,7 +216,7 @@ namespace UE::RHI::GPUProfiler
 
 	struct FRHIBreadcrumbNode
 	{
-		FRHIBreadcrumb const& Name;
+		RHI_API static std::atomic<uint32> NextID;
 
 	private:
 		FRHIBreadcrumbNode* Parent = Sentinel;
@@ -214,25 +226,26 @@ namespace UE::RHI::GPUProfiler
 	public:
 		FRHIBreadcrumbAllocator* const Allocator = nullptr;
 
-#if DO_CHECK
+		FRHIBreadcrumbData const& Data;
+		uint32 const CPUTraceMarkerID = 0;
+		uint32 const ID;
+
+	#if DO_CHECK
 		// Used to track use of this breadcrumb on each GPU pipeline. Breadcrumbs can only be begun/ended once per pipe.
 		std::atomic<std::underlying_type_t<ERHIPipeline>> BeginPipes = std::underlying_type_t<ERHIPipeline>(ERHIPipeline::None);
-		std::atomic<std::underlying_type_t<ERHIPipeline>> EndPipes = std::underlying_type_t<ERHIPipeline>(ERHIPipeline::None);
-#endif
+		std::atomic<std::underlying_type_t<ERHIPipeline>> EndPipes   = std::underlying_type_t<ERHIPipeline>(ERHIPipeline::None);
+	#endif
 
-		FRHIBreadcrumbNode(FRHIBreadcrumbAllocator& Allocator, FRHIBreadcrumb const& Name)
-			: Name(Name)
-			, Allocator(&Allocator)
+		FRHIBreadcrumbNode(FRHIBreadcrumbData const& Data, FRHIBreadcrumbAllocator& Allocator)
+			: Allocator(&Allocator)
+			, Data(Data)
+			, ID(NextID.fetch_add(1, std::memory_order_relaxed) | 0x80000000) // Set the top bit to avoid collision with zero (i.e. "no breadcrumb")
 		{}
 
-		FRHIBreadcrumbNode* GetParent() const
-		{
-			return Parent;
-		}
-
 		FRHIBreadcrumbNode*      & GetNextPtr(ERHIPipeline Pipeline)       { return NextPtrs[GetRHIPipelineIndex(Pipeline)]; }
 		FRHIBreadcrumbNode* const& GetNextPtr(ERHIPipeline Pipeline) const { return NextPtrs[GetRHIPipelineIndex(Pipeline)]; }
 
+		FRHIBreadcrumbNode* GetParent() const { return Parent; }
 		inline void SetParent(FRHIBreadcrumbNode* Node);
 
 		inline void BeginCPU() const;
@@ -262,14 +275,49 @@ namespace UE::RHI::GPUProfiler
 		// that haven't been attached to the root yet, specifically to be distinct from nullptr which is the root.
 		static RHI_API FRHIBreadcrumbNode* const Sentinel;
 
-	private:
-		// Constructor for the sentinel value
-		FRHIBreadcrumbNode();
-		static FRHIBreadcrumbNode SentinelNode;
+		// The maximum length of a breadcrumb string, including the null terminator.
+		static constexpr uint32 MaxLength = 128;
+
+		struct FBuffer
+		{
+			TCHAR Data[MaxLength];
+		};
 
+		virtual TCHAR const* GetTCHAR(FBuffer& Buffer) const = 0;
+
+		TCHAR const* GetTCHARNoFormat() const
+		{
+			return Data.StaticName;
+		}
+
+	protected:
+		// Constructor for the sentinel value
+		FRHIBreadcrumbNode(FRHIBreadcrumbData const& Data);
 		friend struct FRHIBreadcrumbList;
+
+		void CreateTraceMarkers()
+		{
+		#if RHI_BREADCRUMBS_EMIT_CPU
+			if (TRACE_CPUPROFILER_EVENT_MANUAL_IS_ENABLED())
+			{
+				FRHIBreadcrumbNode::FBuffer Buffer;
+				TCHAR const* Str = GetTCHAR(Buffer);
+				const_cast<uint32&>(CPUTraceMarkerID) = FCpuProfilerTrace::OutputDynamicEventType(Str
+			#if RHI_BREADCRUMBS_EMIT_LOCATION
+					, Data.File, Data.Line
+			#endif
+				);
+			}
+		#endif
+		}
 	};
 
+	// Typedef for backwards compatibility. The FRHIBreadcrumb and FRHIBreadcrumbNode have been merged into one type.
+	using FRHIBreadcrumb = FRHIBreadcrumbNode;
+
+	template <typename TDesc, typename... TValues>
+	using TRHIBreadcrumbInitializer = std::tuple<TDesc const*, std::tuple<TValues...>>;
+
 	class FRHIBreadcrumbAllocatorArray : public TArray<TSharedRef<class FRHIBreadcrumbAllocator>, TInlineAllocator<2>>
 	{
 	public:
@@ -278,34 +326,33 @@ namespace UE::RHI::GPUProfiler
 
 	class FRHIBreadcrumbAllocator : public TSharedFromThis<FRHIBreadcrumbAllocator>
 	{
-		template <typename...>
-		friend class TRHIBreadcrumb;
 		friend FRHIBreadcrumbNode;
 
 		FMemStackBase Inner;
 		FRHIBreadcrumbAllocatorArray Parents;
 
+	public:
+		FRHIBreadcrumbAllocatorArray const& GetParents() const { return Parents; }
+
 		template <typename TType, typename... TArgs>
 		TType* Alloc(TArgs&&... Args)
 		{
+			static_assert(std::is_trivially_destructible<TType>::value, "Only trivially destructable types may be used with the RHI breadcrumb allocator.");
 			return new (Inner.Alloc(sizeof(TType), alignof(TType))) TType(Forward<TArgs>(Args)...);
 		}
 
-	public:
-		FRHIBreadcrumbAllocatorArray const& GetParents() const { return Parents; }
-
-		template<size_t N, typename... TArgs>
-		inline FRHIBreadcrumbNode* AllocBreadcrumb(FRHIBreadcrumbData&& Stat, TCHAR const(&FormatString)[N], TArgs&&... Args);
-
-		inline void* Alloc(uint32 Size, uint32 Align)
+		void* Alloc(uint32 Size, uint32 Align)
 		{
 			return Inner.Alloc(Size, Align);
 		}
 
-#if ENABLE_RHI_VALIDATION
+		template <typename TDesc, typename... TValues>
+		inline FRHIBreadcrumbNode* AllocBreadcrumb(TRHIBreadcrumbInitializer<TDesc, TValues...> const& Args);
+
+	#if ENABLE_RHI_VALIDATION
 		// Used by RHI validation for circular reference detection.
 		bool bVisited = false;
-#endif
+	#endif
 	};
 
 	inline void FRHIBreadcrumbAllocatorArray::AddUnique(FRHIBreadcrumbAllocator* Allocator)
@@ -566,55 +613,7 @@ namespace UE::RHI::GPUProfiler
 			return HashCombineFast(GetTypeHash(Range.First), GetTypeHash(Range.Last));
 		}
 	};
-
-	class FRHIBreadcrumb
-	{
-		RHI_API static std::atomic<uint32> NextID;
-
-	protected:
-		FRHIBreadcrumb(FRHIBreadcrumbData&& Data)
-			: ID(NextID.fetch_add(1, std::memory_order_relaxed) | 0x80000000) // Set the top bit to avoid collision with zero (i.e. "no breadcrumb")
-			, Data(Data)
-		{}
-
-		FRHIBreadcrumb(FRHIBreadcrumb const&) = delete;
-		FRHIBreadcrumb(FRHIBreadcrumb&&) = delete;
-
-	public:
-		uint32 const ID;
-		FRHIBreadcrumbData const Data;
-
-#if RHI_BREADCRUMBS_EMIT_CPU
-		uint32 const CPUTraceMarkerID = 0;
-#endif
-
-		void CreateTraceMarkers()
-		{
-#if RHI_BREADCRUMBS_EMIT_CPU
-			if (TRACE_CPUPROFILER_EVENT_MANUAL_IS_ENABLED())
-			{
-				FBuffer Buffer;
-				TCHAR const* Str = GetTCHAR(Buffer);
-				const_cast<uint32&>(CPUTraceMarkerID) = FCpuProfilerTrace::OutputDynamicEventType(Str
-			#if RHI_BREADCRUMBS_EMIT_LOCATION
-					, Data.File, Data.Line
-			#endif
-				);
-			}
-#endif
-		}
-
-		// The maximum length of a breadcrumb string, including the null terminator.
-		static constexpr uint32 MaxLength = 128;
-
-		struct FBuffer
-		{
-			TCHAR Data[MaxLength];
-		};
-		virtual TCHAR const* GetTCHAR(FBuffer& Storage) const = 0;
-		virtual TCHAR const* GetTCHARNoFormat() const = 0;
-	};
-
+	
 	inline void FRHIBreadcrumbNode::SetParent(FRHIBreadcrumbNode* Node)
 	{
 		check(Parent == nullptr || Parent == FRHIBreadcrumbNode::Sentinel);
@@ -628,27 +627,27 @@ namespace UE::RHI::GPUProfiler
 
 	inline void FRHIBreadcrumbNode::BeginCPU() const
 	{
-#if RHI_BREADCRUMBS_EMIT_CPU
-		if (Name.CPUTraceMarkerID)
+	#if RHI_BREADCRUMBS_EMIT_CPU
+		if (CPUTraceMarkerID)
 		{
-			FCpuProfilerTrace::OutputBeginEvent(Name.CPUTraceMarkerID);
+			FCpuProfilerTrace::OutputBeginEvent(CPUTraceMarkerID);
 		}
-#endif
+	#endif
 	}
 
 	inline void FRHIBreadcrumbNode::EndCPU() const
 	{
-#if RHI_BREADCRUMBS_EMIT_CPU
-		if (Name.CPUTraceMarkerID)
+	#if RHI_BREADCRUMBS_EMIT_CPU
+		if (CPUTraceMarkerID)
 		{
 			FCpuProfilerTrace::OutputEndEvent();
 		}
-#endif
+	#endif
 	}
 
 	inline void FRHIBreadcrumbNode::WalkIn(FRHIBreadcrumbNode* Node)
 	{
-#if RHI_BREADCRUMBS_EMIT_CPU
+	#if RHI_BREADCRUMBS_EMIT_CPU
 		if (TRACE_CPUPROFILER_EVENT_MANUAL_IS_ENABLED())
 		{
 			auto Recurse = [](FRHIBreadcrumbNode* Current, auto& Recurse) -> void
@@ -661,12 +660,12 @@ namespace UE::RHI::GPUProfiler
 			};
 			Recurse(Node, Recurse);
 		}
-#endif
+	#endif
 	}
 
 	inline void FRHIBreadcrumbNode::WalkOut(FRHIBreadcrumbNode* Node)
 	{
-#if RHI_BREADCRUMBS_EMIT_CPU
+	#if RHI_BREADCRUMBS_EMIT_CPU
 		if (TRACE_CPUPROFILER_EVENT_MANUAL_IS_ENABLED())
 		{
 			while (Node && Node != Sentinel)
@@ -675,202 +674,385 @@ namespace UE::RHI::GPUProfiler
 				Node = Node->GetParent();
 			}
 		}
-#endif
+	#endif
 	}
 
-	// Generic copy-by-value
-	template <typename T>
-	struct TRHIBreadcrumbValue
+	namespace UE::RHI::Breadcrumbs::Private
 	{
-		static_assert(!(std::is_same_v<T, TCHAR*> || std::is_same_v<T, TCHAR const*>),
-			"Do not use raw TCHAR pointers with RHI breadcrumbs. Pass the FString, FName, or string literal instead. If you are certain your TCHAR pointer is a string literal "
-			"(e.g. from a function returning a literal) and you cannot pass that literal directly to an RHI breadcrumb, use the RHI_BREADCRUMB_FORCE_STRING_LITERAL macro to silence "
-			"this static assert. Incorrect use of RHI_BREADCRUMB_FORCE_STRING_LITERAL will lead to use-after-free, as only the raw pointer is retained by the breadcrumb."
-		);
-
-		T Value;
-		TRHIBreadcrumbValue(T const& Value)
-			: Value(Value)
-		{}
+		// Replacement for std::remove_cvref, since this isn't available until C++20.
+		template <typename T>
+		using TRemoveCVRef = std::remove_cv_t<std::remove_reference_t<T>>;
+
+		// Used to control the static_asserts below.
+		template <typename...>
+		struct TFalse { static constexpr bool Value = false; };
 
-		struct FConvert
+		//
+		// The TValue types are used to capture and store vararg format values in RHI breadcrumbs. Capturing values like this means we can avoid the string
+		// formatting cost until we actually need the final string (e.g. we're emitting breadcrumbs to an external profiler, we're running 'profilegpu' etc.).
+		// 
+		// The inner FConvert type is used to prepare the TValue when calling FCString::Snprintf to generate the final breadcrumb string.
+		// 
+		// This base definition catches all integer, float and enum values.
+		//
+		template <typename T>
+		struct TValue
 		{
-			T const& Inner;
-			FConvert(TRHIBreadcrumbValue const& Value)
-				: Inner(Value.Value)
-			{}
-		};
-	};
+			static constexpr bool bValidType = std::is_integral_v<T> || std::is_floating_point_v<T> || std::is_enum_v<T>;
 
-	// String literal - keep the string pointer
-	template <size_t N>
-	struct TRHIBreadcrumbValue<TCHAR[N]>
-	{
-		TCHAR const* Value;
+			T const Value;
 
-		TRHIBreadcrumbValue(TCHAR const* Value)
-			: Value(Value)
-		{}
+			template <typename TArg>
+			TValue(TArg const& Value)
+				: Value(Value)
+			{
+				static_assert(bValidType || TFalse<TArg>::Value, "Type is not compatible with RHI breadcrumbs.");
+			}
 
-		struct FConvert
-		{
-			TCHAR const* Inner;
-			FConvert(TRHIBreadcrumbValue const& Value)
-				: Inner(Value.Value)
-			{}
+			struct FConvert
+			{
+				T const& Inner;
+				FConvert(TValue const& Value)
+					: Inner(Value.Value)
+				{}
+			};
 		};
-	};
 
-	// FName - keep the FName itself and defer resolving
-	template <>
-	struct TRHIBreadcrumbValue<FName>
-	{
-		FName Value;
-		TRHIBreadcrumbValue(FName const& Value)
-			: Value(Value)
-		{}
-
-		struct FConvert
+		// Disallow all pointer types. Provide a specific assert message for TCHAR pointers.
+		template <typename T>
+		struct TValue<T*>
 		{
-			TCHAR Inner[FRHIBreadcrumb::MaxLength];
-			FConvert(TRHIBreadcrumbValue const& Name)
+			static constexpr bool bValidType = false;
+
+			template <typename... TArgs>
+			TValue(TArgs&&...)
 			{
-				Name.Value.ToStringTruncate(Inner);
+				if constexpr (std::is_same_v<std::remove_const_t<T>, TCHAR>)
+				{
+					static_assert(TFalse<TArgs...>::Value,
+						"Do not use raw TCHAR pointers with RHI breadcrumbs. Pass the FString, FName, or string literal instead. If you are certain your TCHAR pointer is a string literal "
+						"(e.g. from a function returning a literal) and you cannot pass that literal directly to an RHI breadcrumb, use the RHI_BREADCRUMB_FORCE_STRING_LITERAL macro to silence "
+						"this static assert. Incorrect use of RHI_BREADCRUMB_FORCE_STRING_LITERAL will lead to use-after-free, as only the raw pointer is retained by the breadcrumb."
+					);
+				}
+				else
+				{			
+					static_assert(TFalse<TArgs...>::Value, "RHI breadcrumbs do not support arbitrary pointer types.");
+				}
 			}
-		};
-	};
 
-	// FDebugName - keep the FDebugName itself and defer resolving
-	template <>
-	struct TRHIBreadcrumbValue<FDebugName>
-	{
-		FDebugName Value;
-		TRHIBreadcrumbValue(FDebugName const& Value)
-			: Value(Value)
-		{}
+			struct FConvert
+			{
+				uint32 Inner;
+				template <typename... TArgs>
+				FConvert(TArgs&&...)
+					: Inner(0)
+				{}
+			};
+		};
 
-		struct FConvert
+		// String literal - keep the string pointer
+		template <size_t N>
+		struct TValue<TCHAR[N]>
 		{
-			TCHAR Inner[FRHIBreadcrumb::MaxLength];
-			FConvert(TRHIBreadcrumbValue const& Name)
+			static constexpr bool bValidType = true;
+
+			TCHAR const(&Value)[N];
+
+			TValue(TCHAR const(&Value)[N])
+				: Value(Value)
+			{}
+
+			struct FConvert
 			{
-				Name.Value.ToStringTruncate(Inner);
-			}
+				TCHAR const(&Inner)[N];
+				FConvert(TValue const& Value)
+					: Inner(Value.Value)
+				{}
+			};
 		};
-	};
 
-	// FString - Take an immediate copy of the string. Total length is limited by fixed buffer size.
-	template <>
-	struct TRHIBreadcrumbValue<FString>
-	{
-		TCHAR Buffer[FRHIBreadcrumb::MaxLength];
-		TRHIBreadcrumbValue(FString const& Value)
+		// FName - keep the FName itself and defer resolving
+		template <>
+		struct TValue<FName>
 		{
-			FCString::Strncpy(Buffer, *Value, UE_ARRAY_COUNT(Buffer));
-		}
+			static constexpr bool bValidType = true;
 
-		struct FConvert
-		{
-			TCHAR const* Inner;
-			FConvert(TRHIBreadcrumbValue const& String)
-				: Inner(String.Buffer)
+			FName Value;
+			TValue(FName const& Value)
+				: Value(Value)
 			{}
+
+			struct FConvert
+			{
+				TCHAR Inner[FRHIBreadcrumb::MaxLength];
+				FConvert(TValue const& Name)
+				{
+					Name.Value.ToStringTruncate(Inner);
+				}
+			};
 		};
-	};
 
-	// Breadcrumb implementation for printf formatted names
-	template <typename... TArgs>
-	class TRHIBreadcrumb final : public FRHIBreadcrumb
-	{
-		friend FRHIBreadcrumbAllocator;
+		// FDebugName - keep the FDebugName itself and defer resolving
+		template <>
+		struct TValue<FDebugName>
+		{
+			static constexpr bool bValidType = true;
 
-		TCHAR const* FormatString;
-		std::tuple<TRHIBreadcrumbValue<TArgs>...> Values;
+			FDebugName Value;
+			TValue(FDebugName const& Value)
+				: Value(Value)
+			{}
 
-		template <size_t... Indices>
-		void ToStringImpl(FBuffer& Buffer, std::index_sequence<Indices...>) const
+			struct FConvert
+			{
+				TCHAR Inner[FRHIBreadcrumb::MaxLength];
+				FConvert(TValue const& Name)
+				{
+					Name.Value.ToStringTruncate(Inner);
+				}
+			};
+		};
+
+		// FString - Take an immediate copy of the string. Total length is limited by fixed buffer size.
+		template <>
+		struct TValue<FString>
 		{
-			// Perform type conversions (call ToString() on FName etc)
-			std::tuple<typename TRHIBreadcrumbValue<TArgs>::FConvert...> Converted
+			static constexpr bool bValidType = true;
+
+			TCHAR Buffer[FRHIBreadcrumb::MaxLength];
+			TValue(FString const& Value)
+			{
+				FCString::Strncpy(Buffer, *Value, UE_ARRAY_COUNT(Buffer));
+			}
+
+			struct FConvert
 			{
-				std::get<Indices>(Values)...
+				TCHAR const* Inner;
+				FConvert(TValue const& String)
+					: Inner(String.Buffer)
+				{}
 			};
+		};
 
-			FCString::Snprintf(
-				  Buffer.Data
-				, UE_ARRAY_COUNT(Buffer.Data)
-				, reinterpret_cast<TCHAR const(&)[1]>(*FormatString) // Cast hack to workaround static_assert in Snprintf checking for a string literal. 
-				, (std::get<Indices>(Converted).Inner)...
-			);
-		}
+		// Determines if a vararg type is a string literal / value tuple, as defined by RHI_BREADCRUMB_FIELD.
+		template <          typename T> struct TIsField                                   { static constexpr bool Value = false; };
+		template <size_t N, typename T> struct TIsField<std::tuple<TCHAR const(&)[N], T>> { static constexpr bool Value = true;  };
 
-		TRHIBreadcrumb(TRHIBreadcrumb const& Other) = delete;
+		// Determines the TValue type used to store a vararg value in a breadcrumb. Also unpacks the value type from RHI_BREADCRUMB_FIELD tuples.
+		template <          typename T> struct TGetValueType                                   { using TType = TValue<TRemoveCVRef<T>>; };
+		template <size_t N, typename T> struct TGetValueType<std::tuple<TCHAR const(&)[N], T>> { using TType = TValue<TRemoveCVRef<T>>; };
 
-	public:
-		TRHIBreadcrumb(FRHIBreadcrumbData&& Data, TCHAR const* FormatString, TArgs const&... Args)
-			: FRHIBreadcrumb(MoveTemp(Data))
-			, FormatString(FormatString)
-			, Values(Args...)
-		{
-			CreateTraceMarkers();
-		}
+		// Helper to concatenate index sequences.
+		template <size_t  , typename...  > struct TConcatIndexSequence;
+		template <size_t N, size_t... Seq> struct TConcatIndexSequence<N, std::index_sequence<Seq...>> { using TType = std::index_sequence<N, Seq...>; };
 
-		virtual TCHAR const* GetTCHAR(FBuffer& Buffer) const override
+		// Generates a std::index_sequence containing the indices of the RHI_BREADCRUMB_FIELD tuples in a parameter pack.
+		template <size_t N, typename...                       > struct TFindFieldIndices;
+		template <size_t N                                    > struct TFindFieldIndices<N                  > { using TType = std::index_sequence<>; };
+		template <size_t N, typename TFirst, typename... TRest> struct TFindFieldIndices<N, TFirst, TRest...>
 		{
-			ToStringImpl(Buffer, std::make_index_sequence<sizeof...(TArgs)>{});
-			return Buffer.Data;
-		}
+			using TType = typename std::conditional<TIsField<TFirst>::Value
+				, typename TConcatIndexSequence<N, typename TFindFieldIndices<N + 1, TRest...>::TType>::TType
+				, typename TFindFieldIndices<N + 1, TRest...>::TType
+			>::type;
+		};
+
+		struct FForceNoSprintf {};
 
-		virtual TCHAR const* GetTCHARNoFormat() const override
+		// Infers the size of a string literal character array. Size is zero for any type that is not a string literal.
+		template <typename   > struct TStringLiteralSize                       { static constexpr size_t Value = 0;    };
+		template <size_t Size> struct TStringLiteralSize<TCHAR const(&)[Size]> { static constexpr size_t Value = Size; };
+
+		//
+		// Helper traits for dealing wth varargs values. The std::tuple specializations are for the RHI_BREADCRUMB_FIELD values.
+		//
+		//		- TDescType    : Type passed to the TRHIBreadcrumbDesc template.
+		//		- TValueType   : The under-lying non-reference type of the vararg. Fields are unpacked.
+		//		- TValueRef    : A reference to the TValueType type.
+		//		- ForwardValue : Similar to std::forward, but also unpacks and forwards the value from a std::tuple RHI_BREADCRUMB_FIELD.
+		//
+		template <          class T> struct TFieldTraits                                       { using TDescType = TRemoveCVRef<T>;                                using TValueType = T; using TValueRef = T  ; static TValueRef ForwardValue(auto&& Arg) { return std::forward<TValueRef>(Arg); } };
+		template <          class T> struct TFieldTraits<T& >                                  { using TDescType = TRemoveCVRef<T>;                                using TValueType = T; using TValueRef = T& ; static TValueRef ForwardValue(auto&& Arg) { return std::forward<TValueRef>(Arg); } };
+		template <          class T> struct TFieldTraits<T&&>                                  { using TDescType = TRemoveCVRef<T>;                                using TValueType = T; using TValueRef = T&&; static TValueRef ForwardValue(auto&& Arg) { return std::forward<TValueRef>(Arg); } };
+		template <size_t N, class T> struct TFieldTraits<std::tuple<TCHAR const(&)[N], T  >&&> { using TDescType = std::tuple<TCHAR const(&)[N], TRemoveCVRef<T>>; using TValueType = T; using TValueRef = T  ; static TValueRef ForwardValue(auto&& Arg) { return std::forward<TValueRef>(std::get<1>(Arg)); } };
+		template <size_t N, class T> struct TFieldTraits<std::tuple<TCHAR const(&)[N], T& >&&> { using TDescType = std::tuple<TCHAR const(&)[N], TRemoveCVRef<T>>; using TValueType = T; using TValueRef = T& ; static TValueRef ForwardValue(auto&& Arg) { return std::forward<TValueRef>(std::get<1>(Arg)); } };
+		template <size_t N, class T> struct TFieldTraits<std::tuple<TCHAR const(&)[N], T&&>&&> { using TDescType = std::tuple<TCHAR const(&)[N], TRemoveCVRef<T>>; using TValueType = T; using TValueRef = T&&; static TValueRef ForwardValue(auto&& Arg) { return std::forward<TValueRef>(std::get<1>(Arg)); } };
+
+		// The breadcrumb macros work with both RHI command lists and RHI contexts. This helper retrieves the relevant RHI command list from both types.
+		static inline FRHIComputeCommandList& GetRHICmdList(FRHIComputeCommandList& RHICmdList);
+		static inline FRHIComputeCommandList& GetRHICmdList(IRHIComputeContext    & RHIContext);
+
+		// Returns the full path string for the breadcrumb currently at the top of the CPU stack,
+		// for either RHI command lists or RHI contexts. Used by the breadcrumb check macros.
+		static inline FString GetSafeBreadcrumbPath(auto&& RHICmdList_Or_RHIContext);
+
+		template <size_t Size, bool bHasVarargs>
+		struct TFormatString
 		{
-			return FormatString;
-		}
-	};
+			TCHAR const(&FormatString)[Size];
 
-	// Breadcrumb implementation for string literals
-	template <>
-	class TRHIBreadcrumb<> final : public FRHIBreadcrumb
-	{
-		friend FRHIBreadcrumbAllocator;
+			TFormatString(TCHAR const(&FormatString)[Size])
+				: FormatString(FormatString)
+			{}
 
-		TCHAR const* StringLiteral;
+		protected:
+			template <typename TValuesTuple, size_t... Indices>
+			TCHAR const* ToStringImpl(TCHAR const* StaticName, FRHIBreadcrumbNode::FBuffer& Buffer, TValuesTuple const& Values, std::index_sequence<Indices...>) const
+			{
+				// Perform type conversions (call ToString() on FName etc)
+				std::tuple<typename std::tuple_element_t<Indices, TValuesTuple>::FConvert...> Converted
+				{
+					std::get<Indices>(Values)...
+				};
+
+				FCString::Snprintf(
+					Buffer.Data
+					, UE_ARRAY_COUNT(Buffer.Data)
+					, FormatString
+					, (std::get<Indices>(Converted).Inner)...
+				);
 
-		TRHIBreadcrumb(TRHIBreadcrumb const& Other) = delete;
+				return Buffer.Data;
+			}
+		};
 
-	public:
-		TRHIBreadcrumb(FRHIBreadcrumbData&& Data, TCHAR const* StringLiteral)
-			: FRHIBreadcrumb(MoveTemp(Data))
-			, StringLiteral(StringLiteral)
+		template <>
+		struct TFormatString<0, true>
 		{
-			CreateTraceMarkers();
-		}
+			UE_DEPRECATED(5.6,
+				"Use of the non-\"_F\" versions of the RHI_BREADCRUMB_EVENT family of macros with printf varargs has been deprecated. "
+				"The additional values passed to these macros will be ignored, and the raw printf format string will form the name of the breadcrumb. "
+				"Use the \"_F\" versions of these macros instead (which require both a static name and a format string), or remove the varargs.")
+			TFormatString(std::nullptr_t)
+			{}
 
-		virtual TCHAR const* GetTCHAR(FBuffer&) const override
-		{
-			return StringLiteral;
-		}
+			TFormatString(FForceNoSprintf&&)
+			{}
 
-		virtual TCHAR const* GetTCHARNoFormat() const override
+			template <typename TValuesTuple, size_t... Indices>
+			TCHAR const* ToStringImpl(TCHAR const* StaticName, FRHIBreadcrumbNode::FBuffer& Buffer, TValuesTuple const& Values, std::index_sequence<Indices...>) const
+			{
+				if constexpr (std::tuple_size_v<TValuesTuple> == 1)
+				{
+					using TConvert = typename std::tuple_element_t<0, TValuesTuple>::FConvert;
+
+					if constexpr (std::is_same_v<decltype(std::declval<TConvert>().Inner), TCHAR const*>)
+					{
+						// The breadcrumb has no format string, and a single vararg which looks like a null terminated TCHAR const* string pointer.
+						// Just return the pointer directly from the value.
+						return TConvert(std::get<0>(Values)).Inner;
+					}
+					else
+					{
+						return StaticName;
+					}
+				}
+				else
+				{
+					return StaticName;
+				}
+			}
+		};
+
+		template <>
+		struct TFormatString<0, false>
 		{
-			return StringLiteral;
-		}
-	};
+			TFormatString(std::nullptr_t)
+			{}
 
-	template<size_t N, typename... TArgs>
-	inline FRHIBreadcrumbNode* FRHIBreadcrumbAllocator::AllocBreadcrumb(FRHIBreadcrumbData&& Data, TCHAR const(&FormatString)[N], TArgs&&... Args)
-	{
-		struct FStorage
+			template <typename TValuesTuple, size_t... Indices>
+			TCHAR const* ToStringImpl(TCHAR const* StaticName, FRHIBreadcrumbNode::FBuffer& Buffer, TValuesTuple const& Values, std::index_sequence<Indices...>) const
+			{
+				return StaticName;
+			}
+		};
+
+		// Returns true if the given args are compatible with RHI breadcrumbs, i.e. they have matching TValue specializations.
+		template <typename... TArgs> struct TIsValidArgs   { static constexpr bool Value = (TGetValueType<typename TFieldTraits<TArgs>::TDescType>::TType::bValidType && ...); };
+		template <                 > struct TIsValidArgs<> { static constexpr bool Value = true; };
+
+		//
+		// Contains the definition of an RHI breadcrumb (i.e. the value types, number and name of fields, etc).
+		// These are instantiated as static objects in RHI_BREADCRUMB_PRIVATE_DEFINE, which is used by the breadcrumb macros.
+		//
+		template <size_t FormatStringSize, typename... TValues>
+		class TRHIBreadcrumbDesc final : public FRHIBreadcrumbData, public TFormatString<FormatStringSize, sizeof...(TValues) != 0>
 		{
-			TRHIBreadcrumb<std::decay_t<TArgs>...> Name;
-			FRHIBreadcrumbNode Node;
+			using TFormatString = TFormatString<FormatStringSize, sizeof...(TValues) != 0>;
 
-			FStorage(FRHIBreadcrumbAllocator& Allocator, FRHIBreadcrumbData&& Data, TCHAR const(&FormatString)[N], TArgs&&... Args)
-				: Name(MoveTemp(Data), FormatString, Forward<TArgs>(Args)...)
-				, Node(Allocator, Name)
+			template <typename TTuple, size_t... Indices>
+			static auto ExtractFieldNames(TTuple&& Tuple, std::index_sequence<Indices...>)
+			{
+				return std::array<TCHAR const*, sizeof...(Indices)>
+				{
+					std::get<0>(std::get<Indices>(Tuple))...
+				};
+			}
+
+		public:
+			// Tuple of vararg value types to store in the breadcrumb. Field tuples have been unpacked and the name component discarded.
+			using TValuesTuple = std::tuple<typename TGetValueType<TValues>::TType...>;
+			static constexpr size_t NumValues = sizeof...(TValues);
+
+			// std::index_sequence for all vararg values.
+			using TValuesIndexSequence = std::index_sequence_for<TValues...>;
+
+			// std::index_sequence defining the indices of the tuple field values
+			using TFieldsIndexSequence = typename TFindFieldIndices<0, TValues...>::TType;
+			static constexpr size_t NumFields = TFieldsIndexSequence::size();
+
+			// Array of the string literal field names
+			std::array<TCHAR const*, NumFields> const FieldNames;
+
+			template <typename... TInnerValues>
+			TRHIBreadcrumbDesc(FRHIBreadcrumbData&& BaseData, TFormatString FormatString, TInnerValues&&... Values)
+				: FRHIBreadcrumbData(MoveTemp(BaseData))
+				, TFormatString(FormatString)
+				, FieldNames(ExtractFieldNames(std::forward_as_tuple(Values...), TFieldsIndexSequence()))
 			{}
+
+			TCHAR const* ToString(FRHIBreadcrumbNode::FBuffer& Buffer, TValuesTuple const& Values) const
+			{
+				return TFormatString::ToStringImpl(StaticName, Buffer, Values, TValuesIndexSequence());
+			}
 		};
 
-		return &Alloc<FStorage>(*this, MoveTemp(Data), FormatString, Forward<TArgs>(Args)...)->Node;
+		// Breadcrumb implementation for printf formatted names
+		// Privately inherit from the TValuesTuple to use empty-base-class optimization for breadcrumbs with no varargs.
+		template <typename TDesc>
+		class TRHIBreadcrumb final : public FRHIBreadcrumbNode, private TDesc::TValuesTuple
+		{
+			friend FRHIBreadcrumbAllocator;
+			friend FRHIBreadcrumbNode;
+
+			TDesc const& Desc;
+
+			TRHIBreadcrumb(TRHIBreadcrumb const&) = delete;
+			TRHIBreadcrumb(TRHIBreadcrumb&&     ) = delete;
+
+		public:
+			template <typename... TValues>
+			TRHIBreadcrumb(FRHIBreadcrumbAllocator& Allocator, TDesc const& Desc, TValues&&... Values)
+				: FRHIBreadcrumbNode(Desc, Allocator)
+				, TDesc::TValuesTuple(Forward<TValues>(Values)...)
+				, Desc(Desc)
+			{
+				CreateTraceMarkers();
+			}
+
+			TCHAR const* GetTCHAR(FBuffer& Buffer) const override
+			{
+				return Desc.ToString(Buffer, *this);
+			}
+
+		private:
+			// Constructor for the sentinel value
+			TRHIBreadcrumb(TDesc const& Desc)
+				: FRHIBreadcrumbNode(Desc)
+				, Desc(Desc)
+			{}
+		};
 	}
 
 	class FRHIBreadcrumbNodeRef
@@ -897,6 +1079,20 @@ namespace UE::RHI::GPUProfiler
 		FRHIBreadcrumbNode* Get() const { return Node; }
 	};
 
+	struct FRHIBreadcrumbScope
+	{
+		FRHIComputeCommandList& RHICmdList;
+		FRHIBreadcrumbNode* const Node;
+
+	private:
+		FRHIBreadcrumbScope(FRHIComputeCommandList& RHICmdList, FRHIBreadcrumbNode* Node);
+
+	public:
+		template <typename TDesc, typename... TValues>
+		inline FRHIBreadcrumbScope(FRHIComputeCommandList& RHICmdList, TRHIBreadcrumbInitializer<TDesc, TValues...>&& Args);
+		inline ~FRHIBreadcrumbScope();
+	};
+
 	//
 	// A helper class to manually create, begin and end a breadcrumb on a given RHI command list.
 	// For use in places where the Begin/End operations are separate, and a scoped breadcrumb event is not appropriate.
@@ -911,125 +1107,127 @@ namespace UE::RHI::GPUProfiler
 		uint32 ThreadId;
 	#endif
 
+		inline FRHIBreadcrumbEventManual(FRHIComputeCommandList& RHICmdList, FRHIBreadcrumbNode* Node);
+
 	public:
-		template<size_t N, typename... TArgs>
-		inline FRHIBreadcrumbEventManual(FRHIComputeCommandList& RHICmdList, FRHIBreadcrumbData&& Data, TCHAR const(&FormatString)[N], TArgs&&... Args);
+		template <typename TDesc, typename... TValues>
+		inline FRHIBreadcrumbEventManual(FRHIComputeCommandList& RHICmdList, TRHIBreadcrumbInitializer<TDesc, TValues...>&& Args);
 
 		inline void End(FRHIComputeCommandList& RHICmdList);
 
 		inline ~FRHIBreadcrumbEventManual();
 	};
 
-	//
-	// A helper class for scoped breadcrumb. Used by the RHI_BREADCRUMB_EVENT macros.
-	//
-	class FRHIBreadcrumbEventScope
-	{
-		FRHIComputeCommandList& RHICmdList;
-		FRHIBreadcrumbNode* const Node;
-		ERHIPipeline const Pipeline;
-
-		FRHIBreadcrumbEventScope(FRHIBreadcrumbEventScope const&) = delete;
-		FRHIBreadcrumbEventScope(FRHIBreadcrumbEventScope&&) = delete;
-
-		template<size_t N, typename... TArgs>
-		inline FRHIBreadcrumbEventScope(FRHIComputeCommandList& InRHICmdList, FRHIBreadcrumbData&& Data, ERHIPipeline InPipeline, bool bCondition, TCHAR const(&FormatString)[N], TArgs&&... Args);
-
-	public:
-		// Top-of-pipe breadcrumb event scope for RHI command lists
-		template<size_t N, typename... TArgs>
-		inline FRHIBreadcrumbEventScope(FRHIComputeCommandList& InRHICmdList, FRHIBreadcrumbData&& Data, bool bCondition, TCHAR const(&FormatString)[N], TArgs&&... Args);
+	// Private macro used to define a new breadcrumb type. Also forwards the given values through to a returned TRHIBreadcrumbInitializer,
+	// which can then be passed to the constructors for FRHIBreadcrumbScope / FRHIBreadcrumbEventManual, or the allocator AllocBreadcrumb() function.
+	#define RHI_BREADCRUMB_PRIVATE_DEFINE(StaticName, FormatString, ValueType, GPUStat)	\
+		[&](auto&&... Values)															\
+		{																				\
+			using namespace UE::RHI::Breadcrumbs::Private;								\
+																						\
+			TRHIBreadcrumbDesc<															\
+				TStringLiteralSize<decltype(FormatString)>::Value,						\
+				typename TFieldTraits<decltype(Values)>::TDescType...					\
+			> static const Desc(														\
+				  FRHIBreadcrumbData(StaticName, __FILE__, __LINE__, GPUStat)			\
+				, FormatString															\
+				, Values...																\
+			);																			\
+																						\
+			return std::tuple(															\
+				&Desc,																	\
+				std::tuple<typename TFieldTraits<decltype(Values)>::ValueType...>(		\
+					TFieldTraits<decltype(Values)>::ForwardValue(Values)...				\
+				)																		\
+			);																			\
+		}
 
-		// Bottom-of-pipe breadcrumb event scope for RHI contexts
-		template<size_t N, typename... TArgs>
-		inline FRHIBreadcrumbEventScope(IRHIComputeContext& InRHIContext, FRHIBreadcrumbData&& Data, bool bCondition, TCHAR const(&FormatString)[N], TArgs&&... Args);
+	#define RHI_BREADCRUMB_DESC_FORWARD_VALUES(StaticName, FormatString, GPUStat) RHI_BREADCRUMB_PRIVATE_DEFINE(StaticName, FormatString, TValueRef , GPUStat)
+	#define RHI_BREADCRUMB_DESC_COPY_VALUES(   StaticName, FormatString, GPUStat) RHI_BREADCRUMB_PRIVATE_DEFINE(StaticName, FormatString, TValueType, GPUStat)
 
-		inline ~FRHIBreadcrumbEventScope();
-	};
-
-#if RHI_NEW_GPU_PROFILER
-	#define RHI_GPU_STAT_ARGS(StatName) &GPUStat_##StatName
-#else
-	#define RHI_GPU_STAT_ARGS(StatName) GET_STATID(Stat_GPU_##StatName), CSV_STAT_FNAME(StatName)
-#endif
+	#if RHI_NEW_GPU_PROFILER && HAS_GPU_STATS	
+		#define RHI_GPU_STAT_ARGS(StatName) FRHIBreadcrumbData_Stats(&GPUStat_##StatName)
+		#define RHI_GPU_STAT_ARGS_NONE      FRHIBreadcrumbData_Stats(nullptr)
+	#elif HAS_GPU_STATS
+		#define RHI_GPU_STAT_ARGS(StatName) FRHIBreadcrumbData_Stats(GET_STATID(Stat_GPU_##StatName), CSV_STAT_FNAME(StatName))
+		#define RHI_GPU_STAT_ARGS_NONE      FRHIBreadcrumbData_Stats(TStatId(), NAME_None)
+	#else
+		#define RHI_GPU_STAT_ARGS(StatName) FRHIBreadcrumbData_Stats()
+		#define RHI_GPU_STAT_ARGS_NONE      FRHIBreadcrumbData_Stats()
+	#endif
 
-	#define RHI_BREADCRUMB_EVENT(                 RHICmdList_Or_RHIContext,                  Format, ...) FRHIBreadcrumbEventScope ANONYMOUS_VARIABLE(BreadcrumbEvent)(RHICmdList_Or_RHIContext, FRHIBreadcrumbData(__FILE__, __LINE__), true     , TEXT(Format), ## __VA_ARGS__)
-	#define RHI_BREADCRUMB_EVENT_CONDITIONAL(     RHICmdList_Or_RHIContext,       Condition, Format, ...) FRHIBreadcrumbEventScope ANONYMOUS_VARIABLE(BreadcrumbEvent)(RHICmdList_Or_RHIContext, FRHIBreadcrumbData(__FILE__, __LINE__), Condition, TEXT(Format), ## __VA_ARGS__)
-#if HAS_GPU_STATS
-	#define RHI_BREADCRUMB_EVENT_STAT(            RHICmdList_Or_RHIContext, Stat,            Format, ...) FRHIBreadcrumbEventScope ANONYMOUS_VARIABLE(BreadcrumbEvent)(RHICmdList_Or_RHIContext, FRHIBreadcrumbData(__FILE__, __LINE__, RHI_GPU_STAT_ARGS(Stat)), true     , TEXT(Format), ## __VA_ARGS__)
-	#define RHI_BREADCRUMB_EVENT_CONDITIONAL_STAT(RHICmdList_Or_RHIContext, Stat, Condition, Format, ...) FRHIBreadcrumbEventScope ANONYMOUS_VARIABLE(BreadcrumbEvent)(RHICmdList_Or_RHIContext, FRHIBreadcrumbData(__FILE__, __LINE__, RHI_GPU_STAT_ARGS(Stat)), Condition, TEXT(Format), ## __VA_ARGS__)
-#else
-	#define RHI_BREADCRUMB_EVENT_STAT(            RHICmdList_Or_RHIContext, Stat,            Format, ...) FRHIBreadcrumbEventScope ANONYMOUS_VARIABLE(BreadcrumbEvent)(RHICmdList_Or_RHIContext, FRHIBreadcrumbData(__FILE__, __LINE__), true     , TEXT(Format), ## __VA_ARGS__)
-	#define RHI_BREADCRUMB_EVENT_CONDITIONAL_STAT(RHICmdList_Or_RHIContext, Stat, Condition, Format, ...) FRHIBreadcrumbEventScope ANONYMOUS_VARIABLE(BreadcrumbEvent)(RHICmdList_Or_RHIContext, FRHIBreadcrumbData(__FILE__, __LINE__), Condition, TEXT(Format), ## __VA_ARGS__)
-#endif
+	// Varargs in breadcrumb macros can be given a name by wrapping them with this macro.
+	// Named fields are exposed to Unreal Insights as metadata on event markers.
+	#define RHI_BREADCRUMB_FIELD(Name, Value) std::forward_as_tuple(TEXT(Name), Value)
+
+	#define RHI_BREADCRUMB_EVENT_PRIVATE_IMPL(RHICmdList_Or_RHIContext, Stat, Condition, StaticName, Format, ...)\
+		TOptional<FRHIBreadcrumbScope> PREPROCESSOR_JOIN(BreadcrumbScope, __LINE__);							 \
+		do																										 \
+		{																										 \
+			if (Condition)																						 \
+			{																									 \
+				PREPROCESSOR_JOIN(BreadcrumbScope, __LINE__).Emplace(											 \
+					UE::RHI::Breadcrumbs::Private::GetRHICmdList(RHICmdList_Or_RHIContext),						 \
+					RHI_BREADCRUMB_DESC_FORWARD_VALUES(															 \
+						  StaticName																			 \
+						, Format																				 \
+						, Stat																					 \
+					)(__VA_ARGS__)																				 \
+				);																								 \
+			}																									 \
+		} while(false)
+
+	// Note, the varargs are deprecated and ignored in these macros.
+	#define RHI_BREADCRUMB_EVENT(                             RHICmdList_Or_RHIContext,                  StaticName,         ...) RHI_BREADCRUMB_EVENT_PRIVATE_IMPL(RHICmdList_Or_RHIContext, RHI_GPU_STAT_ARGS_NONE ,      true, TEXT(StaticName),      nullptr, ##__VA_ARGS__)
+	#define RHI_BREADCRUMB_EVENT_CONDITIONAL(                 RHICmdList_Or_RHIContext,       Condition, StaticName,         ...) RHI_BREADCRUMB_EVENT_PRIVATE_IMPL(RHICmdList_Or_RHIContext, RHI_GPU_STAT_ARGS_NONE , Condition, TEXT(StaticName),      nullptr, ##__VA_ARGS__)
+	#define RHI_BREADCRUMB_EVENT_STAT(                        RHICmdList_Or_RHIContext, Stat,            StaticName,         ...) RHI_BREADCRUMB_EVENT_PRIVATE_IMPL(RHICmdList_Or_RHIContext, RHI_GPU_STAT_ARGS(Stat),      true, TEXT(StaticName),      nullptr, ##__VA_ARGS__)
+	#define RHI_BREADCRUMB_EVENT_CONDITIONAL_STAT(            RHICmdList_Or_RHIContext, Stat, Condition, StaticName,         ...) RHI_BREADCRUMB_EVENT_PRIVATE_IMPL(RHICmdList_Or_RHIContext, RHI_GPU_STAT_ARGS(Stat), Condition, TEXT(StaticName),      nullptr, ##__VA_ARGS__)
+													          
+	// Format versions of the breadcrumb macros.	          
+	#define RHI_BREADCRUMB_EVENT_F(                           RHICmdList_Or_RHIContext,                  StaticName, Format, ...) RHI_BREADCRUMB_EVENT_PRIVATE_IMPL(RHICmdList_Or_RHIContext, RHI_GPU_STAT_ARGS_NONE ,      true, TEXT(StaticName), TEXT(Format), ##__VA_ARGS__)
+	#define RHI_BREADCRUMB_EVENT_CONDITIONAL_F(               RHICmdList_Or_RHIContext,       Condition, StaticName, Format, ...) RHI_BREADCRUMB_EVENT_PRIVATE_IMPL(RHICmdList_Or_RHIContext, RHI_GPU_STAT_ARGS_NONE , Condition, TEXT(StaticName), TEXT(Format), ##__VA_ARGS__)
+	#define RHI_BREADCRUMB_EVENT_STAT_F(                      RHICmdList_Or_RHIContext, Stat,            StaticName, Format, ...) RHI_BREADCRUMB_EVENT_PRIVATE_IMPL(RHICmdList_Or_RHIContext, RHI_GPU_STAT_ARGS(Stat),      true, TEXT(StaticName), TEXT(Format), ##__VA_ARGS__)
+	#define RHI_BREADCRUMB_EVENT_CONDITIONAL_STAT_F(          RHICmdList_Or_RHIContext, Stat, Condition, StaticName, Format, ...) RHI_BREADCRUMB_EVENT_PRIVATE_IMPL(RHICmdList_Or_RHIContext, RHI_GPU_STAT_ARGS(Stat), Condition, TEXT(StaticName), TEXT(Format), ##__VA_ARGS__)
 
 	// Used only for back compat with SCOPED_DRAW_EVENTF
-	#define RHI_BREADCRUMB_EVENT_STR_DEPRECATED(            RHICmdList_Or_RHIContext,            Format, ...) FRHIBreadcrumbEventScope ANONYMOUS_VARIABLE(BreadcrumbEvent)(RHICmdList_Or_RHIContext, FRHIBreadcrumbData(__FILE__, __LINE__), true     , Format, ## __VA_ARGS__)
-	#define RHI_BREADCRUMB_EVENT_CONDITIONAL_STR_DEPRECATED(RHICmdList_Or_RHIContext, Condition, Format, ...) FRHIBreadcrumbEventScope ANONYMOUS_VARIABLE(BreadcrumbEvent)(RHICmdList_Or_RHIContext, FRHIBreadcrumbData(__FILE__, __LINE__), Condition, Format, ## __VA_ARGS__)
-
-	static FString GetSafeBreadcrumbPath(FRHIBreadcrumbNode* Node)
-	{
-		return Node ? Node->GetFullPath() : TEXT("NoBreadcrumb");
-	}
-
-#if DO_CHECK
-	#define RHI_BREADCRUMB_CHECKF(RHICmdList, Condition, Format, ...) \
-	{ \
-		if(UNLIKELY(!(Condition))) \
-		{ \
-			if (FDebug::CheckVerifyFailedImpl2(#Condition, __FILE__, __LINE__, Format TEXT("\nBreadcrumbs: %s"), ##__VA_ARGS__, *GetSafeBreadcrumbPath(GetTopBreadcrumbNode(RHICmdList)))) \
-			{ \
-				PLATFORM_BREAK(); \
-			} \
-			CA_ASSUME(false); \
-		} \
-	}
-#else // DO_CHECK
-	#define RHI_BREADCRUMB_CHECKF(RHICmdList, Condition, Format, ...)
-#endif // DO_CHECK
-
-	#define RHI_BREADCRUMB_CHECK(RHICmdList, Condition) RHI_BREADCRUMB_CHECKF(RHICmdList, Condition, TEXT(""))
-
-	#define RHI_BREADCRUMB_CLOG(RHICmdList, Condition, Log, Severity, Format, ...) \
-	{ \
-		if(UNLIKELY(Condition)) \
-		{ \
-			UE_LOG(Log, Severity, Format TEXT("\nBreadcrumbs: %s"), ##__VA_ARGS__, *GetSafeBreadcrumbPath(GetTopBreadcrumbNode(RHICmdList))); \
-		} \
-	}
+	#define RHI_BREADCRUMB_EVENT_F_STR_DEPRECATED(            RHICmdList_Or_RHIContext,                  StaticName, Format, ...) RHI_BREADCRUMB_EVENT_PRIVATE_IMPL(RHICmdList_Or_RHIContext, RHI_GPU_STAT_ARGS_NONE ,      true, TEXT(StaticName),      Format , ##__VA_ARGS__)
+	#define RHI_BREADCRUMB_EVENT_F_CONDITIONAL_STR_DEPRECATED(RHICmdList_Or_RHIContext,       Condition, StaticName, Format, ...) RHI_BREADCRUMB_EVENT_PRIVATE_IMPL(RHICmdList_Or_RHIContext, RHI_GPU_STAT_ARGS_NONE , Condition, TEXT(StaticName),      Format , ##__VA_ARGS__)
 
-	#define RHI_BREADCRUMB_LOG(RHICmdList, Log, Severity, Format, ...) RHI_BREADCRUMB_CLOG(RHICmdList, true, Log, Severity, Format, ##__VA_ARGS__)
+	// Log and check macros that include the current breadcrumb path
+	#define RHI_BREADCRUMB_LOG(   RHICmdList_Or_RHIContext,            CategoryName, Verbosity, Format, ...) UE_LOG(            CategoryName, Verbosity, Format TEXT("\nBreadcrumbs: %s"), ##__VA_ARGS__, *UE::RHI::Breadcrumbs::Private::GetSafeBreadcrumbPath(RHICmdList_Or_RHIContext))
+	#define RHI_BREADCRUMB_CLOG(  RHICmdList_Or_RHIContext, Condition, CategoryName, Verbosity, Format, ...) UE_CLOG(Condition, CategoryName, Verbosity, Format TEXT("\nBreadcrumbs: %s"), ##__VA_ARGS__, *UE::RHI::Breadcrumbs::Private::GetSafeBreadcrumbPath(RHICmdList_Or_RHIContext))
+	#define RHI_BREADCRUMB_CHECKF(RHICmdList_Or_RHIContext, Condition,                          Format, ...) checkf( Condition,                          Format TEXT("\nBreadcrumbs: %s"), ##__VA_ARGS__, *UE::RHI::Breadcrumbs::Private::GetSafeBreadcrumbPath(RHICmdList_Or_RHIContext))
 
-#if DO_CHECK
-	#define RHI_BREADCRUMB_CHECK_SHIPPINGF(RHICmdList, Condition, Format, ...) \
-		RHI_BREADCRUMB_CHECKF(RHICmdList, Condition, Format, ##__VA_ARGS__)
-#else // !DO_CHECK
-	#define RHI_BREADCRUMB_CHECK_SHIPPINGF(RHICmdList, Condition, Format, ...) \
-		RHI_BREADCRUMB_CLOG(RHICmdList, !(Condition), LogRHI, Error, TEXT("Check '%s' failed. ") Format, TEXT(#Condition), ##__VA_ARGS__)
-#endif // !DO_CHECK
-
-	#define RHI_BREADCRUMB_CHECK_SHIPPING(RHICmdList, Condition) RHI_BREADCRUMB_CHECK_SHIPPINGF(RHICmdList, Condition, TEXT(""))
 #else // WITH_RHI_BREADCRUMBS == 0
 
-	#define RHI_BREADCRUMB_EVENT(...)
-	#define RHI_BREADCRUMB_EVENT_STAT(...)
-	#define RHI_BREADCRUMB_EVENT_CONDITIONAL(...)
-	#define RHI_BREADCRUMB_EVENT_CONDITIONAL_STAT(...)
-
-	#define RHI_BREADCRUMB_EVENT_STR_DEPRECATED(...)
-	#define RHI_BREADCRUMB_EVENT_CONDITIONAL_STR_DEPRECATED(...)
-
-	#define RHI_BREADCRUMB_CHECK(...)
-	#define RHI_BREADCRUMB_CHECKF(...)
+	#define RHI_BREADCRUMB_FIELD(                             ...)
+	#define RHI_BREADCRUMB_EVENT(                             ...)
+	#define RHI_BREADCRUMB_EVENT_CONDITIONAL(                 ...)
+	#define RHI_BREADCRUMB_EVENT_STAT(                        ...)
+	#define RHI_BREADCRUMB_EVENT_CONDITIONAL_STAT(            ...)
+	#define RHI_BREADCRUMB_EVENT_F(                           ...)
+	#define RHI_BREADCRUMB_EVENT_CONDITIONAL_F(               ...)
+	#define RHI_BREADCRUMB_EVENT_STAT_F(                      ...)
+	#define RHI_BREADCRUMB_EVENT_CONDITIONAL_STAT_F(          ...)
+	#define RHI_BREADCRUMB_EVENT_F_STR_DEPRECATED(            ...)
+	#define RHI_BREADCRUMB_EVENT_F_CONDITIONAL_STR_DEPRECATED(...)
+
+	// Log and check macros fall back to regular UE_LOG / check when breadcrumbs are not available
+	#define RHI_BREADCRUMB_LOG(   RHICmdList_Or_RHIContext,            CategoryName, Verbosity, Format, ...) UE_LOG(            CategoryName, Verbosity, Format, ##__VA_ARGS__)
+	#define RHI_BREADCRUMB_CLOG(  RHICmdList_Or_RHIContext, Condition, CategoryName, Verbosity, Format, ...) UE_CLOG(Condition, CategoryName, Verbosity, Format, ##__VA_ARGS__)
+	#define RHI_BREADCRUMB_CHECKF(RHICmdList_Or_RHIContext, Condition,                          Format, ...) checkf( Condition,                          Format, ##__VA_ARGS__)
 
-	#define RHI_BREADCRUMB_CHECK_LOG(...)
-	#define RHI_BREADCRUMB_CHECK_CLOG(...)
-
-	#define RHI_BREADCRUMB_CHECK_SHIPPINGF(...)
-	#define RHI_BREADCRUMB_CHECK_SHIPPING(...)
+#endif
 
+#if DO_CHECK
+	#define RHI_BREADCRUMB_CHECK_SHIPPINGF(RHICmdList_Or_RHIContext, Condition, Format, ...) RHI_BREADCRUMB_CHECKF(RHICmdList_Or_RHIContext, Condition, Format, ##__VA_ARGS__)
+#else
+	#define RHI_BREADCRUMB_CHECK_SHIPPINGF(RHICmdList_Or_RHIContext, Condition, Format, ...) RHI_BREADCRUMB_CLOG(RHICmdList_Or_RHIContext, !(Condition), LogRHI, Error, TEXT("Check '%s' failed. ") Format, TEXT(#Condition), ##__VA_ARGS__)
 #endif
 
+#define RHI_BREADCRUMB_CHECK_SHIPPING(RHICmdList_Or_RHIContext, Condition) RHI_BREADCRUMB_CHECK_SHIPPINGF(RHICmdList_Or_RHIContext, Condition, TEXT(""))
+#define RHI_BREADCRUMB_CHECK(         RHICmdList_Or_RHIContext, Condition) RHI_BREADCRUMB_CHECKF(         RHICmdList_Or_RHIContext, Condition, TEXT(""))
+
 //
 // Used to override the static_assert check for string literals in RHI breadcrumbs. This is required when using
 // literals returned by functions, or choosing between two string literals with a ternary operator, like so:
@@ -1038,4 +1236,7 @@ namespace UE::RHI::GPUProfiler
 // 
 // !! DO NOT USE THIS MACRO FOR NON-STRING LITERALS !!
 //
-#define RHI_BREADCRUMB_FORCE_STRING_LITERAL(TCharPointer) reinterpret_cast<TCHAR const(*)[1]>(TCharPointer)
+#define RHI_BREADCRUMB_FORCE_STRING_LITERAL [](auto&& TCharPointer) -> TCHAR const(&)[1]\
+	{																					\
+		return *reinterpret_cast<TCHAR const(*)[1]>(TCharPointer);						\
+	}
diff --git a/Engine/Source/Runtime/RHI/Public/RHICommandList.h b/Engine/Source/Runtime/RHI/Public/RHICommandList.h
--- a/Engine/Source/Runtime/RHI/Public/RHICommandList.h
+++ b/Engine/Source/Runtime/RHI/Public/RHICommandList.h
@@ -699,6 +699,11 @@ public:
 		return ActivePipelines;
 	}
 
+	FORCEINLINE ERHIPipeline GetPipelines() const
+	{
+		return ActivePipelines;
+	}
+
 	FORCEINLINE IRHICommandContext& GetContext()
 	{
 		checkf(IsSingleRHIPipeline(ActivePipelines), TEXT("Exactly one pipeline must be active to call GetContext(). Current pipeline mask is '0x%02x'."), static_cast<std::underlying_type_t<ERHIPipeline>>(ActivePipelines));
@@ -1502,7 +1507,7 @@ protected:
 public:
 #if WITH_RHI_BREADCRUMBS
 	friend FRHIBreadcrumbEventManual;
-	friend FRHIBreadcrumbEventScope;
+	friend FRHIBreadcrumbScope;
 
 	FRHIBreadcrumbNode*& GetCurrentBreadcrumbRef()
 	{
@@ -1636,16 +1641,6 @@ inline bool IRHIComputeContext::ShouldEmitBreadcrumbs() const
 {
 	return GetExecutingCommandList().CPUBreadcrumbState.bEmitBreadcrumbs;
 }
-
-inline FRHIBreadcrumbNode* GetTopBreadcrumbNode(FRHICommandListBase& RHICmdList)
-{
-	return RHICmdList.GetCurrentBreadcrumbRef();
-}
-
-inline FRHIBreadcrumbNode* GetTopBreadcrumbNode(FRHICommandListBase* RHICmdList)
-{
-	return RHICmdList->GetCurrentBreadcrumbRef();
-}
 #endif
 
 struct FUnnamedRhiCommand
@@ -4515,9 +4510,9 @@ public:
 		// No copy/DMA operation inside render passes
 		check(IsOutsideRenderPass());
 
-		RHI_BREADCRUMB_CHECK_SHIPPING(this, SourceBuffer != DestBuffer);
-		RHI_BREADCRUMB_CHECK_SHIPPING(this, DstOffset + NumBytes <= DestBuffer->GetSize());
-		RHI_BREADCRUMB_CHECK_SHIPPING(this, SrcOffset + NumBytes <= SourceBuffer->GetSize());
+		RHI_BREADCRUMB_CHECK_SHIPPING(*this, SourceBuffer != DestBuffer);
+		RHI_BREADCRUMB_CHECK_SHIPPING(*this, DstOffset + NumBytes <= DestBuffer->GetSize());
+		RHI_BREADCRUMB_CHECK_SHIPPING(*this, SrcOffset + NumBytes <= SourceBuffer->GetSize());
 
 		if (Bypass())
 		{
diff --git a/Engine/Source/Runtime/RHI/Public/RHICommandList.inl b/Engine/Source/Runtime/RHI/Public/RHICommandList.inl
--- a/Engine/Source/Runtime/RHI/Public/RHICommandList.inl
+++ b/Engine/Source/Runtime/RHI/Public/RHICommandList.inl
@@ -154,44 +154,71 @@ inline FRHIResourceReplaceBatcher::~FRHIResourceReplaceBatcher()
 
 #if WITH_RHI_BREADCRUMBS
 
-	// Top-of-pipe breadcrumb event scope for RHI command lists
-	template<size_t N, typename... TArgs>
-	inline FRHIBreadcrumbEventScope::FRHIBreadcrumbEventScope(FRHIComputeCommandList& InRHICmdList, FRHIBreadcrumbData&& Data, bool bCondition, TCHAR const(&FormatString)[N], TArgs&&... Args)
-		: FRHIBreadcrumbEventScope(InRHICmdList, MoveTemp(Data), InRHICmdList.GetPipeline(), bCondition, FormatString, Forward<TArgs>(Args)...)
-	{}
+	namespace UE::RHI::Breadcrumbs::Private
+	{
+		static inline FRHIComputeCommandList& GetRHICmdList(FRHIComputeCommandList& RHICmdList) { return RHICmdList; }
+		static inline FRHIComputeCommandList& GetRHICmdList(IRHIComputeContext    & RHIContext) { return static_cast<FRHIComputeCommandList&>(RHIContext.GetExecutingCommandList()); }
+
+		static inline FString GetSafeBreadcrumbPath(auto&& RHICmdList_Or_RHIContext)
+		{
+			FRHICommandListBase& RHICmdList = GetRHICmdList(RHICmdList_Or_RHIContext);
+			FRHIBreadcrumbNode* Node = RHICmdList.GetCurrentBreadcrumbRef();
+			return Node ? Node->GetFullPath() : TEXT("NoBreadcrumb");
+		}
+	}
+
+	template <typename TDesc, typename... TValues>
+	inline FRHIBreadcrumbNode* FRHIBreadcrumbAllocator::AllocBreadcrumb(TRHIBreadcrumbInitializer<TDesc, TValues...> const& Args)
+	{
+		TDesc const* Desc = std::get<0>(Args);
+		if (!Desc)
+		{
+			return nullptr;
+		}
 
-	// Bottom-of-pipe breadcrumb event scope for RHI contexts
-	template<size_t N, typename... TArgs>
-	inline FRHIBreadcrumbEventScope::FRHIBreadcrumbEventScope(IRHIComputeContext& InRHIContext, FRHIBreadcrumbData&& Data, bool bCondition, TCHAR const(&FormatString)[N], TArgs&&... Args)
-		: FRHIBreadcrumbEventScope(static_cast<FRHIComputeCommandList&>(InRHIContext.GetExecutingCommandList()), MoveTemp(Data), InRHIContext.GetPipeline(), bCondition, FormatString, Forward<TArgs>(Args)...)
+		return std::apply([&](auto&&... Values)
+		{
+			return Alloc<UE::RHI::Breadcrumbs::Private::TRHIBreadcrumb<TDesc>>(*this, *Desc, std::forward<decltype(Values)>(Values)...);
+		}, std::get<1>(Args));
+	}
+
+	template <typename TDesc, typename... TValues>
+	inline FRHIBreadcrumbScope::FRHIBreadcrumbScope(FRHIComputeCommandList& RHICmdList, TRHIBreadcrumbInitializer<TDesc, TValues...>&& Args)
+		: FRHIBreadcrumbScope(RHICmdList, RHICmdList.GetBreadcrumbAllocator().AllocBreadcrumb(Args))
 	{}
 
-	template<size_t N, typename... TArgs>
-	inline FRHIBreadcrumbEventScope::FRHIBreadcrumbEventScope(FRHIComputeCommandList& InRHICmdList, FRHIBreadcrumbData&& Data, ERHIPipeline InPipeline, bool bCondition, TCHAR const(&FormatString)[N], TArgs&&... Args)
-		: RHICmdList(InRHICmdList)
-		, Node(bCondition ? RHICmdList.GetBreadcrumbAllocator().AllocBreadcrumb(MoveTemp(Data), FormatString, Forward<TArgs>(Args)...) : nullptr)
-		, Pipeline(InPipeline)
+	inline FRHIBreadcrumbScope::FRHIBreadcrumbScope(FRHIComputeCommandList& RHICmdList, FRHIBreadcrumbNode* Node)
+		: RHICmdList(RHICmdList)
+		, Node(Node)
 	{
 		if (Node)
 		{
 			Node->SetParent(RHICmdList.PersistentState.LocalBreadcrumb);
 			RHICmdList.BeginBreadcrumbCPU(Node, true);
-			RHICmdList.BeginBreadcrumbGPU(Node, Pipeline);
+
+			for (ERHIPipeline Pipeline : MakeFlagsRange(RHICmdList.GetPipelines()))
+			{
+				RHICmdList.BeginBreadcrumbGPU(Node, Pipeline);
+			}
 		}
 	}
 
-	inline FRHIBreadcrumbEventScope::~FRHIBreadcrumbEventScope()
+	inline FRHIBreadcrumbScope::~FRHIBreadcrumbScope()
 	{
 		if (Node)
 		{
-			RHICmdList.EndBreadcrumbGPU(Node, Pipeline);
+			for (ERHIPipeline Pipeline : MakeFlagsRange(RHICmdList.GetPipelines()))
+			{
+				RHICmdList.EndBreadcrumbGPU(Node, Pipeline);
+			}
+
 			RHICmdList.EndBreadcrumbCPU(Node, true);
 		}
 	}
 
-	template<size_t N, typename... TArgs>
-	inline FRHIBreadcrumbEventManual::FRHIBreadcrumbEventManual(FRHIComputeCommandList& RHICmdList, FRHIBreadcrumbData&& Data, TCHAR const(&FormatString)[N], TArgs&&... Args)
-		: Node(RHICmdList.GetBreadcrumbAllocator().AllocBreadcrumb(MoveTemp(Data), FormatString, Forward<TArgs>(Args)...))
+	template <typename TDesc, typename... TValues>
+	inline FRHIBreadcrumbEventManual::FRHIBreadcrumbEventManual(FRHIComputeCommandList& RHICmdList, TRHIBreadcrumbInitializer<TDesc, TValues...>&& Args)
+		: Node(RHICmdList.GetBreadcrumbAllocator().AllocBreadcrumb(Args))
 	#if DO_CHECK
 		, Pipeline(RHICmdList.GetPipeline())
 		, ThreadId(FPlatformTLS::GetCurrentThreadId())
diff --git a/Engine/Source/Runtime/RHICore/Private/RHICoreIntelBreadcrumbs.cpp b/Engine/Source/Runtime/RHICore/Private/RHICoreIntelBreadcrumbs.cpp
--- a/Engine/Source/Runtime/RHICore/Private/RHICoreIntelBreadcrumbs.cpp
+++ b/Engine/Source/Runtime/RHICore/Private/RHICoreIntelBreadcrumbs.cpp
@@ -103,7 +103,7 @@ namespace UE::RHICore::Intel::GPUCrashDumps
 			FRHIBreadcrumb::FBuffer* Buffer = new FRHIBreadcrumb::FBuffer;
 			NameStorage.Add(Buffer);
 
-			TCHAR const* NameStr = Breadcrumb->Name.GetTCHAR(*Buffer);
+			TCHAR const* NameStr = Breadcrumb->GetTCHAR(*Buffer);
 
 			*ppResolvedBuffer = const_cast<TCHAR*>(NameStr);
 			*pResolvedBufferSize = (FCString::Strlen(NameStr) + 1) * sizeof(TCHAR); // Include null terminator
diff --git a/Engine/Source/Runtime/RHICore/Private/RHICoreNvidiaAftermath.cpp b/Engine/Source/Runtime/RHICore/Private/RHICoreNvidiaAftermath.cpp
--- a/Engine/Source/Runtime/RHICore/Private/RHICoreNvidiaAftermath.cpp
+++ b/Engine/Source/Runtime/RHICore/Private/RHICoreNvidiaAftermath.cpp
@@ -754,7 +754,7 @@ namespace UE::RHICore::Nvidia::Aftermath
 								FRHIBreadcrumbNode const* Node = static_cast<FRHIBreadcrumbNode const*>(Marker.markerData);
 								Name = Node == FRHIBreadcrumbNode::Sentinel
 									? RootNodeName
-									: Node->Name.GetTCHAR(Buffer);
+									: Node->GetTCHAR(Buffer);
 						#else
 							Name = BreadcrumbsDisabledStr;
 						#endif
@@ -867,7 +867,7 @@ namespace UE::RHICore::Nvidia::Aftermath
 				FRHIBreadcrumb::FBuffer* Buffer = new FRHIBreadcrumb::FBuffer;
 				NameStorage.Add(Buffer);
 
-				TCHAR const* NameStr = Breadcrumb->Name.GetTCHAR(*Buffer);
+				TCHAR const* NameStr = Breadcrumb->GetTCHAR(*Buffer);
 
 				*ResolvedMarkerData = const_cast<TCHAR*>(NameStr);
 				*ResolvedMarkerDataSize = (FCString::Strlen(NameStr) + 1) * sizeof(TCHAR); // Include null terminator
diff --git a/Engine/Source/Runtime/RHICore/Public/RHICoreIntelBreadcrumbs.h b/Engine/Source/Runtime/RHICore/Public/RHICoreIntelBreadcrumbs.h
--- a/Engine/Source/Runtime/RHICore/Public/RHICoreIntelBreadcrumbs.h
+++ b/Engine/Source/Runtime/RHICore/Public/RHICoreIntelBreadcrumbs.h
@@ -58,7 +58,7 @@
 				Size = sizeof( FRHIBreadcrumbNode* );			// pointer size
 			#else
 				// Generate the breadcrumb node name
-				TCHAR const* NameStr = Breadcrumb ? Breadcrumb->Name.GetTCHAR(Buffer) : RootNodeName;
+				TCHAR const* NameStr = Breadcrumb ? Breadcrumb->GetTCHAR(Buffer) : RootNodeName;
 				Ptr = NameStr;
 				Size = (FCString::Strlen(NameStr) + 1) * sizeof(TCHAR);// Include null terminator
 			#endif
diff --git a/Engine/Source/Runtime/RHICore/Public/RHICoreNvidiaAftermath.h b/Engine/Source/Runtime/RHICore/Public/RHICoreNvidiaAftermath.h
--- a/Engine/Source/Runtime/RHICore/Public/RHICoreNvidiaAftermath.h
+++ b/Engine/Source/Runtime/RHICore/Public/RHICoreNvidiaAftermath.h
@@ -107,7 +107,7 @@
 				Ptr = Breadcrumb ? Breadcrumb : FRHIBreadcrumbNode::Sentinel;
 			#else
 				// Generate the breadcrumb node name and have Aftermath copy the string.
-				TCHAR const* NameStr = Breadcrumb ? Breadcrumb->Name.GetTCHAR(Buffer) : RootNodeName;
+				TCHAR const* NameStr = Breadcrumb ? Breadcrumb->GetTCHAR(Buffer) : RootNodeName;
 				Ptr = NameStr;
 				Size = (FCString::Strlen(NameStr) + 1) * sizeof(TCHAR);// Include null terminator
 			#endif
diff --git a/Engine/Source/Runtime/RenderCore/Private/DumpGPU.cpp b/Engine/Source/Runtime/RenderCore/Private/DumpGPU.cpp
--- a/Engine/Source/Runtime/RenderCore/Private/DumpGPU.cpp
+++ b/Engine/Source/Runtime/RenderCore/Private/DumpGPU.cpp
@@ -1929,7 +1929,8 @@ public:
 		{
 			if (FRDGScope_RHI const* RHIScope = ParentScope->Get<FRDGScope_RHI>())
 			{
-				bDumpPass = WildcardFilter.IsMatch(RHIScope->Name.GetTCHAR());
+				FRHIBreadcrumb::FBuffer Buffer;
+				bDumpPass = WildcardFilter.IsMatch(RHIScope->GetTCHAR(Buffer));
 			}
 		}
 #endif // RDG_EVENTS
@@ -2594,7 +2595,8 @@ void FRDGBuilder::DumpResourcePassOutputs(const FRDGPass* Pass)
 			{
 				if (FRDGScope_RHI const* RHIScope = ParentScope->Get<FRDGScope_RHI>())
 				{
-					ParentEventScopeNames.Add(MakeShareable(new FJsonValueString(RHIScope->Name.GetTCHAR())));
+					FRHIBreadcrumb::FBuffer Buffer;
+					ParentEventScopeNames.Add(MakeShareable(new FJsonValueString(RHIScope->GetTCHAR(Buffer))));
 				}
 			}
 #endif // RDG_EVENTS
diff --git a/Engine/Source/Runtime/RenderCore/Private/RenderCaptureInterface.cpp b/Engine/Source/Runtime/RenderCore/Private/RenderCaptureInterface.cpp
--- a/Engine/Source/Runtime/RenderCore/Private/RenderCaptureInterface.cpp
+++ b/Engine/Source/Runtime/RenderCore/Private/RenderCaptureInterface.cpp
@@ -35,7 +35,7 @@ namespace RenderCaptureInterface
 #if WITH_RHI_BREADCRUMBS
 				if (bPushEvent)
 				{
-					Breadcrumb->Emplace(RHICommandListLocal, FRHIBreadcrumbData(__FILE__, __LINE__), TEXT("%s"), EventName);
+					Breadcrumb->Emplace(RHICommandListLocal, RHI_BREADCRUMB_DESC_FORWARD_VALUES(TEXT("FScopedCapture"), TEXT("%s"), RHI_GPU_STAT_ARGS_NONE)(EventName));
 				}
 #endif
 			});
@@ -58,7 +58,7 @@ namespace RenderCaptureInterface
 			if (bEvent)
 			{
 				RHIBreadcrumb = MakeUnique<TOptional<FRHIBreadcrumbEventManual>>();
-				RHIBreadcrumb->Emplace(*RHICommandList, FRHIBreadcrumbData(__FILE__, __LINE__), TEXT("%s"), FString(InEventName));
+				RHIBreadcrumb->Emplace(*RHICommandList, RHI_BREADCRUMB_DESC_FORWARD_VALUES(TEXT("FScopedCapture"), TEXT("%s"), RHI_GPU_STAT_ARGS_NONE)(FString(InEventName)));
 			}
 #endif
 		}
@@ -83,10 +83,7 @@ namespace RenderCaptureInterface
 				});
 
 #if RDG_EVENTS
-			if (bEvent)
-			{
-				RDGEvent.Emplace(*GraphBuilder, ERDGScopeFlags::None, FRHIBreadcrumbData(__FILE__, __LINE__), RDG_EVENT_NAME("%s", InEventName));
-			}
+			RDG_EVENT_SCOPE_CONSTRUCT(RDGEvent, *GraphBuilder, bEvent, ERDGScopeFlags::None, RHI_GPU_STAT_ARGS_NONE, TEXT("FScopedCapture"), TEXT("%s"), FString(InEventName));
 #endif
 		}
 	}
diff --git a/Engine/Source/Runtime/RenderCore/Private/RenderGraphBuilder.cpp b/Engine/Source/Runtime/RenderCore/Private/RenderGraphBuilder.cpp
--- a/Engine/Source/Runtime/RenderCore/Private/RenderGraphBuilder.cpp
+++ b/Engine/Source/Runtime/RenderCore/Private/RenderGraphBuilder.cpp
@@ -1742,7 +1742,7 @@ void FRDGBuilder::WaitForParallelSetupTasks(ERDGSetupTaskWaitPoint WaitPoint)
 void FRDGBuilder::Execute()
 {
 	CSV_SCOPED_TIMING_STAT_EXCLUSIVE(RDG);
-	SCOPED_DRAW_EVENTF(RHICmdList, FRDGBuilder_Execute, TEXT("FRDGBuilder::Execute"));
+	RHI_BREADCRUMB_EVENT_F(RHICmdList, "RenderGraphExecute", "RenderGraphExecute - %s", BuilderName);
 
 #if WITH_RHI_BREADCRUMBS
 	check(LocalCurrentBreadcrumb == FRHIBreadcrumbNode::Sentinel);
@@ -2595,14 +2595,7 @@ FRDGPass* FRDGBuilder::SetupParameterPass(FRDGPass* Pass)
 	IF_RDG_ENABLE_DEBUG(UserValidation.ValidateAddPass(Pass));
 	CSV_SCOPED_TIMING_STAT_EXCLUSIVE_CONDITIONAL(RDGBuilder_SetupPass, GRDGVerboseCSVStats != 0);
 
-#if RDG_EVENTS
-	TOptional<TRDGEventScopeGuard<FRDGScope_RHI>> PassNameScope;
-	if (ScopeState.ScopeMode == ERDGScopeMode::AllEventsAndPassNames)
-	{
-		FRDGEventName Name = Pass->GetEventName();
-		PassNameScope.Emplace(*this, ERDGScopeFlags::None, FRHIBreadcrumbData(__FILE__, __LINE__), MoveTemp(Name));
-	}
-#endif
+	RDG_EVENT_SCOPE_CONDITIONAL_NAME(*this, ScopeState.ScopeMode == ERDGScopeMode::AllEventsAndPassNames, Pass->GetEventName());
 
 	SetupPassInternals(Pass);
 
diff --git a/Engine/Source/Runtime/RenderCore/Private/RenderGraphEvent.cpp b/Engine/Source/Runtime/RenderCore/Private/RenderGraphEvent.cpp
--- a/Engine/Source/Runtime/RenderCore/Private/RenderGraphEvent.cpp
+++ b/Engine/Source/Runtime/RenderCore/Private/RenderGraphEvent.cpp
@@ -285,47 +285,6 @@ FRDGEventName::FRDGEventName(const TCHAR* EventFormat, ...)
 #endif
 }
 
-#if WITH_RHI_BREADCRUMBS
-FRHIBreadcrumbNode* FRDGEventName::AllocBreadcrumb(FRHIBreadcrumbData&& Data, FRHIBreadcrumbAllocator& Allocator) const
-{
-#if RDG_EVENTS == RDG_EVENTS_STRING_COPY
-
-		if (FormattedEventName.IsEmpty())
-		{
-			TCHAR const* String = EventFormat[0] != 0
-				? EventFormat
-				: TEXT("<unnamed pass>");
-			
-			// Cast hack to force treating EventFormat as a string literal.
-			return Allocator.AllocBreadcrumb(MoveTemp(Data), *reinterpret_cast<TCHAR const(*)[1]>(String));
-		}
-		else
-		{
-			//
-			// Copy the string into the breadcrumb allocator.
-			// This allows us to allocate a literal breadcrumb, since the lifetime of the string copy will be tied to the breadcrumb node itself.
-			//
-			uint32 Size = (FormattedEventName.Len() + 1) * sizeof(FString::ElementType);
-			FString::ElementType* StrCopy = static_cast<FString::ElementType*>(Allocator.Alloc(Size, alignof(FString::ElementType)));
-			FMemory::Memcpy(StrCopy, *FormattedEventName, Size);
-
-			// Cast hack to force treating StrCopy as a string literal.
-			return Allocator.AllocBreadcrumb(MoveTemp(Data), *reinterpret_cast<TCHAR const(*)[1]>(StrCopy));
-		}
-
-#elif RDG_EVENTS >= RDG_EVENTS_STRING_REF
-
-		// Cast hack to force treating EventFormat as a string literal.
-		return Allocator.AllocBreadcrumb(MoveTemp(Data), *reinterpret_cast<TCHAR const(*)[1]>(EventFormat));
-
-#else
-
-		return nullptr;
-
-#endif
-}
-#endif // WITH_RHI_BREADCRUMBS
-
 const TCHAR* FRDGEventName::GetTCHAR() const
 {
 #if RDG_EVENTS == RDG_EVENTS_STRING_COPY
@@ -355,11 +314,12 @@ FString FRDGScope::GetFullPath(FRDGEventName const& PassName)
 {
 	FString Path = PassName.GetTCHAR();
 #if RDG_EVENTS
+	FRHIBreadcrumb::FBuffer Buffer;
 	for (FRDGScope* Current = Parent; Current; Current = Current->Parent)
 	{
 		if (FRDGScope_RHI* RHIScope = Current->Get<FRDGScope_RHI>())
 		{
-			Path = RHIScope->Name.GetTCHAR() / Path;
+			Path = RHIScope->GetTCHAR(Buffer) / Path;
 		}
 	}
 #endif // RDG_EVENTS
diff --git a/Engine/Source/Runtime/RenderCore/Private/RenderGraphTrace.cpp b/Engine/Source/Runtime/RenderCore/Private/RenderGraphTrace.cpp
--- a/Engine/Source/Runtime/RenderCore/Private/RenderGraphTrace.cpp
+++ b/Engine/Source/Runtime/RenderCore/Private/RenderGraphTrace.cpp
@@ -218,7 +218,8 @@ void FRDGTrace::OutputGraphEnd(const FRDGBuilder& GraphBuilder)
 			{
 				if (Current->CPUFirstPass && Current->CPULastPass)
 				{
-					const TCHAR* Name = RHIScope->Name.GetTCHAR();
+					FRHIBreadcrumb::FBuffer Buffer;
+					const TCHAR* Name = RHIScope->GetTCHAR(Buffer);
 
 					uint32 Depth = 0;
 
diff --git a/Engine/Source/Runtime/RenderCore/Public/ProfilingDebugging/RealtimeGPUProfiler.h b/Engine/Source/Runtime/RenderCore/Public/ProfilingDebugging/RealtimeGPUProfiler.h
--- a/Engine/Source/Runtime/RenderCore/Public/ProfilingDebugging/RealtimeGPUProfiler.h
+++ b/Engine/Source/Runtime/RenderCore/Public/ProfilingDebugging/RealtimeGPUProfiler.h
@@ -25,58 +25,72 @@
 
 #if WITH_RHI_BREADCRUMBS
 
-	struct FRHIBreadcrumbEvent_GameThread
+	struct FRHIBreadcrumbScope_GameThread
 	{
 	private:
-		TOptional<FRHIBreadcrumbEventScope>* Event;
+		TOptional<FRHIBreadcrumbScope>* Event;
 
 	public:
-		template<size_t N, typename... TArgs>
-		FRHIBreadcrumbEvent_GameThread(bool bCondition, FRHIBreadcrumbData&& Data, TCHAR const(&FormatString)[N], TArgs&&... Args)
-			: Event(bCondition ? new TOptional<FRHIBreadcrumbEventScope> : nullptr)
+		template <typename TDesc, typename... TValues>
+		FRHIBreadcrumbScope_GameThread(TRHIBreadcrumbInitializer<TDesc, TValues...>&& Args)
+			: Event(new TOptional<FRHIBreadcrumbScope>)
 		{
 			check(IsInGameThread());
-			if (Event)
+
+			ENQUEUE_RENDER_COMMAND(FRHIBreadcrumbScope_GameThread_Begin)(
+			[
+				Event = Event,
+				Args = MoveTemp(Args)
+			](FRHICommandListImmediate& RHICmdList) mutable
 			{
-				ENQUEUE_RENDER_COMMAND(FRHIBreadcrumbEvent_GameThread_Begin)(
-				[
-					Event = Event,
-					Data = MoveTemp(Data),
-					Values = std::make_tuple(std::forward<TArgs>(Args)...),
-					&FormatString
-				](FRHICommandListImmediate& RHICmdList) mutable
-				{
-					std::apply([Event, &RHICmdList, &Data, &FormatString](auto&&... InnerArgs)
-					{
-						Event->Emplace(RHICmdList, MoveTemp(Data), true, FormatString, std::forward<TArgs>(InnerArgs)...);
-					}, Values);
-				});
-			}
+				Event->Emplace(RHICmdList, MoveTemp(Args));
+			});
 		}
 
-		~FRHIBreadcrumbEvent_GameThread()
+		~FRHIBreadcrumbScope_GameThread()
 		{
 			check(IsInGameThread());
-			if (Event)
+
+			ENQUEUE_RENDER_COMMAND(FRHIBreadcrumbScope_GameThread_End)([Event = Event](FRHICommandListImmediate&)
 			{
-				ENQUEUE_RENDER_COMMAND(FRHIBreadcrumbEvent_GameThread_End)([Event = Event](FRHICommandListImmediate& RHICmdList)
-				{
-					delete Event;
-				});
-			}
+				delete Event;
+			});
 		}
 	};
 
-	#define RHI_BREADCRUMB_EVENT_GAMETHREAD(                       Format, ...) FRHIBreadcrumbEvent_GameThread ANONYMOUS_VARIABLE(BreadcrumbEvent_GameThread)(true     , FRHIBreadcrumbData(__FILE__, __LINE__), TEXT(Format), ##__VA_ARGS__)
-	#define RHI_BREADCRUMB_EVENT_CONDITIONAL_GAMETHREAD(Condition, Format, ...) FRHIBreadcrumbEvent_GameThread ANONYMOUS_VARIABLE(BreadcrumbEvent_GameThread)(Condition, FRHIBreadcrumbData(__FILE__, __LINE__), TEXT(Format), ##__VA_ARGS__)
+	#define RHI_BREADCRUMB_EVENT_GAMETHREAD_PRIVATE_IMPL(Stat, Condition, StaticName, Format, ...)	\
+		TOptional<FRHIBreadcrumbScope_GameThread> PREPROCESSOR_JOIN(BreadcrumbScope, __LINE__);		\
+		do																							\
+		{																							\
+			if (Condition)																			\
+			{																						\
+				PREPROCESSOR_JOIN(BreadcrumbScope, __LINE__).Emplace(								\
+					RHI_BREADCRUMB_DESC_COPY_VALUES(												\
+						  StaticName																\
+						, Format																	\
+						, Stat																		\
+					)(__VA_ARGS__)																	\
+				);																					\
+			}																						\
+		} while(false)
+
+	// Note, the varargs are deprecated and ignored in these two macros.
+	#define RHI_BREADCRUMB_EVENT_GAMETHREAD(                         StaticName,         ...) RHI_BREADCRUMB_EVENT_GAMETHREAD_PRIVATE_IMPL(RHI_GPU_STAT_ARGS_NONE,      true, TEXT(StaticName),      nullptr, ##__VA_ARGS__)
+	#define RHI_BREADCRUMB_EVENT_CONDITIONAL_GAMETHREAD(  Condition, StaticName,         ...) RHI_BREADCRUMB_EVENT_GAMETHREAD_PRIVATE_IMPL(RHI_GPU_STAT_ARGS_NONE, Condition, TEXT(StaticName),      nullptr, ##__VA_ARGS__)
+
+	// Format versions of the breadcrumb macros.
+	#define RHI_BREADCRUMB_EVENT_GAMETHREAD_F(                       StaticName, Format, ...) RHI_BREADCRUMB_EVENT_GAMETHREAD_PRIVATE_IMPL(RHI_GPU_STAT_ARGS_NONE,      true, TEXT(StaticName), TEXT(Format), ##__VA_ARGS__)
+	#define RHI_BREADCRUMB_EVENT_CONDITIONAL_GAMETHREAD_F(Condition, StaticName, Format, ...) RHI_BREADCRUMB_EVENT_GAMETHREAD_PRIVATE_IMPL(RHI_GPU_STAT_ARGS_NONE, Condition, TEXT(StaticName), TEXT(Format), ##__VA_ARGS__)
 
 	// Used only for back compat with SCOPED_DRAW_EVENTF_GAMETHREAD
-	#define RHI_BREADCRUMB_EVENT_GAMETHREAD_STR_DEPRECATED(Format, ...) FRHIBreadcrumbEvent_GameThread ANONYMOUS_VARIABLE(BreadcrumbEvent_GameThread)(true, FRHIBreadcrumbData(__FILE__, __LINE__), Format, ##__VA_ARGS__)
+	#define RHI_BREADCRUMB_EVENT_GAMETHREAD_STR_DEPRECATED(          StaticName, Format, ...) RHI_BREADCRUMB_EVENT_GAMETHREAD_PRIVATE_IMPL(RHI_GPU_STAT_ARGS_NONE,      true, TEXT(StaticName),      Format , ##__VA_ARGS__)
 
 #else
 
 	#define RHI_BREADCRUMB_EVENT_GAMETHREAD(...)                do { } while(0)
 	#define RHI_BREADCRUMB_EVENT_CONDITIONAL_GAMETHREAD(...)    do { } while(0)
+	#define RHI_BREADCRUMB_EVENT_GAMETHREAD_F(...)			    do { } while(0)
+	#define RHI_BREADCRUMB_EVENT_CONDITIONAL_GAMETHREAD_F(...)  do { } while(0)
 	#define RHI_BREADCRUMB_EVENT_GAMETHREAD_STR_DEPRECATED(...) do { } while(0)
 
 #endif
@@ -84,22 +98,22 @@
 // Macros to allow for scoping of draw events outside of RHI function implementations
 // Render-thread event macros:
 #define SCOPED_DRAW_EVENT(RHICmdList, Name)                                      RHI_BREADCRUMB_EVENT(RHICmdList, #Name);
-#define SCOPED_DRAW_EVENTF(RHICmdList, Name, Format, ...)                        RHI_BREADCRUMB_EVENT_STR_DEPRECATED(RHICmdList, Format, ##__VA_ARGS__);
+#define SCOPED_DRAW_EVENTF(RHICmdList, Name, Format, ...)                        RHI_BREADCRUMB_EVENT_F_STR_DEPRECATED(RHICmdList, #Name, Format, ##__VA_ARGS__);
 #define SCOPED_CONDITIONAL_DRAW_EVENT(RHICmdList, Name, Condition)               RHI_BREADCRUMB_EVENT_CONDITIONAL(RHICmdList, Condition, #Name);
-#define SCOPED_CONDITIONAL_DRAW_EVENTF(RHICmdList, Name, Condition, Format, ...) RHI_BREADCRUMB_EVENT_CONDITIONAL_STR_DEPRECATED(RHICmdList, Condition, Format, ##__VA_ARGS__);
+#define SCOPED_CONDITIONAL_DRAW_EVENTF(RHICmdList, Name, Condition, Format, ...) RHI_BREADCRUMB_EVENT_F_CONDITIONAL_STR_DEPRECATED(RHICmdList, Condition, #Name, Format, ##__VA_ARGS__);
 
 // Non-render-thread event macros:
-#define SCOPED_DRAW_EVENT_GAMETHREAD(Name)                                       RHI_BREADCRUMB_EVENT_GAMETHREAD(#Name);
-#define SCOPED_DRAW_EVENTF_GAMETHREAD(Name, Format, ...)                         RHI_BREADCRUMB_EVENT_GAMETHREAD_STR_DEPRECATED(Format, ##__VA_ARGS__);
-#define SCOPED_CONDITIONAL_DRAW_EVENT_GAMETHREAD(Name, Condition)                RHI_BREADCRUMB_EVENT_CONDITIONAL_GAMETHREAD(Condition, #Name);
-#define SCOPED_CONDITIONAL_DRAW_EVENTF_GAMETHREAD(Name, Condition, Format, ...)  RHI_BREADCRUMB_EVENT_CONDITIONAL_GAMETHREAD(Condition, Format, ##__VA_ARGS__);
+#define SCOPED_DRAW_EVENT_GAMETHREAD(Name)															 UE_DEPRECATED_MACRO(5.6, "SCOPED_DRAW_EVENT_GAMETHREAD has been deprecated. Use RHI_BREADCRUMB_EVENT_GAMETHREAD instead."                           ) RHI_BREADCRUMB_EVENT_GAMETHREAD(#Name);
+#define SCOPED_DRAW_EVENTF_GAMETHREAD(Name, Format, ...)											 UE_DEPRECATED_MACRO(5.6, "SCOPED_DRAW_EVENTF_GAMETHREAD has been deprecated. Use RHI_BREADCRUMB_EVENT_GAMETHREAD_F instead."                        ) RHI_BREADCRUMB_EVENT_GAMETHREAD_STR_DEPRECATED(#Name, Format, ##__VA_ARGS__);
+#define SCOPED_CONDITIONAL_DRAW_EVENT_GAMETHREAD(Name, Condition)									 UE_DEPRECATED_MACRO(5.6, "SCOPED_CONDITIONAL_DRAW_EVENT_GAMETHREAD has been deprecated. Use RHI_BREADCRUMB_EVENT_CONDITIONAL_GAMETHREAD instead."   ) RHI_BREADCRUMB_EVENT_CONDITIONAL_GAMETHREAD(Condition, #Name);
+#define SCOPED_CONDITIONAL_DRAW_EVENTF_GAMETHREAD(Name, Condition, Format, ...)						 UE_DEPRECATED_MACRO(5.6, "SCOPED_CONDITIONAL_DRAW_EVENTF_GAMETHREAD has been deprecated. Use RHI_BREADCRUMB_EVENT_CONDITIONAL_GAMETHREAD_F instead.") RHI_BREADCRUMB_EVENT_CONDITIONAL_GAMETHREAD_F(Condition, #Name, Format, ##__VA_ARGS__);
 
 // Deprecated macros
 #define BEGIN_DRAW_EVENTF(RHICmdList, Name, Event, Format, ...)                                      UE_DEPRECATED_MACRO(5.5, "BEGIN_DRAW_EVENTF has been deprecated. Equivalent functionality can be implemented by constructing / destructing an instance of FRHIBreadcrumbEventManual."                      )
 #define STOP_DRAW_EVENT(Event)                                                                       UE_DEPRECATED_MACRO(5.5, "STOP_DRAW_EVENT has been deprecated. Equivalent functionality can be implemented by constructing / destructing an instance of FRHIBreadcrumbEventManual."                        )
-#define STOP_DRAW_EVENT_GAMETHREAD(...)                                                              UE_DEPRECATED_MACRO(5.5, "STOP_DRAW_EVENT_GAMETHREAD has been deprecated. Equivalent functionality can be implemented by constructing / destructing an instance of FRHIBreadcrumbEvent_GameThread."        )
-#define BEGIN_DRAW_EVENTF_GAMETHREAD(...)                                                            UE_DEPRECATED_MACRO(5.5, "BEGIN_DRAW_EVENTF_GAMETHREAD has been deprecated. Equivalent functionality can be implemented by constructing / destructing an instance of FRHIBreadcrumbEvent_GameThread."      )
-#define BEGIN_DRAW_EVENTF_COLOR_GAMETHREAD(...)                                                      UE_DEPRECATED_MACRO(5.5, "BEGIN_DRAW_EVENTF_COLOR_GAMETHREAD has been deprecated. Equivalent functionality can be implemented by constructing / destructing an instance of FRHIBreadcrumbEvent_GameThread.")
+#define STOP_DRAW_EVENT_GAMETHREAD(...)                                                              UE_DEPRECATED_MACRO(5.5, "STOP_DRAW_EVENT_GAMETHREAD has been deprecated. Equivalent functionality can be implemented by constructing / destructing an instance of FRHIBreadcrumbScope_GameThread."        )
+#define BEGIN_DRAW_EVENTF_GAMETHREAD(...)                                                            UE_DEPRECATED_MACRO(5.5, "BEGIN_DRAW_EVENTF_GAMETHREAD has been deprecated. Equivalent functionality can be implemented by constructing / destructing an instance of FRHIBreadcrumbScope_GameThread."      )
+#define BEGIN_DRAW_EVENTF_COLOR_GAMETHREAD(...)                                                      UE_DEPRECATED_MACRO(5.5, "BEGIN_DRAW_EVENTF_COLOR_GAMETHREAD has been deprecated. Equivalent functionality can be implemented by constructing / destructing an instance of FRHIBreadcrumbScope_GameThread.")
 #define SCOPED_DRAW_EVENT_COLOR(RHICmdList, Color, Name)                                             UE_DEPRECATED_MACRO(5.5, "SCOPED_DRAW_EVENT_COLOR has been deprecated. Use RHI_BREADCRUMB_EVENT instead."                                               ) RHI_BREADCRUMB_EVENT(RHICmdList, #Name)
 #define SCOPED_GPU_EVENT(RHICmdList, Name)                                                           UE_DEPRECATED_MACRO(5.5, "SCOPED_GPU_EVENT has been deprecated. Use RHI_BREADCRUMB_EVENT instead."                                                      ) RHI_BREADCRUMB_EVENT(RHICmdList, #Name)
 #define SCOPED_GPU_EVENT_COLOR(RHICmdList, Color, Name)                                              UE_DEPRECATED_MACRO(5.5, "SCOPED_GPU_EVENT_COLOR has been deprecated. Use RHI_BREADCRUMB_EVENT instead."                                                ) RHI_BREADCRUMB_EVENT(RHICmdList, #Name)
@@ -144,10 +158,14 @@
 		#define DECLARE_GPU_STAT_NAME_TYPE(StatName, NameString) \
 			struct TRHIGPUStatNameProvider_##StatName            \
 			{                                                    \
-				static constexpr TCHAR const* GetName()          \
+				static constexpr TCHAR const* GetDisplayName()   \
 				{                                                \
 					return NameString;                           \
 				}                                                \
+				static constexpr TCHAR const* GetStatName()      \
+				{                                                \
+					return TEXT(#StatName);                      \
+				}                                                \
 			}
 
 		// Extern GPU stats are needed where a stat is used in multiple CPPs. Use the DECLARE_GPU_STAT_NAMED_EXTERN in the header and DEFINE_GPU_STAT in the CPPs
diff --git a/Engine/Source/Runtime/RenderCore/Public/RenderGraphBuilder.h b/Engine/Source/Runtime/RenderCore/Public/RenderGraphBuilder.h
--- a/Engine/Source/Runtime/RenderCore/Public/RenderGraphBuilder.h
+++ b/Engine/Source/Runtime/RenderCore/Public/RenderGraphBuilder.h
@@ -57,24 +57,6 @@ class FRDGBuilder
 
 	} AsyncDeleter;
 
-	struct
-	{
-		// Allocator for all root graph allocations on the graph builder thread.
-		FRDGAllocator Root;
-
-		// Allocator for async pass and parallel execute setup.
-		FRDGAllocator Task;
-
-		// Allocator for all allocations related to states / transitions.
-		FRDGAllocator Transition;
-
-		int32 GetByteCount() const
-		{
-			return Root.GetByteCount() + Task.GetByteCount() + Transition.GetByteCount();
-		}
-
-	} Allocators;
-
 	FRDGAllocatorScope RootAllocatorScope;
 
 public:
diff --git a/Engine/Source/Runtime/RenderCore/Public/RenderGraphEvent.h b/Engine/Source/Runtime/RenderCore/Public/RenderGraphEvent.h
--- a/Engine/Source/Runtime/RenderCore/Public/RenderGraphEvent.h
+++ b/Engine/Source/Runtime/RenderCore/Public/RenderGraphEvent.h
@@ -48,9 +48,14 @@ public:
 
 	const TCHAR* GetTCHAR() const;
 
-#if WITH_RHI_BREADCRUMBS
-	FRHIBreadcrumbNode* AllocBreadcrumb(FRHIBreadcrumbData&& Data, FRHIBreadcrumbAllocator& Allocator) const;
-#endif // WITH_RHI_BREADCRUMBS
+	bool HasFormattedString() const
+	{
+	#if RDG_EVENTS == RDG_EVENTS_STRING_COPY
+		return !FormattedEventName.IsEmpty();
+	#else
+		return false;
+	#endif
+	}
 
 private:
 #if RDG_EVENTS >= RDG_EVENTS_STRING_REF
@@ -163,18 +168,17 @@ struct FRDGScope_Budget
 #if RDG_EVENTS
 
 	// Scope type for inserting named events on the CPU and GPU timelines.
-	struct FRDGScope_RHI
+	class FRDGScope_RHI
 	{
-		FRDGEventName Name;
-
-	#if WITH_RHI_BREADCRUMBS
 		FRHIBreadcrumbNode* Node = nullptr;
-	#endif
 
-		inline FRDGScope_RHI(FRDGScopeState& State, FRHIBreadcrumbData&& Data, FRDGEventName&& Name);
-		inline void ImmediateEnd(FRDGScopeState& State);
+		inline FRDGScope_RHI(FRDGScopeState& State, FRHIBreadcrumbNode* Node);
 
-	#if WITH_RHI_BREADCRUMBS
+	public:
+		template <typename TDesc, typename... TValues>
+		inline FRDGScope_RHI(FRDGScopeState& State, TRHIBreadcrumbInitializer<TDesc, TValues...>&& Args);
+
+		inline void ImmediateEnd(FRDGScopeState& State);
 
 		void BeginCPU(FRHIComputeCommandList& RHICmdList, bool bPreScope)
 		{
@@ -200,17 +204,15 @@ struct FRDGScope_Budget
 			}
 		}
 
-	#else
-
-		void BeginCPU(FRHIComputeCommandList& RHICmdList, bool bPreScope) {}
-		void EndCPU  (FRHIComputeCommandList& RHICmdList, bool bPreScope) {}
-
-	#endif
-
 		// Nothing to do for Begin/EndGPU. The RHI API only requires breadcrumbs to be 
 		// begun/ended once, and will automatically fixup other pipelines whenever we switch.
 		void BeginGPU(FRHIComputeCommandList& RHICmdList) {}
 		void EndGPU  (FRHIComputeCommandList& RHICmdList) {}
+
+		TCHAR const* GetTCHAR(FRHIBreadcrumb::FBuffer& Buffer) const
+		{
+			return Node->GetTCHAR(Buffer);
+		}
 	};
 
 #endif // RDG_EVENTS
@@ -295,6 +297,12 @@ struct FRDGScope
 	
 	FStorage Impl;
 
+	template <typename TScopeType>
+	static constexpr uint32 GetTypeMask()
+	{
+		return 1u << FStorage::GetTypeIndex<TScopeType>();
+	}
+
 #if RDG_ENABLE_TRACE
 	bool bVisited = false;
 #endif
@@ -322,6 +330,9 @@ class TRDGEventScopeGuard
 	FRDGScopeState& State;
 	FRDGScope* const Scope;
 
+	TRDGEventScopeGuard(TRDGEventScopeGuard const&) = delete;
+	TRDGEventScopeGuard(TRDGEventScopeGuard&&) = delete;
+
 public:
 	template <typename... TArgs>
 	inline TRDGEventScopeGuard(FRDGScopeState& State, ERDGScopeFlags Flags, TArgs&&... Args);
@@ -329,7 +340,6 @@ public:
 
 private:
 	static constexpr uint32 TypeMask = 1u << FRDGScope::FStorage::GetTypeIndex<TScopeType>();
-	static inline FRDGScope* Allocate(FRDGScopeState& State, ERDGScopeFlags Flags);
 };
 
 
@@ -341,90 +351,147 @@ private:
  */
 #if RDG_EVENTS
 
-	// Skip expensive string formatting for the relatively common case of no varargs.  We detect this by stringizing the varargs and checking if the string is non-empty (more than just a null terminator).
-	#define RDG_EVENT_NAME(Format, ...) (sizeof(#__VA_ARGS__ "") > 1 ? FRDGEventName(TEXT(Format), ##__VA_ARGS__) : FRDGEventName(1, TEXT(Format)))
+	namespace UE::RHI::Breadcrumbs::Private
+	{
+		template <>
+		struct TValue<FRDGEventName>
+		{
+			static constexpr bool bValidType = true;
 
-	#define RDG_EVENT_SCOPE(GraphBuilder, Format, ...) \
-		TRDGEventScopeGuard<FRDGScope_RHI> PREPROCESSOR_JOIN(__RDG_ScopeRef_,__LINE__)( \
-			(GraphBuilder)                                                              \
-			, ERDGScopeFlags::None                                                      \
-			, FRHIBreadcrumbData(__FILE__, __LINE__)                                    \
-			, RDG_EVENT_NAME(Format, ##__VA_ARGS__)                                     \
-		)
+			TCHAR const* StringPointer;
+			TCHAR StringStorage[FRHIBreadcrumb::MaxLength];
 
-	#if HAS_GPU_STATS
-		#define RDG_EVENT_SCOPE_STAT(GraphBuilder, StatName, Format, ...)                   \
-			TRDGEventScopeGuard<FRDGScope_RHI> PREPROCESSOR_JOIN(__RDG_ScopeRef_,__LINE__)( \
-				(GraphBuilder)                                                              \
-				, ERDGScopeFlags::Stat                                                      \
-				, FRHIBreadcrumbData(                                                       \
-					  __FILE__                                                              \
-					, __LINE__                                                              \
-					, RHI_GPU_STAT_ARGS(StatName)                                           \
-				)                                                                           \
-				, RDG_EVENT_NAME(Format, ##__VA_ARGS__)                                     \
-			)
-	#else
-		#define RDG_EVENT_SCOPE_STAT(GraphBuilder, StatName, Format, ...) \
-			    RDG_EVENT_SCOPE(GraphBuilder, Format, ##__VA_ARGS__)
-	#endif
+			TValue(FRDGEventName const& Value)
+			{
+				if (Value.HasFormattedString())
+				{
+					// We must take a string copy as RHI breadcrumb TValues must be trivially destructable.
+					FCString::Strncpy(StringStorage, Value.GetTCHAR(), UE_ARRAY_COUNT(StringStorage));
+					StringPointer = StringStorage;
+				}
+				else
+				{
+					StringPointer = Value.GetTCHAR();
+				}
+			}
 
-	#define RDG_EVENT_SCOPE_CONDITIONAL(GraphBuilder, Condition, Format, ...)                      \
-		TOptional<TRDGEventScopeGuard<FRDGScope_RHI>> PREPROCESSOR_JOIN(__RDG_ScopeRef_,__LINE__); \
-		do                                                                                         \
-		{                                                                                          \
-			if (Condition)                                                                         \
-			{                                                                                      \
-				PREPROCESSOR_JOIN(__RDG_ScopeRef_,__LINE__).Emplace(                               \
-					(GraphBuilder)                                                                 \
-					, ERDGScopeFlags::None                                                         \
-					, FRHIBreadcrumbData(                                                          \
-						  __FILE__                                                                 \
-						, __LINE__                                                                 \
-					)                                                                              \
-					, RDG_EVENT_NAME(Format, ##__VA_ARGS__)                                        \
-				);                                                                                 \
-			}   \
-		} while (false)
+			struct FConvert
+			{
+				TCHAR const* Inner;
+				FConvert(TValue const& Value)
+					: Inner(Value.StringPointer)
+				{}
+			};
+		};
+
+		// Type to force the selection of the below TValue specialization.
+		struct FRDGFormatTag {};
+
+		template <typename... TValues>
+		struct TValue<std::tuple<FRDGFormatTag&, TValues...>>
+		{
+			static constexpr bool bValidType = true;
 
-	#if HAS_GPU_STATS
-		#define RDG_EVENT_SCOPE_CONDITIONAL_STAT(GraphBuilder, Condition, StatName, Format, ...)       \
-			TOptional<TRDGEventScopeGuard<FRDGScope_RHI>> PREPROCESSOR_JOIN(__RDG_ScopeRef_,__LINE__); \
-			do                                                                                         \
-			{                                                                                          \
-				if (Condition)                                                                         \
-				{                                                                                      \
-					PREPROCESSOR_JOIN(__RDG_ScopeRef_,__LINE__).Emplace(                               \
-						(GraphBuilder)                                                                 \
-						, ERDGScopeFlags::Stat                                                         \
-						, FRHIBreadcrumbData(                                                          \
-							  __FILE__                                                                 \
-							, __LINE__                                                                 \
-							, RHI_GPU_STAT_ARGS(StatName)                                              \
-						)                                                                              \
-						, RDG_EVENT_NAME(Format, ##__VA_ARGS__)                                        \
-					);                                                                                 \
-				}   \
-			} while (false)
-	#else
-		#define RDG_EVENT_SCOPE_CONDITIONAL_STAT(GraphBuilder, Condition, StatName, Format, ...) \
-			    RDG_EVENT_SCOPE_CONDITIONAL(GraphBuilder, Condition, Format, ##__VA_ARGS__)
-	#endif
+			TCHAR Value[FRHIBreadcrumb::MaxLength];
+
+			TValue(std::tuple<FRDGFormatTag&, TValues...> const& Wrapper)
+			{
+				std::apply([this](FRDGFormatTag&, auto&& FormatString, auto&&... Values)
+				{
+					FCString::Snprintf(Value, UE_ARRAY_COUNT(Value), FormatString, std::forward<decltype(Values)>(Values)...);
+				}, Wrapper);
+			}
+
+			struct FConvert
+			{
+				TCHAR const* Inner;
+				FConvert(TValue const& Value)
+					: Inner(Value.Value)
+				{}
+			};
+		};
+	}
+
+	// Returns an RHI breadcrumb initializer appropriate for the given varargs.
+	#define RDG_BREADCRUMB_DESC_FORWARD_VALUES(StaticName, FormatString, GPUStatArgs)					\
+		[&](UE::RHI::Breadcrumbs::Private::FRDGFormatTag&& Tag, auto&&... Values) -> auto				\
+		{																								\
+			using namespace UE::RHI::Breadcrumbs::Private;												\
+			if constexpr (RDG_EVENTS <= RDG_EVENTS_STRING_REF)											\
+			{																							\
+				/* RDG events are only capturing the format string, no args. */							\
+				return RHI_BREADCRUMB_DESC_FORWARD_VALUES(StaticName, nullptr, GPUStatArgs)();			\
+			}																							\
+			else if constexpr (std::is_same_v<TRemoveCVRef<decltype(FormatString)>, FRDGEventName>)		\
+			{																							\
+				/* Single FRDGEventName */																\
+				static_assert(sizeof...(Values) == 0, "Unexpected additional arguments");				\
+				return RHI_BREADCRUMB_DESC_FORWARD_VALUES(StaticName, FForceNoSprintf(), GPUStatArgs)(	\
+					std::forward<decltype(FormatString)>(FormatString)									\
+				);																						\
+			}																							\
+			else if constexpr (TIsValidArgs<decltype(Values)...>::Value)								\
+			{																							\
+				/* Varargs are compatible with RHI breadcrumbs. Forward the args directly. */			\
+				return RHI_BREADCRUMB_DESC_FORWARD_VALUES(StaticName, FormatString, GPUStatArgs)(		\
+					std::forward<decltype(Values)>(Values)...											\
+				);																						\
+			}																							\
+			else																						\
+			{																							\
+				/* Args are not compatible with RHI breadcrumbs. Do the snprintf now */					\
+				return RHI_BREADCRUMB_DESC_COPY_VALUES(StaticName, FForceNoSprintf(), GPUStatArgs)(		\
+					std::forward_as_tuple(																\
+						  Tag																			\
+						, FormatString																	\
+						, std::forward<decltype(Values)>(Values)...										\
+					)																					\
+				);																						\
+			}																							\
+		}
 
-	// The 'Final' version disables any further child scopes or pass events. It is intended
-	// to group overlapping passes as events can disable overlap on certain GPUs.
-	#define RDG_EVENT_SCOPE_FINAL(GraphBuilder, Format, ...)                                       \
-		TOptional<TRDGEventScopeGuard<FRDGScope_RHI>> PREPROCESSOR_JOIN(__RDG_ScopeRef_,__LINE__); \
-		do                                                                                         \
-		{                                                                                          \
-			PREPROCESSOR_JOIN(__RDG_ScopeRef_,__LINE__).Emplace(                                   \
-				(GraphBuilder)                                                                     \
-				, ERDGScopeFlags::Final                                                            \
-				, FRHIBreadcrumbData(__FILE__, __LINE__)                                           \
-				, RDG_EVENT_NAME(Format, ##__VA_ARGS__)                                            \
-			);                                                                                     \
+	#define RDG_EVENT_SCOPE_CONSTRUCT(ObjectName, GraphBuilder, Condition, ScopeFlags, GPUStatArgs, StaticName, FormatString, ...)	\
+		do																															\
+		{																															\
+			if ((Condition) && ((GraphBuilder).ShouldAllocScope((ObjectName), ScopeFlags)))											\
+			{																														\
+				(ObjectName).Emplace(																								\
+					  (GraphBuilder)																								\
+					, ScopeFlags																									\
+					, RDG_BREADCRUMB_DESC_FORWARD_VALUES(																			\
+						  StaticName 																								\
+						, FormatString																								\
+						, GPUStatArgs																								\
+					)(UE::RHI::Breadcrumbs::Private::FRDGFormatTag(), ##__VA_ARGS__)												\
+				);																													\
+			}																														\
 		} while (false)
 
+	#define RDG_EVENT_SCOPE_IMPL(GraphBuilder, Condition, ScopeFlags, GPUStatArgs, StaticName, FormatString, ...)	\
+		TOptional<TRDGEventScopeGuard<FRDGScope_RHI>> PREPROCESSOR_JOIN(__RDG_ScopeRef_, __LINE__);					\
+		RDG_EVENT_SCOPE_CONSTRUCT(																					\
+			  PREPROCESSOR_JOIN(__RDG_ScopeRef_, __LINE__)															\
+			, GraphBuilder																							\
+			, Condition																								\
+			, ScopeFlags																							\
+			, GPUStatArgs																							\
+			, StaticName																							\
+			, FormatString																							\
+			, ##__VA_ARGS__																							\
+		)
+
+	// Skip expensive string formatting for the relatively common case of no varargs.  We detect this by stringizing the varargs and checking if the string is non-empty (more than just a null terminator).
+	#define RDG_EVENT_NAME(Format, ...) (sizeof(#__VA_ARGS__ "") > 1 ? FRDGEventName(TEXT(Format), ##__VA_ARGS__) : FRDGEventName(1, TEXT(Format)))
+
+	#define RDG_EVENT_SCOPE(                 GraphBuilder,                      Format, ...) RDG_EVENT_SCOPE_IMPL(GraphBuilder,      true, ERDGScopeFlags::None , RHI_GPU_STAT_ARGS_NONE	 , TEXT(Format)    , TEXT(Format), ##__VA_ARGS__)
+	#define RDG_EVENT_SCOPE_STAT(            GraphBuilder,            StatName, Format, ...) RDG_EVENT_SCOPE_IMPL(GraphBuilder,      true, ERDGScopeFlags::Stat , RHI_GPU_STAT_ARGS(StatName), TEXT(Format)    , TEXT(Format), ##__VA_ARGS__)
+	#define RDG_EVENT_SCOPE_CONDITIONAL(     GraphBuilder, Condition,           Format, ...) RDG_EVENT_SCOPE_IMPL(GraphBuilder, Condition, ERDGScopeFlags::None , RHI_GPU_STAT_ARGS_NONE	 , TEXT(Format)    , TEXT(Format), ##__VA_ARGS__)
+	#define RDG_EVENT_SCOPE_CONDITIONAL_STAT(GraphBuilder, Condition, StatName, Format, ...) RDG_EVENT_SCOPE_IMPL(GraphBuilder, Condition, ERDGScopeFlags::Stat , RHI_GPU_STAT_ARGS(StatName), TEXT(Format)    , TEXT(Format), ##__VA_ARGS__)
+	// The 'Final' version disables any further child scopes or pass events. It is intended to group overlapping passes as events can disable overlap on certain GPUs.
+	#define RDG_EVENT_SCOPE_FINAL(           GraphBuilder,                      Format, ...) RDG_EVENT_SCOPE_IMPL(GraphBuilder,      true, ERDGScopeFlags::Final, RHI_GPU_STAT_ARGS_NONE	 , TEXT(Format)    , TEXT(Format), ##__VA_ARGS__)
+	// Used in places which have an existing FRDGEventName, e.g. RDG pass name scopes. Prefer to use the other RDG scope macros instead.
+	#define RDG_EVENT_SCOPE_CONDITIONAL_NAME(GraphBuilder, Condition,           EventName  ) RDG_EVENT_SCOPE_IMPL(GraphBuilder, Condition, ERDGScopeFlags::None , RHI_GPU_STAT_ARGS_NONE     , TEXT("RDGEvent"), EventName)
+
 #else
 
 	#define RDG_EVENT_NAME(...) FRDGEventName()
@@ -434,6 +501,7 @@ private:
 	#define RDG_EVENT_SCOPE_CONDITIONAL(...)      do { } while (false)
 	#define RDG_EVENT_SCOPE_CONDITIONAL_STAT(...) do { } while (false)
 	#define RDG_EVENT_SCOPE_FINAL(...)            do { } while (false)
+	#define RDG_EVENT_SCOPE_CONDITIONAL_NAME(...) do { } while (false)
 
 #endif
 
@@ -499,6 +567,24 @@ protected:
 
 	} ScopeState;
 	
+	struct
+	{
+		// Allocator for all root graph allocations on the graph builder thread.
+		FRDGAllocator Root;
+
+		// Allocator for async pass and parallel execute setup.
+		FRDGAllocator Task;
+
+		// Allocator for all allocations related to states / transitions.
+		FRDGAllocator Transition;
+
+		int32 GetByteCount() const
+		{
+			return Root.GetByteCount() + Task.GetByteCount() + Transition.GetByteCount();
+		}
+
+	} Allocators;
+
 public:
 	/** The RHI command list used for the render graph. */
 	FRHICommandListImmediate& RHICmdList;
@@ -546,6 +632,9 @@ public:
 		return ScopeState.ScopeMode != ERDGScopeMode::Disabled;
 	}
 
+	template <typename TScopeType>
+	inline bool ShouldAllocScope(TOptional<TRDGEventScopeGuard<TScopeType>> const&, ERDGScopeFlags Flags) const;
+
 	template <typename TScopeType>
 	friend class TRDGEventScopeGuard;
 	friend FRDGScope_Budget;
diff --git a/Engine/Source/Runtime/RenderCore/Public/RenderGraphEvent.inl b/Engine/Source/Runtime/RenderCore/Public/RenderGraphEvent.inl
--- a/Engine/Source/Runtime/RenderCore/Public/RenderGraphEvent.inl
+++ b/Engine/Source/Runtime/RenderCore/Public/RenderGraphEvent.inl
@@ -100,13 +100,14 @@ inline void FRDGScope_GPU::EndGPU(FRHIComputeCommandList& RHICmdList)
 
 #if RDG_EVENTS
 
-	inline FRDGScope_RHI::FRDGScope_RHI(FRDGScopeState& State, FRHIBreadcrumbData&& Data, FRDGEventName&& Name)
-		: Name(MoveTemp(Name))
-	#if WITH_RHI_BREADCRUMBS
-		, Node(Name.AllocBreadcrumb(MoveTemp(Data), State.GetBreadcrumbAllocator()))
-	#endif
+	template <typename TDesc, typename... TValues>
+	inline FRDGScope_RHI::FRDGScope_RHI(FRDGScopeState& State, TRHIBreadcrumbInitializer<TDesc, TValues...>&& Args)
+		: FRDGScope_RHI(State, State.GetBreadcrumbAllocator().AllocBreadcrumb(MoveTemp(Args)))
+	{}
+
+	inline FRDGScope_RHI::FRDGScope_RHI(FRDGScopeState& State, FRHIBreadcrumbNode* Node)
+		: Node(Node)
 	{
-	#if WITH_RHI_BREADCRUMBS
 		if (Node)
 		{
 			Node->SetParent(State.CurrentBreadcrumbRef);
@@ -119,7 +120,6 @@ inline void FRDGScope_GPU::EndGPU(FRHIComputeCommandList& RHICmdList)
 				State.LocalBreadcrumbList.Append(Node);
 			}
 		}
-	#endif
 	}
 
 	inline void FRDGScope_RHI::ImmediateEnd(FRDGScopeState& State)
@@ -139,17 +139,22 @@ template <typename TScopeType>
 template <typename... TArgs>
 inline TRDGEventScopeGuard<TScopeType>::TRDGEventScopeGuard(FRDGScopeState& State, ERDGScopeFlags Flags, TArgs&&... Args)
 	: State(State)
-	, Scope(Allocate(State, Flags))
+	, Scope(State.Allocators.Root.Alloc<FRDGScope>(State.ScopeState.Current))
 {
-	if (Scope)
+	if (EnumHasAnyFlags(Flags, ERDGScopeFlags::Final))
 	{
-		Scope->Impl.Emplace<TScopeType>(State, Forward<TArgs>(Args)...);
+		// Mask off any nested scopes of the same type
+		State.ScopeState.Mask |= FRDGScope::GetTypeMask<TScopeType>();
+	}
 
-		if (State.ScopeState.bImmediate)
-		{
-			Scope->BeginCPU(State.RHICmdList, false);
-			Scope->BeginGPU(State.RHICmdList);
-		}
+	State.ScopeState.Current = Scope;
+
+	Scope->Impl.Emplace<TScopeType>(State, Forward<TArgs>(Args)...);
+
+	if (State.ScopeState.bImmediate)
+	{
+		Scope->BeginCPU(State.RHICmdList, false);
+		Scope->BeginGPU(State.RHICmdList);
 	}
 }
 
@@ -166,32 +171,23 @@ inline TRDGEventScopeGuard<TScopeType>::~TRDGEventScopeGuard()
 
 		Scope->ImmediateEnd(State);
 
-		State.ScopeState.Mask &= ~(TypeMask);
+		State.ScopeState.Mask &= ~(FRDGScope::GetTypeMask<TScopeType>());
 		State.ScopeState.Current = State.ScopeState.Current->Parent;
 	}
 }
 
 template <typename TScopeType>
-inline FRDGScope* TRDGEventScopeGuard<TScopeType>::Allocate(FRDGScopeState& State, ERDGScopeFlags Flags)
+inline bool FRDGScopeState::ShouldAllocScope(TOptional<TRDGEventScopeGuard<TScopeType>> const&, ERDGScopeFlags Flags) const
 {
-	if (State.ScopeState.ScopeMode == ERDGScopeMode::Disabled && !EnumHasAnyFlags(Flags, ERDGScopeFlags::AlwaysEnable))
+	if (ScopeState.ScopeMode == ERDGScopeMode::Disabled && !EnumHasAnyFlags(Flags, ERDGScopeFlags::AlwaysEnable))
 	{
-		return nullptr;
+		return false;
 	}
 
-	if (State.ScopeState.ScopeMode == ERDGScopeMode::TopLevelOnly && (State.ScopeState.Mask & TypeMask))
+	if (ScopeState.ScopeMode == ERDGScopeMode::TopLevelOnly && (ScopeState.Mask & FRDGScope::GetTypeMask<TScopeType>()))
 	{
-		return nullptr;
+		return false;
 	}
 
-	if (EnumHasAnyFlags(Flags, ERDGScopeFlags::Final))
-	{
-		// Mask off any nested scopes of the same type
-		State.ScopeState.Mask |= TypeMask;
-	}
-
-	FRDGScope* Scope = FRDGAllocator::GetTLS().Alloc<FRDGScope>(State.ScopeState.Current);
-	State.ScopeState.Current = Scope;
-
-	return Scope;
+	return true;
 }
diff --git a/Engine/Source/Runtime/Renderer/Private/InstanceCulling/InstanceCullingContext.cpp b/Engine/Source/Runtime/Renderer/Private/InstanceCulling/InstanceCullingContext.cpp
--- a/Engine/Source/Runtime/Renderer/Private/InstanceCulling/InstanceCullingContext.cpp
+++ b/Engine/Source/Runtime/Renderer/Private/InstanceCulling/InstanceCullingContext.cpp
@@ -1718,7 +1718,7 @@ void FInstanceCullingContext::SubmitDrawCommands(
 			SceneArgs.PrimitiveIdOffset = OverrideArgs.InstanceDataByteOffset + DrawCommandInfo.InstanceDataByteOffset;
 
 #if WANTS_DRAW_MESH_EVENTS
-			RHI_BREADCRUMB_EVENT_CONDITIONAL(RHICmdList, GShowMaterialDrawEvents != 0, "%s %s (%u instances)"
+			RHI_BREADCRUMB_EVENT_CONDITIONAL_F(RHICmdList, GShowMaterialDrawEvents != 0, "MaterialDrawEvent", "%s %s (%u instances)"
 				, VisibleMeshDrawCommand.MeshDrawCommand->GetDebugData().MaterialRenderProxy->GetMaterialName()
 				, VisibleMeshDrawCommand.MeshDrawCommand->GetDebugData().ResourceName
 				, VisibleMeshDrawCommand.MeshDrawCommand->NumInstances * InstanceFactor
diff --git a/Engine/Source/Runtime/Renderer/Private/MeshPassProcessor.cpp b/Engine/Source/Runtime/Renderer/Private/MeshPassProcessor.cpp
--- a/Engine/Source/Runtime/Renderer/Private/MeshPassProcessor.cpp
+++ b/Engine/Source/Runtime/Renderer/Private/MeshPassProcessor.cpp
@@ -1561,7 +1561,7 @@ bool FMeshDrawCommand::SubmitDraw(
 	FMeshDrawCommandStateCache& RESTRICT StateCache)
 {
 #if WANTS_DRAW_MESH_EVENTS
-	RHI_BREADCRUMB_EVENT_CONDITIONAL(RHICmdList, GShowMaterialDrawEvents != 0, "%s %s (%u instances)"
+	RHI_BREADCRUMB_EVENT_CONDITIONAL_F(RHICmdList, GShowMaterialDrawEvents != 0, "MaterialDrawEvent", "%s %s (%u instances)"
 		, MeshDrawCommand.DebugData.MaterialRenderProxy->GetMaterialName()
 		, MeshDrawCommand.DebugData.ResourceName
 		, MeshDrawCommand.NumInstances * InstanceFactor
diff --git a/Engine/Source/Runtime/Renderer/Private/SceneRenderBuilder.cpp b/Engine/Source/Runtime/Renderer/Private/SceneRenderBuilder.cpp
--- a/Engine/Source/Runtime/Renderer/Private/SceneRenderBuilder.cpp
+++ b/Engine/Source/Runtime/Renderer/Private/SceneRenderBuilder.cpp
@@ -424,17 +424,15 @@ class FSceneRenderProcessor
 #if WITH_RHI_BREADCRUMBS
 			if (Group && IsLocationActive(Location))
 			{
-				FRHIBreadcrumbData BreadcrumbData(__FILE__, __LINE__);
-
 				// User specified an explicit group name.
 				if (!Group->Name.IsEmpty())
 				{
-					Event.Emplace(RHICmdList, MoveTemp(BreadcrumbData), TEXT("%s"), Group->Name);
+					Event.Emplace(RHICmdList, RHI_BREADCRUMB_DESC_FORWARD_VALUES(TEXT("SceneRenderGroup"), TEXT("%s"), RHI_GPU_STAT_ARGS_NONE)(Group->Name));
 				}
 				// User didn't specify a name, but we have more than one renderer, so a group event is useful.
 				else if (!Group->RenderNodes.IsEmpty() && Group->RenderNodes.GetHead() != Group->RenderNodes.GetTail())
 				{
-					Event.Emplace(RHICmdList, MoveTemp(BreadcrumbData), TEXT("SceneRenderGroup"));
+					Event.Emplace(RHICmdList, RHI_BREADCRUMB_DESC_FORWARD_VALUES(TEXT("SceneRenderGroup"), nullptr, RHI_GPU_STAT_ARGS_NONE)());
 				}
 			}
 #endif
@@ -865,8 +863,8 @@ void FSceneRenderProcessor::Execute()
 					RenderState.GroupEvent.End(RHICmdList, FGroupEventLocation::SceneRenderCommand);
 				};
 
-				RHI_BREADCRUMB_EVENT_CONDITIONAL(RHICmdList, !RenderNode.Name.IsEmpty(), "%s", RenderNode.Name);
-				RHI_BREADCRUMB_EVENT_CONDITIONAL(RHICmdList,  RenderNode.Name.IsEmpty(), "SceneRender");
+				RHI_BREADCRUMB_EVENT_CONDITIONAL_F(RHICmdList, !RenderNode.Name.IsEmpty(), "SceneRender", "SceneRender - %s", RenderNode.Name);
+				RHI_BREADCRUMB_EVENT_CONDITIONAL(RHICmdList, RenderNode.Name.IsEmpty(), "SceneRender");
 
 				FSceneRenderer& Renderer = *RenderNode.Renderer;
 				FDeferredUpdateResource::UpdateResources(RHICmdList);
diff --git a/Engine/Source/Runtime/Renderer/Private/ShadowDepthRendering.cpp b/Engine/Source/Runtime/Renderer/Private/ShadowDepthRendering.cpp
--- a/Engine/Source/Runtime/Renderer/Private/ShadowDepthRendering.cpp
+++ b/Engine/Source/Runtime/Renderer/Private/ShadowDepthRendering.cpp
@@ -1612,7 +1612,7 @@ void FSceneRenderer::RenderShadowDepthMapAtlases(FRDGBuilder& GraphBuilder)
 				FString LightNameWithLevel;
 				GetLightNameForDrawEvent(CurrentLightForDrawEvent, LightNameWithLevel);
 
-				RDGScope.Emplace(GraphBuilder, ERDGScopeFlags::None, FRHIBreadcrumbData(__FILE__, __LINE__), RDG_EVENT_NAME("%s", *LightNameWithLevel));
+				RDG_EVENT_SCOPE_CONSTRUCT(RDGScope, GraphBuilder, true, ERDGScopeFlags::None, RHI_GPU_STAT_ARGS_NONE, TEXT("RenderDepthForLight"), TEXT("%s"), *LightNameWithLevel);
 			}
 		#endif
 		};
diff --git a/Engine/Source/Runtime/VulkanRHI/Private/VulkanRHI.cpp b/Engine/Source/Runtime/VulkanRHI/Private/VulkanRHI.cpp
--- a/Engine/Source/Runtime/VulkanRHI/Private/VulkanRHI.cpp
+++ b/Engine/Source/Runtime/VulkanRHI/Private/VulkanRHI.cpp
@@ -1257,7 +1257,7 @@ void FVulkanCommandListContext::RHIEndDrawingViewport(FRHIViewport* ViewportRHI,
 		{
 			if (!NameStr)
 			{
-				NameStr = Breadcrumb->Name.GetTCHAR(Buffer);
+				NameStr = Breadcrumb->GetTCHAR(Buffer);
 			}
 			return NameStr;
 		};
diff --git a/Engine/Source/Runtime/Windows/D3D11RHI/Private/D3D11Device.cpp b/Engine/Source/Runtime/Windows/D3D11RHI/Private/D3D11Device.cpp
--- a/Engine/Source/Runtime/Windows/D3D11RHI/Private/D3D11Device.cpp
+++ b/Engine/Source/Runtime/Windows/D3D11RHI/Private/D3D11Device.cpp
@@ -278,7 +278,7 @@ void FD3D11DynamicRHI::Shutdown()
 		#else
 			// @todo dev-pr avoid TCHAR -> ANSI conversion
 			FRHIBreadcrumb::FBuffer Buffer;
-			TCHAR const* Name = Breadcrumb->Name.GetTCHAR(Buffer);
+			TCHAR const* Name = Breadcrumb->GetTCHAR(Buffer);
 
 			GPUProfilingData.PushEvent(Name, FColor::White);
 		#endif
