RHI Validation, printing debug information for Executing pipeline in Acquire and Discard

diff --git a/Engine/Source/Runtime/RHI/Private/RHIValidation.cpp b/Engine/Source/Runtime/RHI/Private/RHIValidation.cpp
--- a/Engine/Source/Runtime/RHI/Private/RHIValidation.cpp
+++ b/Engine/Source/Runtime/RHI/Private/RHIValidation.cpp
@@ -2086,7 +2086,7 @@ namespace RHIValidation
 		return Trace;
 	}
 
-	void FTransientState::Acquire(FResource* Resource, void* CreateTrace)
+	void FTransientState::Acquire(FResource* Resource, void* CreateTrace, ERHIPipeline ExecutingPipeline)
 	{
 		RHI_VALIDATION_CHECK(bTransient, *GetReasonString_AcquireNonTransient(Resource));
 		RHI_VALIDATION_CHECK(Status == EStatus::None, *GetReasonString_DuplicateAcquireTransient(Resource, AcquireBacktrace, CreateTrace));
@@ -2101,11 +2101,11 @@ namespace RHIValidation
 
 		if (Resource->LoggingMode != ELoggingMode::None)
 		{
-			Log(Resource, {}, CreateTrace, TEXT("Acquire"), TEXT("Acquire"), TEXT("Transient Acquire"));
+			Log(Resource, {}, CreateTrace, TEXT("Acquire"), TEXT("Acquire"), *FString::Printf(TEXT("Transient Acquire, Executing Pipeline : %s"), *GetRHIPipelineName(ExecutingPipeline)));
 		}
 	}
 
-	void FTransientState::Discard(FResource* Resource, void* CreateTrace, ERHIPipeline DiscardPipelines)
+	void FTransientState::Discard(FResource* Resource, void* CreateTrace, ERHIPipeline DiscardPipelines, ERHIPipeline ExecutingPipeline)
 	{
 		RHI_VALIDATION_CHECK(bTransient, *GetReasonString_DiscardNonTransient(Resource));
 		RHI_VALIDATION_CHECK(Status != EStatus::None, *GetReasonString_DiscardWithoutAcquireTransient(Resource, CreateTrace));
@@ -2122,7 +2122,7 @@ namespace RHIValidation
 
 			if (Resource->LoggingMode != ELoggingMode::None)
 			{
-				Log(Resource, {}, CreateTrace, TEXT("Discard"), TEXT("Discard"), TEXT("Transient Discard"));
+				Log(Resource, {}, CreateTrace, TEXT("Discard"), TEXT("Discard"), *FString::Printf(TEXT("Transient Discard, Executing Pipeline : %s"), *GetRHIPipelineName(ExecutingPipeline)) );
 			}
 		}
 	}
@@ -2199,7 +2199,7 @@ namespace RHIValidation
 
 			if (TargetState.Access == ERHIAccess::Discard)
 			{
-				Resource->TransientState.Discard(Resource, CreateTrace, CurrentStateFromRHI.Pipelines);
+				Resource->TransientState.Discard(Resource, CreateTrace, CurrentStateFromRHI.Pipelines, ExecutingPipeline);
 			}
 		}
 
@@ -2659,7 +2659,7 @@ namespace RHIValidation
 			break;
 
 		case EOpType::AcquireTransient:
-			Data_AcquireTransient.Resource->TransientState.Acquire(Data_AcquireTransient.Resource, Data_AcquireTransient.CreateBacktrace);
+			Data_AcquireTransient.Resource->TransientState.Acquire(Data_AcquireTransient.Resource, Data_AcquireTransient.CreateBacktrace, Queue.Pipeline);
 			Data_AcquireTransient.Resource->ReleaseOpRef();
 			break;
 
diff --git a/Engine/Source/Runtime/RHI/Public/RHIValidationCommon.h b/Engine/Source/Runtime/RHI/Public/RHIValidationCommon.h
--- a/Engine/Source/Runtime/RHI/Public/RHIValidationCommon.h
+++ b/Engine/Source/Runtime/RHI/Public/RHIValidationCommon.h
@@ -311,8 +311,8 @@ namespace RHIValidation
 		FORCEINLINE bool IsAcquired() const { return Status == EStatus::Acquired; }
 		FORCEINLINE bool IsDiscarded() const { return Status == EStatus::Discarded; }
 
-		void Acquire(FResource* Resource, void* CreateTrace);
-		void Discard(FResource* Resource, void* CreateTrace, ERHIPipeline DiscardPipelines);
+		void Acquire(FResource* Resource, void* CreateTrace, ERHIPipeline ExecutingPipeline);
+		void Discard(FResource* Resource, void* CreateTrace, ERHIPipeline DiscardPipelines, ERHIPipeline ExecutingPipeline);
 
 		static void AliasingOverlap(FResource* ResourceBefore, FResource* ResourceAfter, void* CreateTrace);
 	};
