Fixed NeedsDiscardToGraphicsWorkaround to not allow overlap of resources that can't actually discard to graphics.

diff --git a/Engine/Source/Runtime/RHI/Public/RHITransientResourceAllocator.h b/Engine/Source/Runtime/RHI/Public/RHITransientResourceAllocator.h
--- a/Engine/Source/Runtime/RHI/Public/RHITransientResourceAllocator.h
+++ b/Engine/Source/Runtime/RHI/Public/RHITransientResourceAllocator.h
@@ -56,7 +56,14 @@ public:
 	// Returns whether two regions described by the discard and acquire fences contain each other. If they do, that means the memory would being used by both pipes simultaneously and cannot be aliased.
 	static bool Contains(const FRHITransientAllocationFences& Discard, const FRHITransientAllocationFences& Acquire)
 	{
-		return Contains(Discard.GraphicsForkJoin, Acquire.Graphics) || Contains(Acquire.GraphicsForkJoin, Discard.Graphics);
+		return
+			Contains(Discard.GraphicsForkJoin, Acquire.Graphics) ||
+			Contains(Acquire.GraphicsForkJoin, Discard.Graphics) ||
+			// We have to discard on the graphics pipe, which means we can't alias with async compute resources that have a graphics fork pass earlier than the discard join pass.
+			(GRHIGlobals.NeedsTransientDiscardOnGraphicsWorkaround &&
+				Discard.GetPipelines() == ERHIPipeline::All &&
+				Acquire.GetPipelines() == ERHIPipeline::AsyncCompute &&
+				Acquire.GraphicsForkJoin.Min < Discard.GraphicsForkJoin.Max);
 	}
 
 	FRHITransientAllocationFences() = default;
