[RHI]
* Changed "auto Enumerate(ERHIPipeline Pipeline) const" which returns an internal FOuter class to "FOuter Enumerate(ERHIPipeline Pipeline) const;" and moved FOuter class out of function. This exposed a bug in c++modules/header units that microsoft has no immediate plan to fix.

diff --git a/Engine/Source/Runtime/RHI/Public/RHIBreadcrumbs.h b/Engine/Source/Runtime/RHI/Public/RHIBreadcrumbs.h
--- a/Engine/Source/Runtime/RHI/Public/RHIBreadcrumbs.h
+++ b/Engine/Source/Runtime/RHI/Public/RHIBreadcrumbs.h
@@ -523,101 +523,104 @@ namespace UE::RHI::GPUProfiler
 			}
 		}
 
-		auto Enumerate(ERHIPipeline Pipeline) const
+		class FOuter;
+		FOuter Enumerate(ERHIPipeline Pipeline) const;
+
+		operator bool() const { return First != nullptr; }
+
+		bool operator == (FRHIBreadcrumbRange const& RHS) const
 		{
-			// Either both must be null, or both must be non-null
-			check(!First == !Last);
+			return First == RHS.First && Last == RHS.Last;
+		}
 
-			class FOuter
-			{
-				FRHIBreadcrumbRange const Range;
-				ERHIPipeline const Pipeline;
+		bool operator != (FRHIBreadcrumbRange const& RHS) const
+		{
+			return !(*this == RHS);
+		}
 
-			public:
-				FOuter(FRHIBreadcrumbRange const& Range, ERHIPipeline Pipeline)
-					: Range(Range)
-					, Pipeline(Pipeline)
-				{}
+		friend uint32 GetTypeHash(FRHIBreadcrumbRange const& Range)
+		{
+			return HashCombineFast(GetTypeHash(Range.First), GetTypeHash(Range.Last));
+		}
+	};
+	
+	class FRHIBreadcrumbRange::FOuter
+	{
+		FRHIBreadcrumbRange const Range;
+		ERHIPipeline const Pipeline;
 
-				auto begin() const
-				{
-					struct FIterator
-					{
-						FRHIBreadcrumbNode* Current;
-						FRHIBreadcrumbNode* const Last;
+	public:
+		FOuter(FRHIBreadcrumbRange const& Range, ERHIPipeline Pipeline)
+			: Range(Range)
+			, Pipeline(Pipeline)
+		{}
+
+		auto begin() const
+		{
+			struct FIterator
+			{
+				FRHIBreadcrumbNode* Current;
+				FRHIBreadcrumbNode* const Last;
 #if DO_CHECK
-						FRHIBreadcrumbNode* const First;
+				FRHIBreadcrumbNode* const First;
 #endif
-						ERHIPipeline const Pipeline;
-
-						bool operator != (std::nullptr_t) const
-						{
-							return Current != nullptr;
-						}
-
-						FRHIBreadcrumbNode* operator*() const
-						{
-							return Current;
-						}
+				ERHIPipeline const Pipeline;
 
-						FIterator& operator++()
-						{
-							if (Current == Last)
-							{
-								Current = nullptr;
-							}
-							else
-							{
-								FRHIBreadcrumbNode* Next = Current->GetNextPtr(Pipeline);
+				bool operator != (std::nullptr_t) const
+				{
+					return Current != nullptr;
+				}
 
-								// Next should never be null here. When iterating a non-empty range, we should always expect to reach 'Last' rather than nullptr.
-								checkf(Next, TEXT("Nullptr 'Next' breadcrumb found before reaching the 'Last' breadcrumb in the range. (First: 0x%p, Last: 0x%p, Current: 0x%p)"), First, Last, Current);
+				FRHIBreadcrumbNode* operator*() const
+				{
+					return Current;
+				}
 
-								Current = Next;
-							}
+				FIterator& operator++()
+				{
+					if (Current == Last)
+					{
+						Current = nullptr;
+					}
+					else
+					{
+						FRHIBreadcrumbNode* Next = Current->GetNextPtr(Pipeline);
 
-							return *this;
-						}
-					};
+						// Next should never be null here. When iterating a non-empty range, we should always expect to reach 'Last' rather than nullptr.
+						checkf(Next, TEXT("Nullptr 'Next' breadcrumb found before reaching the 'Last' breadcrumb in the range. (First: 0x%p, Last: 0x%p, Current: 0x%p)"), First, Last, Current);
 
-					return FIterator
-					{
-						  Range.First
-						, Range.Last
-					#if DO_CHECK
-						, Range.First
-					#endif
-						, Pipeline
-					};
-				}
+						Current = Next;
+					}
 
-				constexpr std::nullptr_t end() const
-				{
-					return nullptr;
+					return *this;
 				}
 			};
 
-			return FOuter { *this, Pipeline };
+			return FIterator
+			{
+					Range.First
+				, Range.Last
+			#if DO_CHECK
+				, Range.First
+			#endif
+				, Pipeline
+			};
 		}
 
-		operator bool() const { return First != nullptr; }
-
-		bool operator == (FRHIBreadcrumbRange const& RHS) const
+		constexpr std::nullptr_t end() const
 		{
-			return First == RHS.First && Last == RHS.Last;
+			return nullptr;
 		}
+	};
 
-		bool operator != (FRHIBreadcrumbRange const& RHS) const
-		{
-			return !(*this == RHS);
-		}
+	inline FRHIBreadcrumbRange::FOuter FRHIBreadcrumbRange::Enumerate(ERHIPipeline Pipeline) const
+	{
+		// Either both must be null, or both must be non-null
+		check(!First == !Last);
+
+		return FOuter { *this, Pipeline };
+	}
 
-		friend uint32 GetTypeHash(FRHIBreadcrumbRange const& Range)
-		{
-			return HashCombineFast(GetTypeHash(Range.First), GetTypeHash(Range.Last));
-		}
-	};
-	
 	inline void FRHIBreadcrumbNode::SetParent(FRHIBreadcrumbNode* Node)
 	{
 		check(Parent == nullptr || Parent == FRHIBreadcrumbNode::Sentinel);
@@ -911,12 +914,12 @@ namespace UE::RHI::GPUProfiler
 		template <size_t N, class T> struct TFieldTraits<std::tuple<TCHAR const(&)[N], T&&>&&> { using TDescType = std::tuple<TCHAR const(&)[N], TRemoveCVRef<T>>; using TValueType = T; using TValueRef = T&&; static TValueRef ForwardValue(auto&& Arg) { return std::forward<TValueRef>(std::get<1>(Arg)); } };
 
 		// The breadcrumb macros work with both RHI command lists and RHI contexts. This helper retrieves the relevant RHI command list from both types.
-		static inline FRHIComputeCommandList& GetRHICmdList(FRHIComputeCommandList& RHICmdList);
-		static inline FRHIComputeCommandList& GetRHICmdList(IRHIComputeContext    & RHIContext);
+		inline FRHIComputeCommandList& GetRHICmdList(FRHIComputeCommandList& RHICmdList);
+		inline FRHIComputeCommandList& GetRHICmdList(IRHIComputeContext    & RHIContext);
 
 		// Returns the full path string for the breadcrumb currently at the top of the CPU stack,
 		// for either RHI command lists or RHI contexts. Used by the breadcrumb check macros.
-		static inline FString GetSafeBreadcrumbPath(auto&& RHICmdList_Or_RHIContext);
+		inline FString GetSafeBreadcrumbPath(auto&& RHICmdList_Or_RHIContext);
 
 		template <size_t Size, bool bHasVarargs>
 		struct TFormatString
