Add support for Parallel rendering for platforms implementing render passes

- Added SubCommandLists which are inserted into a CommandList with InsertParallelRenderPass
- Added platform RHIBegin/EndParallelRenderPass for platforms to implement parent commands
- RHIGetParallelCommandContext/RHIFinalizeParallelContext has a default implementation for platforms not needing to implement renderpasses, by default these call Begin/EndRenderPass

diff --git a/Engine/Source/Runtime/D3D12RHI/Private/D3D12RHI.cpp b/Engine/Source/Runtime/D3D12RHI/Private/D3D12RHI.cpp
--- a/Engine/Source/Runtime/D3D12RHI/Private/D3D12RHI.cpp
+++ b/Engine/Source/Runtime/D3D12RHI/Private/D3D12RHI.cpp
@@ -228,6 +228,7 @@ FD3D12DynamicRHI::FD3D12DynamicRHI(const TArray<TSharedPtr<FD3D12Adapter>>& Chos
 	GRHISupportsRHIThread = true;
 
 	GRHISupportsParallelRHIExecute = true;
+	GRHISupportsParallelRenderPasses = true;
 
 	GRHISupportsRawViewsForAnyBuffer = true;
 
diff --git a/Engine/Source/Runtime/NullDrv/Public/NullRHI.h b/Engine/Source/Runtime/NullDrv/Public/NullRHI.h
--- a/Engine/Source/Runtime/NullDrv/Public/NullRHI.h
+++ b/Engine/Source/Runtime/NullDrv/Public/NullRHI.h
@@ -558,10 +558,21 @@ public:
 	{
 		return nullptr;
 	}
+	
+	virtual IRHIComputeContext* RHIGetParallelCommandContext(FRHIParallelRenderPassInfo const& ParallelRenderPass, FRHIGPUMask GPUMask) final override
+	{
+		return nullptr;
+	}
 
 	virtual void RHIFinalizeContext(FRHIFinalizeContextArgs&& Args, TRHIPipelineArray<IRHIPlatformCommandList*>& Output) final override
 	{
 	}
+	
+	IRHIPlatformCommandList* RHIFinalizeParallelContext(IRHIComputeContext* Context) override final
+	{
+		// TODO: Implement and test this
+		return nullptr;
+	}
 
 	virtual void RHISubmitCommandLists(FRHISubmitCommandListsArgs&& Args) final override
 	{
diff --git a/Engine/Source/Runtime/RHI/Private/RHICommandList.cpp b/Engine/Source/Runtime/RHI/Private/RHICommandList.cpp
--- a/Engine/Source/Runtime/RHI/Private/RHICommandList.cpp
+++ b/Engine/Source/Runtime/RHI/Private/RHICommandList.cpp
@@ -164,6 +164,58 @@ FRHICommandListBase::~FRHICommandListBase()
 #endif
 }
 
+void FRHICommandListBase::InsertParallelRenderPass_Base(TSharedPtr<FRHIParallelRenderPassInfo> const& InInfo, TArray<FRHISubCommandList*>&& SubCommandLists)
+{
+	ERHIPipeline CurrentPipelines = ActivePipelines;
+
+	bool bRequiresWait = GRHIParallelRHIExecuteChildWait || GRHIParallelRHIExecuteParentWait;
+	
+	// Finish current command list with the start of the parallel RP
+	if(bRequiresWait)
+	{
+		ParallelRenderPassBegin = InInfo;
+	}
+	
+	EnqueueLambda([InInfo](FRHICommandListBase& ExecutingCmdList)
+	{
+		ExecutingCmdList.GetContext().RHIBeginParallelRenderPass(InInfo, InInfo->PassName);
+	});
+
+	FinishRecording();
+
+	// Split the RHICmdList by moving the current commands into a new instance on the heap, and reconstructing 'this'.
+	{
+		FRHICommandListBase* HeapCmdList = new FRHICommandListBase(MoveTemp(*this));
+
+		(this)->~FRHICommandListBase();
+		new (this) FRHICommandListBase(HeapCmdList->PersistentState);
+
+		AttachedCmdLists.Add(HeapCmdList);
+	}
+
+	// Enqueue the children
+	for (FRHISubCommandList* SubCmdList : SubCommandLists)
+	{
+		// All provided sub command lists must have the same parallel RP info struct.
+		check(SubCmdList->SubRenderPassInfo == InInfo);
+		AttachedCmdLists.Add(SubCmdList);
+	}
+
+	// Restore the pipelines we had active
+	ActivatePipelines(CurrentPipelines);
+
+	// Start the new RHICmdList with the end of the parallel render pass.
+	if(bRequiresWait)
+	{	
+		ParallelRenderPassEnd = InInfo;
+	}
+	
+	EnqueueLambda([](FRHICommandListBase& ExecutingCmdList)
+	{
+		ExecutingCmdList.GetContext().RHIEndParallelRenderPass();
+	});
+}
+
 const int32 FRHICommandListBase::GetUsedMemory() const
 {
 	return MemManager.GetByteCount();
@@ -345,10 +397,17 @@ void FRHICommandListBase::ActivatePipelines(ERHIPipeline Pipelines)
 			{
 				if (!Context)
 				{
-					// Need to handle the "immediate" context separately.
-					Context = ExecutingCmdList.AllowParallelTranslate()
-						? GDynamicRHI->RHIGetCommandContext(Pipeline, FRHIGPUMask::All()) // This mask argument specifies which contexts are included in an mGPU redirector (we always want all of them).
-						: ::RHIGetDefaultContext();
+					if (ExecutingCmdList.IsSubCommandList())
+					{
+						Context = GDynamicRHI->RHIGetParallelCommandContext(*ExecutingCmdList.SubRenderPassInfo, FRHIGPUMask::All());
+					}
+					else
+					{
+						// Need to handle the "immediate" context separately.
+						Context = ExecutingCmdList.AllowParallelTranslate()
+							? GDynamicRHI->RHIGetCommandContext(Pipeline, FRHIGPUMask::All()) // This mask argument specifies which contexts are included in an mGPU redirector (we always want all of them).
+							: ::RHIGetDefaultContext();
+					}
 				}
 
 				if (bSinglePipeline)
@@ -361,6 +420,8 @@ void FRHICommandListBase::ActivatePipelines(ERHIPipeline Pipelines)
 
 			case ERHIPipeline::AsyncCompute:
 			{
+				checkf(!ExecutingCmdList.IsSubCommandList(), TEXT("Sub command lists are only allowed to use the graphics pipe."));
+
 				if (!Context)
 				{
 					Context = GDynamicRHI->RHIGetCommandContext(Pipeline, FRHIGPUMask::All()); // This mask argument specifies which contexts are included in an mGPU redirector (we always want all of them).
@@ -698,11 +759,9 @@ FRHICommandListExecutor::FTaskPipe* FRHICommandListExecutor::EnqueueDispatchTask
 	return &DispatchPipe;
 }
 
-FRHICommandListExecutor::FTaskPipe* FRHICommandListExecutor::FTranslateState::EnqueueTranslateTask(FGraphEventArray&& Prereqs, TFunction<void()>&& Lambda)
+FRHICommandListExecutor::FTaskPipe* FRHICommandListExecutor::FTranslateState::GetTranslateTaskPipe(ENamedThreads::Type& NamedThread)
 {
-	// This is called on the dispatch thread
-
-	ENamedThreads::Type NamedThread = ENamedThreads::AnyHiPriThreadHiPriTask;
+	NamedThread = ENamedThreads::AnyHiPriThreadHiPriTask;
 	FTaskPipe* Pipe = &TranslatePipe;
 
 	if (!GRHICommandList.AllowParallel())
@@ -715,6 +774,15 @@ FRHICommandListExecutor::FTaskPipe* FRHICommandListExecutor::FTranslateState::En
 		NamedThread = ENamedThreads::RHIThread;
 		Pipe = &GRHICommandList.RHIThreadPipe;
 	}
+	
+	return Pipe;
+}
+
+FRHICommandListExecutor::FTaskPipe* FRHICommandListExecutor::FTranslateState::EnqueueTranslateTask(FGraphEventArray&& Prereqs, TFunction<void()>&& Lambda)
+{
+	// This is called on the dispatch thread
+	ENamedThreads::Type NamedThread = ENamedThreads::AnyHiPriThreadHiPriTask;
+	FTaskPipe* Pipe = GetTranslateTaskPipe(NamedThread);
 
 	Pipe->Enqueue(NamedThread, MoveTemp(Prereqs), MoveTemp(Lambda));
 	return Pipe;
@@ -732,52 +800,83 @@ FRHICommandListExecutor::FTaskPipe* FRHICommandListExecutor::EnqueueSubmitTask(F
 	return &RHIThreadPipe;
 }
 
-void FRHICommandListExecutor::FSubmitState::FinalizeCurrent()
+FGraphEventRef FRHICommandListExecutor::FSubmitState::FinalizeCurrent()
 {
-	TranslateEvents.Add(CurrentTranslateJob->Finalize());
+	FGraphEventRef Event = CurrentTranslateJob->Finalize();
+	
+	TranslateEvents.Add(Event);
 	CurrentTranslateJob = nullptr;
+	
+	return Event;
 }
 
-void FRHICommandListExecutor::FSubmitState::Dispatch(FRHICommandListBase* CmdList)
+bool FRHICommandListExecutor::FSubmitState::ShouldSplitTranslateJob(FRHICommandListBase* CmdList)
 {
-	check(CmdList->DispatchEvent->IsComplete());
-	CmdList->bAllowParallelTranslate = bAllowParallelTranslate;
-
-#if WITH_RHI_BREADCRUMBS
-	CmdList->CPUBreadcrumbState.bEmitBreadcrumbs = bEmitBreadcrumbs;
-#endif
+	//
+	// Determine if the current translate batch should be closed, i.e.:
+	//    - We've exceeded the threshold number of RHI commands.
+	//    - The next command list requires single-threaded execution, but we're in a parallel batch.
+	//
+	
+	bool bShouldSplitParallel = bAllowSingleParallelCombine
+			? CurrentTranslateJob->bParallel && !CmdList->AllowParallelTranslate() // Only start a new translate job if we were parallel, but the new command list requires single thread.
+			: CurrentTranslateJob->bParallel !=  CmdList->AllowParallelTranslate(); // Always start a new translate job if bParallel is different. Avoids batching parallel work into single thread translates.
+	
+	bool bShouldSplitForThreshold = MaxCommandsPerTranslate  < 0 ? true  :
+									MaxCommandsPerTranslate == 0 ? false :
+									(CurrentTranslateJob->NumCommands + CmdList->NumCommands) > uint32(MaxCommandsPerTranslate);
+	
+	bool bShouldSplitForParentChild = 	(CmdList->IsSubCommandList() && !CurrentTranslateJob->bUsingSubCmdLists) ||
+										(!CmdList->IsSubCommandList() && CurrentTranslateJob->bUsingSubCmdLists);
+	
+	return bShouldSplitParallel || bShouldSplitForThreshold || bShouldSplitForParentChild;
+}
 
-	if (CurrentTranslateJob)
+void FRHICommandListExecutor::FSubmitState::ConditionalSplitTranslateJob(FRHICommandListBase* CmdList)
+{
+	if (CurrentTranslateJob && ShouldSplitTranslateJob(CmdList))
 	{
-		//
-		// Determine if the current translate batch should be closed, i.e.:
-		//    - We've exceeded the threshold number of RHI commands.
-		//    - The next command list requires single-threaded execution, but we're in a parallel batch.
-		//
-		if ((bAllowSingleParallelCombine
-				? CurrentTranslateJob->bParallel && !CmdList->AllowParallelTranslate() // Only start a new translate job if we were parallel, but the new command list requires single thread.
-				: CurrentTranslateJob->bParallel !=  CmdList->AllowParallelTranslate() // Always start a new translate job if bParallel is different. Avoids batching parallel work into single thread translates.
-			) || (
-				MaxCommandsPerTranslate  < 0 ? true  :
-				MaxCommandsPerTranslate == 0 ? false :
-				(CurrentTranslateJob->NumCommands + CmdList->NumCommands) > uint32(MaxCommandsPerTranslate)
-			))
+		bool bAddChildWait = CurrentTranslateJob->bUsingSubCmdLists && GRHIParallelRHIExecuteChildWait; 
+		bool bAddParentWait = !CurrentTranslateJob->bUsingSubCmdLists && CmdList->IsSubCommandList() && GRHIParallelRHIExecuteParentWait;
+		FGraphEventRef Event = FinalizeCurrent();
+		
+		if (bAddChildWait)
+		{
+			ChildGraphEvents.Add(Event);
+		}
+		else if (bAddParentWait)
 		{
-			FinalizeCurrent();
+			BeginGraphEvent = Event;
 		}
 	}
-
+	
 	if (!CurrentTranslateJob)
 	{
 		// Start a new translate job
 		CurrentTranslateJob = TranslateJobs.Emplace_GetRef(MakeUnique<FTranslateState>()).Get();
 		CurrentTranslateJob->bParallel = CmdList->AllowParallelTranslate();
+		CurrentTranslateJob->bUsingSubCmdLists = CmdList->IsSubCommandList();
 	}
 
+	check(CurrentTranslateJob->bShouldFinalize);
+	CurrentTranslateJob->bShouldFinalize = !(GRHIParallelRHIExecuteChildWait && CmdList->ParallelRenderPassBegin.IsValid());
+
 	CurrentTranslateJob->NumCommands += CmdList->NumCommands;
+}
+
+void FRHICommandListExecutor::FSubmitState::Dispatch(FRHICommandListBase* CmdList)
+{
+	check(CmdList->DispatchEvent->IsComplete());
+	CmdList->bAllowParallelTranslate = bAllowParallelTranslate;
+
+#if WITH_RHI_BREADCRUMBS
+	CmdList->CPUBreadcrumbState.bEmitBreadcrumbs = bEmitBreadcrumbs;
+#endif
+
+	ConditionalSplitTranslateJob(CmdList);	
 
 	{
-	#if WITH_RHI_BREADCRUMBS
+#if WITH_RHI_BREADCRUMBS
 		// Fixup unknown breadcrumb parents
 		for (FRHIBreadcrumbNode* Node : CmdList->CPUBreadcrumbState.UnknownParentList.IterateAndUnlink())
 		{
@@ -859,9 +958,9 @@ void FRHICommandListExecutor::FSubmitState::Dispatch(FRHICommandListBase* CmdLis
 			CmdList->CPUBreadcrumbState.Current = CPUFirst;
 			CmdList->PersistentState.LocalBreadcrumb = CPUFirst;
 		}
-	#endif // WITH_RHI_BREADCRUMBS
+#endif // WITH_RHI_BREADCRUMBS
 
-	#if HAS_GPU_STATS
+#if HAS_GPU_STATS
 		CmdList->InitialDrawStatsCategory = GRHICommandList.CurrentDrawStatsCategory;
 
 		if (CmdList->PersistentState.CurrentDrawStatsCategory.IsSet())
@@ -870,7 +969,7 @@ void FRHICommandListExecutor::FSubmitState::Dispatch(FRHICommandListBase* CmdLis
 		}
 
 		CmdList->PersistentState.CurrentDrawStatsCategory = CmdList->InitialDrawStatsCategory;
-	#endif
+#endif
 
 		FGraphEventArray Prereqs;
 		if (!CmdList->AllowParallelTranslate())
@@ -878,6 +977,27 @@ void FRHICommandListExecutor::FSubmitState::Dispatch(FRHICommandListBase* CmdLis
 			// Wait for all previous translate jobs to complete
 			Prereqs = MoveTemp(TranslateEvents);
 		}
+		
+		// Handle adding prerequisites for Parent end waiting on children
+		if (CmdList->ParallelRenderPassEnd.IsValid())
+		{
+			if (GRHIParallelRHIExecuteChildWait)
+			{
+				Prereqs = TranslateEvents;
+				
+				Prereqs.Add(BeginGraphEvent);
+				Prereqs.Append(ChildGraphEvents);
+			}
+			
+			ChildGraphEvents.Empty();
+			BeginGraphEvent = nullptr;
+		}
+		else if (CmdList->IsSubCommandList() && GRHIParallelRHIExecuteParentWait)
+		{
+			// Children need to wait on the parent begin
+			Prereqs = TranslateEvents;
+			Prereqs.Add(BeginGraphEvent);
+		}
 
 		// Also wait for the previous mutate fence (blocks future translates until the fence has signalled).
 		Prereqs.Add(GRHICommandList.LastMutate);
@@ -899,7 +1019,7 @@ void FRHICommandListExecutor::FSubmitState::Dispatch(FRHICommandListBase* CmdLis
 		Prereqs.Add(GRHICommandList.LastSubmit);
 
 		CmdList->UpdateTransitionsAndResourceState();
-
+		
 		CurrentTranslateJob->EnqueueTranslateTask(MoveTemp(Prereqs),
 			[TranslateJob = CurrentTranslateJob, CmdList]()
 			{
@@ -912,18 +1032,36 @@ void FRHICommandListExecutor::FSubmitState::Dispatch(FRHICommandListBase* CmdLis
 
 void FRHICommandListExecutor::FTranslateState::Translate(FRHICommandListBase* CmdList)
 {
-	// Apply the current translate job's contexts to the command list
-	for (ERHIPipeline Pipeline : MakeFlagsRange(ERHIPipeline::All))
+	if (CmdList->ParallelRenderPassEnd)
 	{
-		if (!CmdList->Contexts[Pipeline])
+		// Retrieve the contexts from the previous parent cmdlist.
+
+	#if DO_CHECK
+		check(!CmdList->UploadContext);
+		for (auto const& Context : CmdList->Contexts)
 		{
-			CmdList->Contexts[Pipeline] = PipelineStates[Pipeline].Context;
+			check(!Context);
 		}
-	}
+	#endif
 
-	if(!CmdList->UploadContext)
+		CmdList->Contexts = CmdList->ParallelRenderPassEnd->Contexts;
+		CmdList->UploadContext = CmdList->ParallelRenderPassEnd->UploadContext;
+	}
+	else
 	{
-		CmdList->UploadContext = UploadContextState;
+		// Apply the current translate job's contexts to the command list
+		for (ERHIPipeline Pipeline : MakeFlagsRange(ERHIPipeline::All))
+		{
+			if (!CmdList->Contexts[Pipeline])
+			{
+				CmdList->Contexts[Pipeline] = PipelineStates[Pipeline].Context;
+			}
+		}
+
+		if(!CmdList->UploadContext)
+		{
+			CmdList->UploadContext = UploadContextState;
+		}
 	}
 	
 	CmdList->ActivePipelines = ERHIPipeline::None;
@@ -968,16 +1106,32 @@ void FRHICommandListExecutor::FTranslateState::Translate(FRHICommandListBase* Cm
 #endif
 	DrawStats.Accumulate(CmdList->DrawStats);
 
+	if (CmdList->ParallelRenderPassBegin && !bShouldFinalize)
+	{
+		// Forward the acquired contexts to the next chain
+		CmdList->ParallelRenderPassBegin->UploadContext = CmdList->UploadContext;
+		CmdList->ParallelRenderPassBegin->Contexts = CmdList->Contexts;
+	}
+
 	delete CmdList;
 }
 
 FGraphEventRef FRHICommandListExecutor::FTranslateState::Finalize()
 {
+	if(!bShouldFinalize)
+	{
+		// Don't finalize if this chain ended with the start of a parallel render pass,
+		// and the contexts will be forwarded to the chain that processes the rest of the parent.
+		ENamedThreads::Type NamedThread = ENamedThreads::AnyHiPriThreadHiPriTask;
+		FTaskPipe* Pipe = GetTranslateTaskPipe(NamedThread);
+		return Pipe->Close();
+	}
+	
 	FTaskPipe* Pipe = EnqueueTranslateTask({},
 		[this]()
 		{
 			SCOPED_NAMED_EVENT(RHI_Finalize, FColor::White);
-
+		
 			TRHIPipelineArray<IRHIPlatformCommandList*> PlatformCommandLists {InPlace, nullptr};
 			FDynamicRHI::FRHIFinalizeContextArgs FinalizeArgs;
 			for (auto& State : PipelineStates)
@@ -988,18 +1142,26 @@ FGraphEventRef FRHICommandListExecutor::FTranslateState::Finalize()
 				}
 			}
 		
-			if(GDynamicRHI)
+			if (GDynamicRHI)
 			{
 				FinalizeArgs.UploadContext = UploadContextState;
 				
-				GDynamicRHI->RHIFinalizeContext(MoveTemp(FinalizeArgs), PlatformCommandLists);
-				
-				for (ERHIPipeline Pipeline : MakeFlagsRange(ERHIPipeline::All))
+				if (bUsingSubCmdLists)
+				{
+					auto& State = PipelineStates[ERHIPipeline::Graphics];
+					State.FinalizedCmdList = GDynamicRHI->RHIFinalizeParallelContext(State.Context);
+				}
+				else
 				{
-					auto& State = PipelineStates[Pipeline]; 
-					if (State.Context)
+					GDynamicRHI->RHIFinalizeContext(MoveTemp(FinalizeArgs), PlatformCommandLists);
+
+					for (ERHIPipeline Pipeline : MakeFlagsRange(ERHIPipeline::All))
 					{
-						State.FinalizedCmdList = PlatformCommandLists[Pipeline];
+						auto& State = PipelineStates[Pipeline];
+						if (State.Context)
+						{
+							State.FinalizedCmdList = PlatformCommandLists[Pipeline];
+						}
 					}
 				}
 			}
@@ -1166,9 +1328,21 @@ RHI_API FGraphEventRef FRHICommandListExecutor::Submit(TConstArrayView<FRHIComma
 	}
 
 	TArray<FRHICommandListBase*> AllCmdLists;
-	AllCmdLists.Reserve(AdditionalCommandLists.Num() + 1);
-	AllCmdLists.Add(ImmCmdList);
-	AllCmdLists.Append(AdditionalCommandLists);
+	auto ConsumeCmdList = [&](auto& Recurse, FRHICommandListBase* CmdList) -> void
+	{
+		for (FRHICommandListBase* AttachedCmdList : CmdList->AttachedCmdLists)
+		{
+			Recurse(Recurse, AttachedCmdList);
+		}
+
+		AllCmdLists.Add(CmdList);
+	};
+
+	ConsumeCmdList(ConsumeCmdList, ImmCmdList);
+	for (FRHICommandListBase* CmdList : AdditionalCommandLists)
+	{
+		ConsumeCmdList(ConsumeCmdList, CmdList);
+	}
 
 	//
 	// Submission of RHI command lists involves a chain of "dispatch" tasks.
diff --git a/Engine/Source/Runtime/RHI/Public/DynamicRHI.h b/Engine/Source/Runtime/RHI/Public/DynamicRHI.h
--- a/Engine/Source/Runtime/RHI/Public/DynamicRHI.h
+++ b/Engine/Source/Runtime/RHI/Public/DynamicRHI.h
@@ -877,7 +877,21 @@ public:
 	// Called by parallel worker threads, and the render thread. Platform implementations must be thread safe.
 	//
 	virtual IRHIComputeContext* RHIGetCommandContext(ERHIPipeline Pipeline, FRHIGPUMask GPUMask) = 0;
-
+	
+	//
+	// Retrieves a new command context to begin the recording of a new platform command list.
+	// The returned context is only for graphics pipelines and is used in conjunction with SubCommandList
+	// The returned context is expected to be in a state to execute draw commands immediately, i.e. renderpasses must be created and state set.
+	//
+	// Called by parallel worker threads, and the render thread. Platform implementations must be thread safe.
+	//
+	virtual IRHIComputeContext* RHIGetParallelCommandContext(FRHIParallelRenderPassInfo const& ParallelRenderPass, FRHIGPUMask GPUMask)
+	{
+		IRHICommandContext* Context = static_cast<IRHICommandContext*>(RHIGetCommandContext(ERHIPipeline::Graphics, GPUMask));
+		Context->RHIBeginRenderPass(ParallelRenderPass, ParallelRenderPass.PassName);
+		
+		return Context;
+	}
 	
 	virtual IRHIUploadContext* RHIGetUploadContext()
 	{
@@ -899,6 +913,19 @@ public:
 	};
 	
 	virtual void RHIFinalizeContext(FRHIFinalizeContextArgs&& Args, TRHIPipelineArray<IRHIPlatformCommandList*>& Output) = 0;
+	
+	virtual IRHIPlatformCommandList* RHIFinalizeParallelContext(IRHIComputeContext* Context)
+	{
+		FRHIFinalizeContextArgs Args;
+		static_cast<IRHICommandContext*>(Context)->RHIEndRenderPass();
+		
+		Args.Contexts.Add(Context);
+		
+		TRHIPipelineArray<IRHIPlatformCommandList*> Output {nullptr, nullptr};
+		RHIFinalizeContext(MoveTemp(Args), Output);
+		
+		return Output[ERHIPipeline::Graphics];
+	}
 
 	//
 	// Submits a batch of previously recorded/finalized command lists to the GPU. 
diff --git a/Engine/Source/Runtime/RHI/Public/RHICommandList.h b/Engine/Source/Runtime/RHI/Public/RHICommandList.h
--- a/Engine/Source/Runtime/RHI/Public/RHICommandList.h
+++ b/Engine/Source/Runtime/RHI/Public/RHICommandList.h
@@ -405,13 +405,6 @@ struct TRHILambdaCommand_NoMarker final : public FRHICommandBase
 	}
 };
 
-class FRHIContextArray : public TRHIPipelineArray<IRHIComputeContext*>
-{
-	using Base = TRHIPipelineArray<IRHIComputeContext*>;
-public:
-	using Base::Base;
-};
-
 template <typename RHICmdListType, typename LAMBDA>
 struct TRHILambdaCommandMultiPipe final : public FRHICommandBase
 {
@@ -728,11 +721,20 @@ public:
 	}
 	
 	inline bool Bypass() const;
+	
+	inline bool IsSubCommandList() const
+	{
+	   return SubRenderPassInfo.IsValid();
+	}
 
 private:
+	RHI_API void InvalidBufferFatalError(const FRHIBufferCreateDesc& CreateDesc);
+
+protected:
 	RHI_API void ActivatePipelines(ERHIPipeline Pipelines);
 
-	RHI_API void InvalidBufferFatalError(const FRHIBufferCreateDesc& CreateDesc);
+	// Takes the array of sub command lists and inserts them logically into a render pass at this point in time.
+	RHI_API void InsertParallelRenderPass_Base(TSharedPtr<FRHIParallelRenderPassInfo> const& InInfo, TArray<FRHISubCommandList*>&& SubCommandLists);
 
 public:
 	RHI_API void TransitionInternal(TConstArrayView<FRHITransitionInfo> Infos, ERHITransitionCreateFlags CreateFlags = ERHITransitionCreateFlags::None);
@@ -1370,6 +1372,12 @@ protected:
 
 	struct FRHICommandRHIThreadFence* LastLockFenceCommand = nullptr;
 
+	TArray<FRHICommandListBase*> AttachedCmdLists;
+
+	TSharedPtr<FRHIParallelRenderPassInfo> SubRenderPassInfo;
+	TSharedPtr<FRHIParallelRenderPassInfo> ParallelRenderPassBegin;
+	TSharedPtr<FRHIParallelRenderPassInfo> ParallelRenderPassEnd;
+
 	// Graph event used to gate the execution of the command list on the completion of any dependent tasks
 	// e.g. PSO async compilation and parallel RHICmdList recording tasks.
 	FGraphEventRef DispatchEvent;
@@ -1562,6 +1570,7 @@ private:
 	friend class FRHICommandListScopedFlushAndExecute;
 	friend class FRHICommandListScopedFence;
 	friend class FRHIComputeCommandList;
+	friend class FRHISubCommandList;
 	friend class FRHICommandListImmediate;
 	friend class FRHICommandList_RecursiveHazardous;
 	friend class FRHIComputeCommandList_RecursiveHazardous;
@@ -4403,7 +4412,7 @@ public:
 	FORCEINLINE_DEBUGGABLE void PollOcclusionQueries()
 	{
 	}
-
+	
 	FORCEINLINE_DEBUGGABLE void BeginRenderPass(const FRHIRenderPassInfo& InInfo, const TCHAR* Name)
 	{
 		check(!IsInsideRenderPass());
@@ -4476,6 +4485,12 @@ public:
 		ResetSubpass(ESubpassHint::None);
 	}
 
+	// Takes the array of sub command lists and inserts them logically into a render pass at this point in time.
+	void InsertParallelRenderPass(TSharedPtr<FRHIParallelRenderPassInfo> const& InInfo, TArray<FRHISubCommandList*>&& SubCommandLists)
+	{
+		InsertParallelRenderPass_Base(InInfo, MoveTemp(SubCommandLists));
+	}
+
 	FORCEINLINE_DEBUGGABLE void NextSubpass()
 	{
 		check(IsInsideRenderPass());
@@ -5689,6 +5704,18 @@ public:
 	}
 };
 
+class FRHISubCommandList : public FRHICommandList
+{
+public:
+	FRHISubCommandList(FRHIGPUMask GPUMask, TSharedPtr<FRHIParallelRenderPassInfo> const& RenderPassInfo)
+		: FRHICommandList(GPUMask)
+	{
+		RenderPassInfo->Validate();
+		SubRenderPassInfo = RenderPassInfo;
+		CacheActiveRenderTargets(*RenderPassInfo);
+	}
+};
+
 class FRHICommandListExecutor
 {
 public:
@@ -5845,9 +5872,12 @@ private:
 		FTaskPipe TranslatePipe;
 		uint32 NumCommands = 0;
 		bool bParallel = false;
+		bool bUsingSubCmdLists = false;
+		bool bShouldFinalize = true;
 
 		FRHIDrawStats DrawStats{};
 
+		FTaskPipe* GetTranslateTaskPipe(ENamedThreads::Type& NamedThread);
 		FTaskPipe* EnqueueTranslateTask(FGraphEventArray&& Prereqs, TFunction<void()>&& Lambda);
 
 		void Translate(FRHICommandListBase* CmdList);
@@ -5877,6 +5907,12 @@ private:
 		TArray<FRHIResource*> ResourcesToDelete {};
 		bool bIncludeExtendedLifetimeResources = false;
 
+		bool ShouldSplitTranslateJob(FRHICommandListBase* CmdList);
+		void ConditionalSplitTranslateJob(FRHICommandListBase* CmdList);
+		
+		FGraphEventRef BeginGraphEvent;
+		FGraphEventArray ChildGraphEvents;
+
 		void Dispatch(FRHICommandListBase* CmdList);
 
 		struct FSubmitArgs
@@ -5889,7 +5925,7 @@ private:
 		#endif
 		};
 		void Submit(const FSubmitArgs& Args);
-		void FinalizeCurrent();
+		FGraphEventRef FinalizeCurrent();
 	} *SubmitState = nullptr;
 
 	FGraphEventRef LastMutate;
diff --git a/Engine/Source/Runtime/RHI/Public/RHIContext.h b/Engine/Source/Runtime/RHI/Public/RHIContext.h
--- a/Engine/Source/Runtime/RHI/Public/RHIContext.h
+++ b/Engine/Source/Runtime/RHI/Public/RHIContext.h
@@ -814,6 +814,16 @@ public:
 	{
 		/* empty default implementation */
 	} 
+
+	virtual void RHIBeginParallelRenderPass(TSharedPtr<FRHIParallelRenderPassInfo> InInfo, const TCHAR* InName)
+	{
+		/* empty default implementation */
+	}
+	
+	virtual void RHIEndParallelRenderPass()
+	{
+		/* empty default implementation */
+	}
 	
 	virtual void RHIBeginRenderPass(const FRHIRenderPassInfo& InInfo, const TCHAR* InName) = 0;
 
diff --git a/Engine/Source/Runtime/RHI/Public/RHIFwd.h b/Engine/Source/Runtime/RHI/Public/RHIFwd.h
--- a/Engine/Source/Runtime/RHI/Public/RHIFwd.h
+++ b/Engine/Source/Runtime/RHI/Public/RHIFwd.h
@@ -33,10 +33,12 @@ class FRHICommandListBase;
 class FRHIComputeCommandList;
 class FRHICommandList;
 class FRHICommandListImmediate;
+class FRHISubCommandList;
 
 // Contexts
 class IRHIComputeContext;
 class IRHICommandContext;
+class IRHIUploadContext;
 class FRHIContextArray;
 
 class FRHIResourceReplaceInfo;
diff --git a/Engine/Source/Runtime/RHI/Public/RHIGlobals.h b/Engine/Source/Runtime/RHI/Public/RHIGlobals.h
--- a/Engine/Source/Runtime/RHI/Public/RHIGlobals.h
+++ b/Engine/Source/Runtime/RHI/Public/RHIGlobals.h
@@ -511,6 +511,17 @@ struct FRHIGlobals
 	* RHICreateBoundShaderState & RHICreateGraphicsPipelineState is threadsafe and GetCachedBoundShaderState must not be used. GetCachedBoundShaderState_Threadsafe has a slightly different protocol.
 	***/
 	bool SupportsParallelRHIExecute = false;
+	
+	/** Enables support for platforms that use render passes to split command list translation across multiple threads. For platforms that require
+	   dependencies on Parent/Child command lists i.e. requiring translation of children to be complete before finishing translation of the
+	   parent, should use ParallelRHIExecuteChildWait and/or ParallelRHIExecuteParentWait */
+	bool SupportsParallelRenderPasses = false;
+	
+	/** Used for parallel render pass support. When true, indicates that the platform RHI requires that the child command list translations in a parallel render pass are completed before the remainder of the parent command list can be translated. */
+	bool ParallelRHIExecuteChildWait = false;
+	
+	/** Used for parallel render pass support. When true, indicates that the platform RHI requires that child command list translations in a parallel render pass do not start until the parent command list translation has reached the BeginParallelRenderPass command. */
+	bool ParallelRHIExecuteParentWait = false;
 
 	/** Whether or not the RHI can perform MSAA sample load. */
 	bool SupportsMSAADepthSampleAccess = false;
@@ -858,6 +869,10 @@ extern RHI_API uint32 GGPUFrameTime;
 #define GRHISupportsRHIThread                                  GRHIGlobals.SupportsRHIThread
 #define GRHISupportsRHIOnTaskThread                            GRHIGlobals.SupportsRHIOnTaskThread
 #define GRHISupportsParallelRHIExecute                         GRHIGlobals.SupportsParallelRHIExecute
+#define GRHISupportsParallelRenderPasses                       GRHIGlobals.SupportsParallelRenderPasses
+#define GRHIParallelRHIExecuteChildWait                        GRHIGlobals.ParallelRHIExecuteChildWait
+#define GRHIParallelRHIExecuteParentWait                       GRHIGlobals.ParallelRHIExecuteParentWait
+#define GRHISupportsParallelRHIExecute                         GRHIGlobals.SupportsParallelRHIExecute
 #define GRHISupportsMSAADepthSampleAccess                      GRHIGlobals.SupportsMSAADepthSampleAccess
 #define GRHISupportsBackBufferWithCustomDepthStencil           GRHIGlobals.SupportsBackBufferWithCustomDepthStencil
 #define GRHIIsHDREnabled                                       GRHIGlobals.IsHDREnabled
diff --git a/Engine/Source/Runtime/RHI/Public/RHIResources.h b/Engine/Source/Runtime/RHI/Public/RHIResources.h
--- a/Engine/Source/Runtime/RHI/Public/RHIResources.h
+++ b/Engine/Source/Runtime/RHI/Public/RHIResources.h
@@ -5455,6 +5455,27 @@ struct FRHIRenderPassInfo
 	RHI_API void ConvertToRenderTargetsInfo(FRHISetRenderTargetsInfo& OutRTInfo) const;
 };
 
+class FRHIContextArray : public TRHIPipelineArray<IRHIComputeContext*>
+{
+	using Base = TRHIPipelineArray<IRHIComputeContext*>;
+public:
+	using Base::Base;
+};
+
+struct FRHIParallelRenderPassInfo : public FRHIRenderPassInfo
+{
+	FRHIContextArray Contexts{ InPlace, nullptr };
+	IRHIUploadContext* UploadContext = nullptr;
+	const TCHAR* PassName;
+	
+	// Contains platform specific data
+	void* RHIPlatformData;
+
+	FRHIParallelRenderPassInfo(FRHIRenderPassInfo&& Info, const TCHAR* InPassName)
+		: FRHIRenderPassInfo(MoveTemp(Info)), PassName(InPassName)
+	{}
+};
+
 //UE_DEPRECATED(5.3, "Use FRHITextureSRVCreateDesc::SetDisableSRGB to create a view which explicitly disables SRGB.")
 enum ERHITextureSRVOverrideSRGBType : uint8
 {
diff --git a/Engine/Source/Runtime/RenderCore/Private/RenderGraphPass.cpp b/Engine/Source/Runtime/RenderCore/Private/RenderGraphPass.cpp
--- a/Engine/Source/Runtime/RenderCore/Private/RenderGraphPass.cpp
+++ b/Engine/Source/Runtime/RenderCore/Private/RenderGraphPass.cpp
@@ -117,24 +117,38 @@ FRHIRenderPassInfo FRDGParameterStruct::GetRenderPassInfo() const
 
 FRHICommandList* FRDGDispatchPassBuilder::CreateCommandList()
 {
-	FRHICommandList* RHICmdList = new FRHICommandList(Pass->GetGPUMask());
-	RHICmdList->SwitchPipeline(Pass->GetPipeline());
-
-	// When parallel executing, the pass commands are embedded directly into the first command list.
-	if (Pass->bParallelExecute && Pass->CommandLists.IsEmpty())
+	FRHICommandList* RHICmdList = nullptr;
+	
+	if (RenderPassInfo && Pass->bParallelExecute && GRHISupportsParallelRenderPasses)
 	{
-		FRDGBuilder::PushPreScopes(*RHICmdList, Pass);
-		FRDGBuilder::ExecutePassPrologue(*RHICmdList, Pass);
-	}
+		FRHISubCommandList* SubCmdList = new FRHISubCommandList(Pass->GetGPUMask(), RenderPassInfo);
+		SubCmdList->SwitchPipeline(Pass->GetPipeline());
 
-	if (RenderPassInfo)
-	{
-		RHICmdList->BeginRenderPass(*RenderPassInfo, TEXT("DispatchPass"));
+		SubCommandLists.Emplace(SubCmdList);
+		RHICmdList = SubCmdList;
 	}
+	else
+	{
+		RHICmdList = new FRHICommandList(Pass->GetGPUMask());
+		RHICmdList->SwitchPipeline(Pass->GetPipeline());
 
+		// When parallel executing, the pass commands are embedded directly into the first command list.
+		if (Pass->bParallelExecute && Pass->CommandLists.IsEmpty())
+		{
+			FRDGBuilder::PushPreScopes(*RHICmdList, Pass);
+			FRDGBuilder::ExecutePassPrologue(*RHICmdList, Pass);
+		}
+		
+		if (RenderPassInfo)
+		{
+			RHICmdList->BeginRenderPass(*RenderPassInfo, TEXT("DispatchPass"));
+		}
+		
+		Pass->CommandLists.Emplace(RHICmdList);
+	}
+	
 	RHICmdList->SetStaticUniformBuffers(StaticUniformBuffers);
 
-	Pass->CommandLists.Emplace(RHICmdList);
 	return RHICmdList;
 }
 
@@ -146,22 +160,40 @@ void FRDGDispatchPassBuilder::Finish()
 		Pass->CommandListsEvent.Trigger();
 		return;
 	}
-
-	const bool bEmptyCommandLists = Pass->CommandLists.IsEmpty();
-
-	// Create a command list to embed the epilogue (and prologue as well if no user command lists were requested).
-	FRHICommandList* RHICmdList = new FRHICommandList(Pass->GetGPUMask());
-	Pass->CommandLists.Emplace(RHICmdList);
-	Pass->CommandListsEvent.Trigger();
-
-	RHICmdList->SwitchPipeline(Pass->GetPipeline());
-
-	if (bEmptyCommandLists)
+	
+	FRHICommandList* RHICmdList = nullptr;
+	
+	if(SubCommandLists.Num())
 	{
+		RHICmdList = new FRHICommandList(Pass->GetGPUMask());
+		RHICmdList->SwitchPipeline(Pass->GetPipeline());
+		
 		FRDGBuilder::PushPreScopes(*RHICmdList, Pass);
 		FRDGBuilder::ExecutePassPrologue(*RHICmdList, Pass);
+		
+		RHICmdList->InsertParallelRenderPass(RenderPassInfo, MoveTemp(SubCommandLists));
+		
+		Pass->CommandLists.Emplace(RHICmdList);
+		Pass->CommandListsEvent.Trigger();
 	}
+	else
+	{
+		const bool bEmptyCommandLists = Pass->CommandLists.IsEmpty();
+		
+		// Create a command list to embed the epilogue (and prologue as well if no user command lists were requested).
+		RHICmdList = new FRHICommandList(Pass->GetGPUMask());
+		Pass->CommandLists.Emplace(RHICmdList);
+		Pass->CommandListsEvent.Trigger();
+
+		RHICmdList->SwitchPipeline(Pass->GetPipeline());
 
+		if (bEmptyCommandLists)
+		{
+			FRDGBuilder::PushPreScopes(*RHICmdList, Pass);
+			FRDGBuilder::ExecutePassPrologue(*RHICmdList, Pass);
+		}
+	}
+	
 	FRDGBuilder::ExecutePassEpilogue(*RHICmdList, Pass);
 	FRDGBuilder::PopPreScopes(*RHICmdList, Pass);
 
@@ -416,4 +448,4 @@ const TCHAR* FRDGPass::GetName() const
 		return Name.GetTCHAR();
 	}
 }
-#endif
\ No newline at end of file
+#endif
diff --git a/Engine/Source/Runtime/RenderCore/Public/RenderGraphPass.h b/Engine/Source/Runtime/RenderCore/Public/RenderGraphPass.h
--- a/Engine/Source/Runtime/RenderCore/Public/RenderGraphPass.h
+++ b/Engine/Source/Runtime/RenderCore/Public/RenderGraphPass.h
@@ -745,7 +745,7 @@ private:
 	{
 		if (Pass->ParameterStruct.HasRenderTargets())
 		{
-			RenderPassInfo.Emplace(Pass->ParameterStruct.GetRenderPassInfo());
+			RenderPassInfo = MakeShared<FRHIParallelRenderPassInfo, ESPMode::ThreadSafe>(Pass->ParameterStruct.GetRenderPassInfo(), TEXT("DispatchPass"));
 		}
 	}
 
@@ -753,8 +753,10 @@ private:
 
 	FRDGDispatchPass* Pass;
 	FUniformBufferStaticBindings StaticUniformBuffers;
-	TOptional<FRHIRenderPassInfo> RenderPassInfo;
+	TSharedPtr<FRHIParallelRenderPassInfo> RenderPassInfo;
 
+	TArray<FRHISubCommandList*> SubCommandLists;
+	
 	friend FRDGBuilder;
 };
 
@@ -829,4 +831,4 @@ class FRDGBuilder;
 class FRDGPass;
 class FRDGTrace;
 class FRDGUserValidation;
-class FShaderParametersMetadata;
\ No newline at end of file
+class FShaderParametersMetadata;
diff --git a/Engine/Source/Runtime/Renderer/Private/MeshDrawCommands.cpp b/Engine/Source/Runtime/Renderer/Private/MeshDrawCommands.cpp
--- a/Engine/Source/Runtime/Renderer/Private/MeshDrawCommands.cpp
+++ b/Engine/Source/Runtime/Renderer/Private/MeshDrawCommands.cpp
@@ -1607,7 +1607,8 @@ public:
 	{
 		FTaskTagScope Scope(ETaskTag::EParallelRenderingThread);
 		SCOPED_NAMED_EVENT_TEXT("DrawVisibleMeshCommandsAnyThreadTask", FColor::Magenta);
-		checkSlow(RHICmdList.IsInsideRenderPass());
+		
+		checkSlow(RHICmdList.IsInsideRenderPass() || RHICmdList.IsSubCommandList());
 
 		// check for the multithreaded shader creation has been moved to FShaderCodeArchive::CreateShader() 
 		{
@@ -1629,7 +1630,10 @@ public:
 				RHICmdList);
 		}
 
-		RHICmdList.EndRenderPass();
+		if(!RHICmdList.IsSubCommandList())
+		{
+			RHICmdList.EndRenderPass();
+		}
 		RHICmdList.FinishRecording();
 	}
 };
@@ -1838,4 +1842,4 @@ void FParallelMeshDrawCommandPass::DumpInstancingStats() const
 void FParallelMeshDrawCommandPass::SetDumpInstancingStats(const FString& InPassNameForStats)
 {
 	PassNameForStats = InPassNameForStats;
-}
\ No newline at end of file
+}
diff --git a/Engine/Source/Runtime/VulkanRHI/Private/VulkanRHI.cpp b/Engine/Source/Runtime/VulkanRHI/Private/VulkanRHI.cpp
--- a/Engine/Source/Runtime/VulkanRHI/Private/VulkanRHI.cpp
+++ b/Engine/Source/Runtime/VulkanRHI/Private/VulkanRHI.cpp
@@ -896,6 +896,7 @@ void FVulkanDynamicRHI::InitInstance()
 			GRHIGlobals.RayTracing.SupportsLooseParamsInShaderRecord = true;
 		}
 
+		GRHISupportsParallelRenderPasses = false;
 #if VULKAN_ENABLE_DUMP_LAYER
 		// Disable RHI thread by default if the dump layer is enabled
 		GRHISupportsRHIThread = false;
