Moved RHI_NEW_GPU_PROFILER out of RHI.Build.cs and defined it instead in GpuProfilerTrace.h, to prepare for enabling the system by default
- This also allows us to enable or disable the new profiler selectively for projects or platforms via .Build.cs / .Target.cs configuration.

Minimal RHI breadcrumb mode for test builds with GPU stats enabled
- Since the GPU stats now depend on the RHI breadcrumb system, we need parts of it enabled when HAS_GPU_STATS is 1. This is required for systems like the CSV profiler.
- WITH_RHI_BREADCRUMBS is still the main macro to control whether the breadcrumbs system is included at all. Two new macros (WITH_RHI_BREADCRUMBS_FULL and WITH_RHI_BREADCRUMBS_MINIMAL) further decide which RHI/RDG scopes are enabled.

Fixed D3D11 breadcrumbs
- Emitting the GPU profiler events should not be guarded by ShouldEmitBreadcrumbs().
- That function is used to determine whether the RHI should output debug markers to platform APIs such as PIX, since doing so requires formatting the breadcrumb string.

diff --git a/Engine/Source/Runtime/RHI/Private/RHIBreadcrumbs.cpp b/Engine/Source/Runtime/RHI/Private/RHIBreadcrumbs.cpp
--- a/Engine/Source/Runtime/RHI/Private/RHIBreadcrumbs.cpp
+++ b/Engine/Source/Runtime/RHI/Private/RHIBreadcrumbs.cpp
@@ -11,8 +11,6 @@
 // Constructor for the sentinel node
 FRHIBreadcrumbNode::FRHIBreadcrumbNode(FRHIBreadcrumbData const& Data)
 	: Data(Data)
-	, CPUTraceMarkerID(0)
-	, ID(0)
 {}
 
 RHI_API FRHIBreadcrumbNode* const FRHIBreadcrumbNode::Sentinel = []()
diff --git a/Engine/Source/Runtime/RHI/Public/GpuProfilerTrace.h b/Engine/Source/Runtime/RHI/Public/GpuProfilerTrace.h
--- a/Engine/Source/Runtime/RHI/Public/GpuProfilerTrace.h
+++ b/Engine/Source/Runtime/RHI/Public/GpuProfilerTrace.h
@@ -8,6 +8,10 @@
 #include "Trace/Config.h"
 #include "UObject/NameTypes.h"
 
+#ifndef RHI_NEW_GPU_PROFILER
+#define RHI_NEW_GPU_PROFILER 0
+#endif
+
 #if !defined(GPUPROFILERTRACE_ENABLED)
 #if UE_TRACE_ENABLED && !UE_BUILD_SHIPPING
 // The trace implementation for the old GPU Profiler (FGpuProfilerTrace API) is deprecated in UE 5.6.
diff --git a/Engine/Source/Runtime/RHI/Public/RHIBreadcrumbs.h b/Engine/Source/Runtime/RHI/Public/RHIBreadcrumbs.h
--- a/Engine/Source/Runtime/RHI/Public/RHIBreadcrumbs.h
+++ b/Engine/Source/Runtime/RHI/Public/RHIBreadcrumbs.h
@@ -16,15 +16,25 @@
 
 #include <array>
 
+//
+// Controls whether the infrastructure for the RHI breadcrumbs system is included in the build.
+// (RHI breadcrumb allocators, platform RHI implemntation etc).
+//
 #ifndef WITH_RHI_BREADCRUMBS
-#define WITH_RHI_BREADCRUMBS ((UE_BUILD_DEBUG || UE_BUILD_DEVELOPMENT) || WITH_PROFILEGPU)
+#define WITH_RHI_BREADCRUMBS (UE_BUILD_DEBUG || UE_BUILD_DEVELOPMENT || WITH_PROFILEGPU || HAS_GPU_STATS)
 #endif
 
+// Enables all RDG/RHI breadcrumb scopes, and features such as Insights markers.
+#define WITH_RHI_BREADCRUMBS_FULL    (WITH_RHI_BREADCRUMBS && (UE_BUILD_DEBUG || UE_BUILD_DEVELOPMENT || WITH_PROFILEGPU))
+
+// Enables only the necessary scopes for GPU stats to work
+#define WITH_RHI_BREADCRUMBS_MINIMAL (WITH_RHI_BREADCRUMBS && (!WITH_RHI_BREADCRUMBS_FULL))
+
 // Whether to emit Unreal Insights breadcrumb events on threads involved in RHI command list recording and execution.
-#define RHI_BREADCRUMBS_EMIT_CPU (WITH_RHI_BREADCRUMBS && CPUPROFILERTRACE_ENABLED && 1)
+#define RHI_BREADCRUMBS_EMIT_CPU (WITH_RHI_BREADCRUMBS_FULL && CPUPROFILERTRACE_ENABLED && 1)
 
 // Whether to store the filename and line number of each RHI breadcrumb and emit this data to Insights.
-#define RHI_BREADCRUMBS_EMIT_LOCATION (WITH_RHI_BREADCRUMBS && (CPUPROFILERTRACE_ENABLED || GPUPROFILERTRACE_ENABLED) && 1)
+#define RHI_BREADCRUMBS_EMIT_LOCATION (WITH_RHI_BREADCRUMBS_FULL && (CPUPROFILERTRACE_ENABLED || GPUPROFILERTRACE_ENABLED) && 1)
 
 #if RHI_NEW_GPU_PROFILER && HAS_GPU_STATS
 namespace UE::RHI::GPUProfiler
@@ -228,8 +238,11 @@ namespace UE::RHI::GPUProfiler
 		FRHIBreadcrumbAllocator* const Allocator = nullptr;
 
 		FRHIBreadcrumbData const& Data;
+	#if RHI_BREADCRUMBS_EMIT_CPU
 		uint32 const CPUTraceMarkerID = 0;
-		uint32 const ID;
+	#endif
+
+		uint32 const ID = 0;
 
 	#if DO_CHECK
 		// Used to track use of this breadcrumb on each GPU pipeline. Breadcrumbs can only be begun/ended once per pipe.
@@ -1264,6 +1277,10 @@ namespace UE::RHI::GPUProfiler
 			}																									 \
 		} while(false)
 
+#endif // WITH_RHI_BREADCRUMBS
+
+#if WITH_RHI_BREADCRUMBS_FULL
+
 	// Note, the varargs are deprecated and ignored in these macros.
 	#define RHI_BREADCRUMB_EVENT(                             RHICmdList_Or_RHIContext,                  StaticName,         ...) RHI_BREADCRUMB_EVENT_PRIVATE_IMPL(RHICmdList_Or_RHIContext, RHI_GPU_STAT_ARGS_NONE ,      true, TEXT(StaticName),      nullptr, ##__VA_ARGS__)
 	#define RHI_BREADCRUMB_EVENT_CONDITIONAL(                 RHICmdList_Or_RHIContext,       Condition, StaticName,         ...) RHI_BREADCRUMB_EVENT_PRIVATE_IMPL(RHICmdList_Or_RHIContext, RHI_GPU_STAT_ARGS_NONE , Condition, TEXT(StaticName),      nullptr, ##__VA_ARGS__)
@@ -1280,12 +1297,30 @@ namespace UE::RHI::GPUProfiler
 	#define RHI_BREADCRUMB_EVENT_F_STR_DEPRECATED(            RHICmdList_Or_RHIContext,                  StaticName, Format, ...) RHI_BREADCRUMB_EVENT_PRIVATE_IMPL(RHICmdList_Or_RHIContext, RHI_GPU_STAT_ARGS_NONE ,      true, TEXT(StaticName),      Format , ##__VA_ARGS__)
 	#define RHI_BREADCRUMB_EVENT_F_CONDITIONAL_STR_DEPRECATED(RHICmdList_Or_RHIContext,       Condition, StaticName, Format, ...) RHI_BREADCRUMB_EVENT_PRIVATE_IMPL(RHICmdList_Or_RHIContext, RHI_GPU_STAT_ARGS_NONE , Condition, TEXT(StaticName),      Format , ##__VA_ARGS__)
 
-	// Log and check macros that include the current breadcrumb path
-	#define RHI_BREADCRUMB_LOG(   RHICmdList_Or_RHIContext,            CategoryName, Verbosity, Format, ...) UE_LOG(            CategoryName, Verbosity, Format TEXT("\nBreadcrumbs: %s"), ##__VA_ARGS__, *UE::RHI::Breadcrumbs::Private::GetSafeBreadcrumbPath(RHICmdList_Or_RHIContext))
-	#define RHI_BREADCRUMB_CLOG(  RHICmdList_Or_RHIContext, Condition, CategoryName, Verbosity, Format, ...) UE_CLOG(Condition, CategoryName, Verbosity, Format TEXT("\nBreadcrumbs: %s"), ##__VA_ARGS__, *UE::RHI::Breadcrumbs::Private::GetSafeBreadcrumbPath(RHICmdList_Or_RHIContext))
-	#define RHI_BREADCRUMB_CHECKF(RHICmdList_Or_RHIContext, Condition,                          Format, ...) checkf( Condition,                          Format TEXT("\nBreadcrumbs: %s"), ##__VA_ARGS__, *UE::RHI::Breadcrumbs::Private::GetSafeBreadcrumbPath(RHICmdList_Or_RHIContext))
+#elif WITH_RHI_BREADCRUMBS_MINIMAL
 
-#else // WITH_RHI_BREADCRUMBS == 0
+	//
+	// Keep only the STAT breadcrumbs enabled in MINIMAL mode.
+	// Also disable the varargs. We don't capture the format strings and varargs in MINIMAL mode
+	//
+
+	// Note, the varargs are deprecated and ignored in these macros.
+	#define RHI_BREADCRUMB_EVENT(                             RHICmdList_Or_RHIContext,                  StaticName,         ...)
+	#define RHI_BREADCRUMB_EVENT_CONDITIONAL(                 RHICmdList_Or_RHIContext,       Condition, StaticName,         ...)
+	#define RHI_BREADCRUMB_EVENT_STAT(                        RHICmdList_Or_RHIContext, Stat,            StaticName,         ...) RHI_BREADCRUMB_EVENT_PRIVATE_IMPL(RHICmdList_Or_RHIContext, RHI_GPU_STAT_ARGS(Stat),      true, TEXT(StaticName), nullptr)
+	#define RHI_BREADCRUMB_EVENT_CONDITIONAL_STAT(            RHICmdList_Or_RHIContext, Stat, Condition, StaticName,         ...) RHI_BREADCRUMB_EVENT_PRIVATE_IMPL(RHICmdList_Or_RHIContext, RHI_GPU_STAT_ARGS(Stat), Condition, TEXT(StaticName), nullptr)
+													          
+	// Format versions of the breadcrumb macros.	          
+	#define RHI_BREADCRUMB_EVENT_F(                           RHICmdList_Or_RHIContext,                  StaticName, Format, ...)
+	#define RHI_BREADCRUMB_EVENT_CONDITIONAL_F(               RHICmdList_Or_RHIContext,       Condition, StaticName, Format, ...)
+	#define RHI_BREADCRUMB_EVENT_STAT_F(                      RHICmdList_Or_RHIContext, Stat,            StaticName, Format, ...) RHI_BREADCRUMB_EVENT_PRIVATE_IMPL(RHICmdList_Or_RHIContext, RHI_GPU_STAT_ARGS(Stat),      true, TEXT(StaticName), nullptr)
+	#define RHI_BREADCRUMB_EVENT_CONDITIONAL_STAT_F(          RHICmdList_Or_RHIContext, Stat, Condition, StaticName, Format, ...) RHI_BREADCRUMB_EVENT_PRIVATE_IMPL(RHICmdList_Or_RHIContext, RHI_GPU_STAT_ARGS(Stat), Condition, TEXT(StaticName), nullptr)
+
+	// Used only for back compat with SCOPED_DRAW_EVENTF
+	#define RHI_BREADCRUMB_EVENT_F_STR_DEPRECATED(            RHICmdList_Or_RHIContext,                  StaticName, Format, ...)
+	#define RHI_BREADCRUMB_EVENT_F_CONDITIONAL_STR_DEPRECATED(RHICmdList_Or_RHIContext,       Condition, StaticName, Format, ...)
+
+#else
 
 	#define RHI_BREADCRUMB_FIELD(                             ...)
 	#define RHI_BREADCRUMB_EVENT(                             ...)
@@ -1299,6 +1334,17 @@ namespace UE::RHI::GPUProfiler
 	#define RHI_BREADCRUMB_EVENT_F_STR_DEPRECATED(            ...)
 	#define RHI_BREADCRUMB_EVENT_F_CONDITIONAL_STR_DEPRECATED(...)
 
+#endif
+
+#if WITH_RHI_BREADCRUMBS_FULL
+
+	// Log and check macros that include the current breadcrumb path
+	#define RHI_BREADCRUMB_LOG(   RHICmdList_Or_RHIContext,            CategoryName, Verbosity, Format, ...) UE_LOG(            CategoryName, Verbosity, Format TEXT("\nBreadcrumbs: %s"), ##__VA_ARGS__, *UE::RHI::Breadcrumbs::Private::GetSafeBreadcrumbPath(RHICmdList_Or_RHIContext))
+	#define RHI_BREADCRUMB_CLOG(  RHICmdList_Or_RHIContext, Condition, CategoryName, Verbosity, Format, ...) UE_CLOG(Condition, CategoryName, Verbosity, Format TEXT("\nBreadcrumbs: %s"), ##__VA_ARGS__, *UE::RHI::Breadcrumbs::Private::GetSafeBreadcrumbPath(RHICmdList_Or_RHIContext))
+	#define RHI_BREADCRUMB_CHECKF(RHICmdList_Or_RHIContext, Condition,                          Format, ...) checkf( Condition,                          Format TEXT("\nBreadcrumbs: %s"), ##__VA_ARGS__, *UE::RHI::Breadcrumbs::Private::GetSafeBreadcrumbPath(RHICmdList_Or_RHIContext))
+
+#else
+
 	// Log and check macros fall back to regular UE_LOG / check when breadcrumbs are not available
 	#define RHI_BREADCRUMB_LOG(   RHICmdList_Or_RHIContext,            CategoryName, Verbosity, Format, ...) UE_LOG(            CategoryName, Verbosity, Format, ##__VA_ARGS__)
 	#define RHI_BREADCRUMB_CLOG(  RHICmdList_Or_RHIContext, Condition, CategoryName, Verbosity, Format, ...) UE_CLOG(Condition, CategoryName, Verbosity, Format, ##__VA_ARGS__)
diff --git a/Engine/Source/Runtime/RHI/Public/RHICommandList.h b/Engine/Source/Runtime/RHI/Public/RHICommandList.h
--- a/Engine/Source/Runtime/RHI/Public/RHICommandList.h
+++ b/Engine/Source/Runtime/RHI/Public/RHICommandList.h
@@ -1648,7 +1648,11 @@ protected:
 //
 inline bool IRHIComputeContext::ShouldEmitBreadcrumbs() const
 {
+#if WITH_RHI_BREADCRUMBS_FULL
 	return GetExecutingCommandList().CPUBreadcrumbState.bEmitBreadcrumbs;
+#else
+	return false;
+#endif
 }
 #endif
 
diff --git a/Engine/Source/Runtime/RHI/RHI.Build.cs b/Engine/Source/Runtime/RHI/RHI.Build.cs
--- a/Engine/Source/Runtime/RHI/RHI.Build.cs
+++ b/Engine/Source/Runtime/RHI/RHI.Build.cs
@@ -13,9 +13,6 @@ public class RHI : ModuleRules
 		PrivateDependencyModuleNames.Add("ApplicationCore");
 		PrivateDependencyModuleNames.Add("Cbor");
 
-		// @todo - new gpu profiler. This is experimental.
-		PublicDefinitions.Add("RHI_NEW_GPU_PROFILER=0");
-
 		PublicDefinitions.AddDefinition("WITH_MGPU", Target.Platform.IsInGroup(UnrealPlatformGroup.Windows) && Target.Platform.IsInGroup(UnrealPlatformGroup.Desktop));
 
 		if (Target.bCompileAgainstEngine)
diff --git a/Engine/Source/Runtime/RenderCore/Public/RenderGraphEvent.h b/Engine/Source/Runtime/RenderCore/Public/RenderGraphEvent.h
--- a/Engine/Source/Runtime/RenderCore/Public/RenderGraphEvent.h
+++ b/Engine/Source/Runtime/RenderCore/Public/RenderGraphEvent.h
@@ -427,12 +427,7 @@ private:
 		[&](UE::RHI::Breadcrumbs::Private::FRDGFormatTag&& Tag, auto&&... Values) -> auto				\
 		{																								\
 			using namespace UE::RHI::Breadcrumbs::Private;												\
-			if constexpr (RDG_EVENTS <= RDG_EVENTS_STRING_REF)											\
-			{																							\
-				/* RDG events are only capturing the format string, no args. */							\
-				return RHI_BREADCRUMB_DESC_FORWARD_VALUES(StaticName, nullptr, GPUStatArgs)();			\
-			}																							\
-			else if constexpr (std::is_same_v<TRemoveCVRef<decltype(FormatString)>, FRDGEventName>)		\
+			if constexpr (std::is_same_v<TRemoveCVRef<decltype(FormatString)>, FRDGEventName>)			\
 			{																							\
 				/* Single FRDGEventName */																\
 				static_assert(sizeof...(Values) == 0, "Unexpected additional arguments");				\
@@ -490,17 +485,44 @@ private:
 			, ##__VA_ARGS__																							\
 		)
 
-	// Skip expensive string formatting for the relatively common case of no varargs.  We detect this by stringizing the varargs and checking if the string is non-empty (more than just a null terminator).
-	#define RDG_EVENT_NAME(Format, ...) (sizeof(#__VA_ARGS__ "") > 1 ? FRDGEventName(TEXT(Format), ##__VA_ARGS__) : FRDGEventName(1, TEXT(Format)))
-
-	#define RDG_EVENT_SCOPE(                 GraphBuilder,                      Format, ...) RDG_EVENT_SCOPE_IMPL(GraphBuilder,      true, ERDGScopeFlags::None , RHI_GPU_STAT_ARGS_NONE	 , TEXT(Format)    , TEXT(Format), ##__VA_ARGS__)
-	#define RDG_EVENT_SCOPE_STAT(            GraphBuilder,            StatName, Format, ...) RDG_EVENT_SCOPE_IMPL(GraphBuilder,      true, ERDGScopeFlags::Stat , RHI_GPU_STAT_ARGS(StatName), TEXT(Format)    , TEXT(Format), ##__VA_ARGS__)
-	#define RDG_EVENT_SCOPE_CONDITIONAL(     GraphBuilder, Condition,           Format, ...) RDG_EVENT_SCOPE_IMPL(GraphBuilder, Condition, ERDGScopeFlags::None , RHI_GPU_STAT_ARGS_NONE	 , TEXT(Format)    , TEXT(Format), ##__VA_ARGS__)
-	#define RDG_EVENT_SCOPE_CONDITIONAL_STAT(GraphBuilder, Condition, StatName, Format, ...) RDG_EVENT_SCOPE_IMPL(GraphBuilder, Condition, ERDGScopeFlags::Stat , RHI_GPU_STAT_ARGS(StatName), TEXT(Format)    , TEXT(Format), ##__VA_ARGS__)
-	// The 'Final' version disables any further child scopes or pass events. It is intended to group overlapping passes as events can disable overlap on certain GPUs.
-	#define RDG_EVENT_SCOPE_FINAL(           GraphBuilder,                      Format, ...) RDG_EVENT_SCOPE_IMPL(GraphBuilder,      true, ERDGScopeFlags::Final, RHI_GPU_STAT_ARGS_NONE	 , TEXT(Format)    , TEXT(Format), ##__VA_ARGS__)
-	// Used in places which have an existing FRDGEventName, e.g. RDG pass name scopes. Prefer to use the other RDG scope macros instead.
-	#define RDG_EVENT_SCOPE_CONDITIONAL_NAME(GraphBuilder, Condition,           EventName  ) RDG_EVENT_SCOPE_IMPL(GraphBuilder, Condition, ERDGScopeFlags::None , RHI_GPU_STAT_ARGS_NONE     , TEXT("RDGEvent"), EventName)
+	#if WITH_RHI_BREADCRUMBS_FULL
+
+		#if RDG_EVENTS >= RDG_EVENTS_STRING_COPY
+			#define RDG_SCOPE_ARGS(Format, ...) TEXT(Format), ##__VA_ARGS__
+		#else
+			#define RDG_SCOPE_ARGS(Format, ...) nullptr
+		#endif
+
+		// Skip expensive string formatting for the relatively common case of no varargs.  We detect this by stringizing the varargs and checking if the string is non-empty (more than just a null terminator).
+		#define RDG_EVENT_NAME(Format, ...) (sizeof(#__VA_ARGS__ "") > 1 ? FRDGEventName(TEXT(Format), ##__VA_ARGS__) : FRDGEventName(1, TEXT(Format)))
+
+		#define RDG_EVENT_SCOPE(                 GraphBuilder,                      Format, ...) RDG_EVENT_SCOPE_IMPL(GraphBuilder,      true, ERDGScopeFlags::None , RHI_GPU_STAT_ARGS_NONE	 , TEXT(Format)    , RDG_SCOPE_ARGS(Format, ##__VA_ARGS__))
+		#define RDG_EVENT_SCOPE_STAT(            GraphBuilder,            StatName, Format, ...) RDG_EVENT_SCOPE_IMPL(GraphBuilder,      true, ERDGScopeFlags::Stat , RHI_GPU_STAT_ARGS(StatName), TEXT(Format)    , RDG_SCOPE_ARGS(Format, ##__VA_ARGS__))
+		#define RDG_EVENT_SCOPE_CONDITIONAL(     GraphBuilder, Condition,           Format, ...) RDG_EVENT_SCOPE_IMPL(GraphBuilder, Condition, ERDGScopeFlags::None , RHI_GPU_STAT_ARGS_NONE	 , TEXT(Format)    , RDG_SCOPE_ARGS(Format, ##__VA_ARGS__))
+		#define RDG_EVENT_SCOPE_CONDITIONAL_STAT(GraphBuilder, Condition, StatName, Format, ...) RDG_EVENT_SCOPE_IMPL(GraphBuilder, Condition, ERDGScopeFlags::Stat , RHI_GPU_STAT_ARGS(StatName), TEXT(Format)    , RDG_SCOPE_ARGS(Format, ##__VA_ARGS__))
+		// The 'Final' version disables any further child scopes or pass events. It is intended to group overlapping passes as events can disable overlap on certain GPUs.
+		#define RDG_EVENT_SCOPE_FINAL(           GraphBuilder,                      Format, ...) RDG_EVENT_SCOPE_IMPL(GraphBuilder,      true, ERDGScopeFlags::Final, RHI_GPU_STAT_ARGS_NONE	 , TEXT(Format)    , RDG_SCOPE_ARGS(Format, ##__VA_ARGS__))
+		// Used in places which have an existing FRDGEventName, e.g. RDG pass name scopes. Prefer to use the other RDG scope macros instead.
+		#define RDG_EVENT_SCOPE_CONDITIONAL_NAME(GraphBuilder, Condition,           EventName  ) RDG_EVENT_SCOPE_IMPL(GraphBuilder, Condition, ERDGScopeFlags::None , RHI_GPU_STAT_ARGS_NONE     , TEXT("RDGEvent"), EventName)
+
+	#else // WITH_RHI_BREADCRUMBS_MINIMAL
+
+		//
+		// Keep only the STAT RDG scopes enabled in MINIMAL mode.
+		// Also disable the varargs. We don't capture the format strings and varargs in MINIMAL mode
+		//
+
+		#define RDG_EVENT_NAME(Format, ...) FRDGEventName()
+
+		#define RDG_EVENT_SCOPE_STAT(            GraphBuilder,            StatName, Format, ...) RDG_EVENT_SCOPE_IMPL(GraphBuilder,      true, ERDGScopeFlags::Stat , RHI_GPU_STAT_ARGS(StatName), TEXT(Format), nullptr)
+		#define RDG_EVENT_SCOPE_CONDITIONAL_STAT(GraphBuilder, Condition, StatName, Format, ...) RDG_EVENT_SCOPE_IMPL(GraphBuilder, Condition, ERDGScopeFlags::Stat , RHI_GPU_STAT_ARGS(StatName), TEXT(Format), nullptr)
+
+		#define RDG_EVENT_SCOPE(                 GraphBuilder,                      Format, ...) do { } while (false)
+		#define RDG_EVENT_SCOPE_CONDITIONAL(     GraphBuilder, Condition,           Format, ...) do { } while (false)
+		#define RDG_EVENT_SCOPE_FINAL(           GraphBuilder,                      Format, ...) do { } while (false)
+		#define RDG_EVENT_SCOPE_CONDITIONAL_NAME(GraphBuilder, Condition,           EventName  ) do { } while (false)
+
+	#endif
 
 #else
 
diff --git a/Engine/Source/Runtime/Windows/D3D11RHI/Private/D3D11Device.cpp b/Engine/Source/Runtime/Windows/D3D11RHI/Private/D3D11Device.cpp
--- a/Engine/Source/Runtime/Windows/D3D11RHI/Private/D3D11Device.cpp
+++ b/Engine/Source/Runtime/Windows/D3D11RHI/Private/D3D11Device.cpp
@@ -268,36 +268,36 @@ void FD3D11DynamicRHI::Shutdown()
 		UE::RHICore::Nvidia::Aftermath::D3D11::BeginBreadcrumb(AftermathHandle, Breadcrumb);
 	#endif
 
-		if (ShouldEmitBreadcrumbs())
-		{
-		#if RHI_NEW_GPU_PROFILER
-			FlushProfilerStats();
+	#if RHI_NEW_GPU_PROFILER
+		FlushProfilerStats();
 
-			auto& Event = EmplaceProfilerEvent<UE::RHI::GPUProfiler::FEvent::FBeginBreadcrumb>(Breadcrumb);
-			InsertProfilerTimestamp(&Event.GPUTimestampTOP);
-		#else
+		auto& Event = EmplaceProfilerEvent<UE::RHI::GPUProfiler::FEvent::FBeginBreadcrumb>(Breadcrumb);
+		InsertProfilerTimestamp(&Event.GPUTimestampTOP);
+	#else
+		if (GPUProfilingData.IsProfilingGPU())
+		{
 			// @todo dev-pr avoid TCHAR -> ANSI conversion
 			FRHIBreadcrumb::FBuffer Buffer;
 			TCHAR const* Name = Breadcrumb->GetTCHAR(Buffer);
 
 			GPUProfilingData.PushEvent(Name, FColor::White);
-		#endif
 		}
+	#endif
 	}
 
 	void FD3D11DynamicRHI::RHIEndBreadcrumbGPU(FRHIBreadcrumbNode* Breadcrumb)
 	{
-		if (ShouldEmitBreadcrumbs())
-		{
-		#if RHI_NEW_GPU_PROFILER
-			FlushProfilerStats();
+	#if RHI_NEW_GPU_PROFILER
+		FlushProfilerStats();
 
-			auto& Event = EmplaceProfilerEvent<UE::RHI::GPUProfiler::FEvent::FEndBreadcrumb>(Breadcrumb);
-			InsertProfilerTimestamp(&Event.GPUTimestampBOP);
-		#else
+		auto& Event = EmplaceProfilerEvent<UE::RHI::GPUProfiler::FEvent::FEndBreadcrumb>(Breadcrumb);
+		InsertProfilerTimestamp(&Event.GPUTimestampBOP);
+	#else
+		if (GPUProfilingData.IsProfilingGPU())
+		{
 			GPUProfilingData.PopEvent();
-		#endif
 		}
+	#endif
 
 	#if NV_AFTERMATH
 		UE::RHICore::Nvidia::Aftermath::D3D11::EndBreadcrumb(AftermathHandle, Breadcrumb);
