bRequiresBindfulUtilityShaders should be set based on the host preview platform.

diff --git a/Engine/Source/Runtime/RHI/Private/DataDrivenShaderPlatformInfo.cpp b/Engine/Source/Runtime/RHI/Private/DataDrivenShaderPlatformInfo.cpp
--- a/Engine/Source/Runtime/RHI/Private/DataDrivenShaderPlatformInfo.cpp
+++ b/Engine/Source/Runtime/RHI/Private/DataDrivenShaderPlatformInfo.cpp
@@ -571,6 +571,7 @@ void FGenericDataDrivenShaderPlatformInfo::UpdatePreviewPlatforms()
 				PREVIEW_USE_RUNTIME_VALUE(bSupportsManualVertexFetch);
 				PREVIEW_USE_RUNTIME_VALUE(bSupportsRealTypes);
 				PREVIEW_USE_RUNTIME_VALUE(bSupportsUniformBufferObjects);
+				PREVIEW_USE_RUNTIME_VALUE(bRequiresBindfulUtilityShaders);
 
 				PREVIEW_USE_RUNTIME_VALUE(bSupportsInlineRayTracing);
 				PREVIEW_USE_RUNTIME_VALUE(bInlineRayTracingRequiresBindless);
