Fix races for NewPSOsToReport and PSOLoggedEvent members of FPipelineFileCacheManager.

NewPSOsToReport was modified under FileCacheLock, but only with read access. Use TMpscQueue instead since the queue is always processed on the game thread.

PSOLoggedEvent is not thread-safe to be calling IsBound() away from the game thread. Just remove the IsBound() calls and live with the idea that we might throw the events away. If we care enough, we can probably augment FPipelineFileCacheManager::ReportNewPSOs() to track with an atomic counter or change the delegate to TS or something, but I'm just getting rid of the race for now.

diff --git a/Engine/Source/Runtime/RHI/Private/PipelineFileCache.cpp b/Engine/Source/Runtime/RHI/Private/PipelineFileCache.cpp
--- a/Engine/Source/Runtime/RHI/Private/PipelineFileCache.cpp
+++ b/Engine/Source/Runtime/RHI/Private/PipelineFileCache.cpp
@@ -183,7 +183,7 @@ TMap<uint32, FPSOUsageData> FPipelineFileCacheManager::RunTimeToPSOUsage;
 TMap<uint32, FPSOUsageData> FPipelineFileCacheManager::NewPSOUsage;
 TMap<uint32, FPipelineStateStats*> FPipelineFileCacheManager::Stats;
 TSet<FPipelineCacheFileFormatPSO> FPipelineFileCacheManager::NewPSOs;
-TArray<FPipelineCacheFileFormatPSO> FPipelineFileCacheManager::NewPSOsToReport;
+TMpscQueue<FPipelineCacheFileFormatPSO> FPipelineFileCacheManager::NewPSOsToReport;
 TSet<uint32> FPipelineFileCacheManager::NewPSOHashes;
 uint32 FPipelineFileCacheManager::NumNewPSOs;
 FString FPipelineFileCacheManager::UserCacheKey;
@@ -3536,9 +3536,11 @@ void FPipelineFileCacheManager::BroadcastNewPSOsDelegate()
 {
 	TArray<FPipelineCacheFileFormatPSO> PSOs;
 	{
-		FRWScopeLock Lock(FileCacheLock, SLT_Write);
-		PSOs = MoveTemp(NewPSOsToReport);
-		NewPSOsToReport.Empty(32);
+		FPipelineCacheFileFormatPSO PSO;
+		while (NewPSOsToReport.Dequeue(PSO))
+		{
+			PSOs.Emplace(MoveTemp(PSO));
+		}
 	}
 
 	if (!PSOs.IsEmpty() && ReportNewPSOs())
@@ -3615,10 +3617,10 @@ void FPipelineFileCacheManager::CacheGraphicsPSO(uint32 RunTimeHash, FGraphicsPi
 						NumNewPSOs++;
 						INC_DWORD_STAT(STAT_NewGraphicsPipelineStateCount);
 						INC_DWORD_STAT(STAT_TotalGraphicsPipelineStateCount);
-							
-						if (ReportNewPSOs() && PSOLoggedEvent.IsBound())
+
+						if (ReportNewPSOs())
 						{
-							NewPSOsToReport.Emplace(MoveTemp(NewEntry));
+							NewPSOsToReport.Enqueue(MoveTemp(NewEntry));
 						}
 					}
 				}
@@ -3697,9 +3699,9 @@ void FPipelineFileCacheManager::CacheComputePSO(uint32 RunTimeHash, FRHIComputeS
 							INC_DWORD_STAT(STAT_NewComputePipelineStateCount);
 							INC_DWORD_STAT(STAT_TotalComputePipelineStateCount);
 
-							if (ReportNewPSOs() && PSOLoggedEvent.IsBound())
+							if (ReportNewPSOs())
 							{
-								NewPSOsToReport.Emplace(MoveTemp(NewEntry));
+								NewPSOsToReport.Enqueue(MoveTemp(NewEntry));
 							}
 						}
 					}
@@ -3784,9 +3786,9 @@ void FPipelineFileCacheManager::CacheRayTracingPSO(const FRayTracingPipelineStat
 						INC_DWORD_STAT(STAT_NewRayTracingPipelineStateCount);
 						INC_DWORD_STAT(STAT_TotalRayTracingPipelineStateCount);
 
-						if (ReportNewPSOs() && PSOLoggedEvent.IsBound())
+						if (ReportNewPSOs())
 						{
-							NewPSOsToReport.Emplace(MoveTemp(NewEntry));
+							NewPSOsToReport.Enqueue(MoveTemp(NewEntry));
 						}
 					}
 
diff --git a/Engine/Source/Runtime/RHI/Public/PipelineFileCache.h b/Engine/Source/Runtime/RHI/Public/PipelineFileCache.h
--- a/Engine/Source/Runtime/RHI/Public/PipelineFileCache.h
+++ b/Engine/Source/Runtime/RHI/Public/PipelineFileCache.h
@@ -4,6 +4,7 @@
 
 #include "CoreMinimal.h"
 #include "Containers/List.h"
+#include "Containers/MpscQueue.h"
 #include "Containers/StringView.h"
 #include "RHI.h"
 #include "RHIShaderBindingLayout.h"
@@ -440,7 +441,7 @@ private:
 	static TMap<uint32, FPSOUsageData> NewPSOUsage;				// For mask or engine updates - Merged + Saved (Our internal PSO hash to latest usage data) - temp working scratch, only holds updates since last "save" so is not the authority on state
 	static TMap<uint32, FPipelineStateStats*> Stats;
 	static TSet<FPipelineCacheFileFormatPSO> NewPSOs;
-	static TArray<FPipelineCacheFileFormatPSO> NewPSOsToReport; // New PSOs that will be broadcast via a delegate (if bound). Cleared when the delegate is broadcast.
+	static TMpscQueue<FPipelineCacheFileFormatPSO> NewPSOsToReport; // New PSOs that will be broadcast via a delegate (if bound). Cleared when the delegate is broadcast.
  	static TSet<uint32> NewPSOHashes;
 	static uint32 NumNewPSOs;
 	static PSOOrder RequestedOrder;
