Fix unmatched EndCPU calls for RHI breadcrumbs used in PSO compilation
- The PSO compilation tasks call WalkIn/WalkOut to push the breadcrumb that was at the top of the RHICmdList's stack at the point compilation is kicked. However, since PSO compilation is started from the top-of-pipe, the full breadcrumb tree has not been constructed yet.
- This results in a race between the PSO compilation thread, and the dispatch thread. WalkIn only sees the RHICmdList's sub-tree, ending in the Sentinel node. Later, WalkOut sees the full tree (the sub-tree has been attached to the root via the dispatch thread), so more End calls are made than necessary.
- The fix is to capture the sub-tree root node when starting the task, and only walk in/out of that sub-tree. This prevents the race with the dispatch thread.

diff --git a/Engine/Source/Runtime/RHI/Private/PipelineStateCache.cpp b/Engine/Source/Runtime/RHI/Private/PipelineStateCache.cpp
--- a/Engine/Source/Runtime/RHI/Private/PipelineStateCache.cpp
+++ b/Engine/Source/Runtime/RHI/Private/PipelineStateCache.cpp
@@ -229,7 +229,8 @@ struct FPSOCompilationDebugData
 {
 	FString PSOCompilationEventName;
 #if WITH_RHI_BREADCRUMBS
-	FRHIBreadcrumbNode* BreadcrumbNode = nullptr;
+	FRHIBreadcrumbNode const* BreadcrumbRoot = nullptr;
+	FRHIBreadcrumbNode const* BreadcrumbNode = nullptr;
 #endif // WITH_RHI_BREADCRUMBS
 };
 
@@ -3019,7 +3020,10 @@ public:
 #if WITH_RHI_BREADCRUMBS
 		if (PSOCompilationDebugData.BreadcrumbNode)
 		{
-			FRHIBreadcrumbNode::WalkIn(PSOCompilationDebugData.BreadcrumbNode);
+			FRHIBreadcrumbNode::WalkInRange(
+				PSOCompilationDebugData.BreadcrumbNode,
+				PSOCompilationDebugData.BreadcrumbRoot
+			);
 		}
 #endif // WITH_RHI_BREADCRUMBS
 
@@ -3204,7 +3208,10 @@ public:
 #if WITH_RHI_BREADCRUMBS
 		if (PSOCompilationDebugData.BreadcrumbNode)
 		{
-			FRHIBreadcrumbNode::WalkOut(PSOCompilationDebugData.BreadcrumbNode);
+			FRHIBreadcrumbNode::WalkOutRange(
+				PSOCompilationDebugData.BreadcrumbNode,
+				PSOCompilationDebugData.BreadcrumbRoot
+			);
 		}
 #endif // WITH_RHI_BREADCRUMBS
 		
@@ -3412,7 +3419,13 @@ FComputePipelineState* PipelineStateCache::GetAndOrCreateComputePipelineState(FR
 
 		FPSOCompilationDebugData PSOCompilationDebugData;
 #if WITH_RHI_BREADCRUMBS
-		PSOCompilationDebugData.BreadcrumbNode = DoAsyncCompile ? RHICmdList.GetCurrentBreadcrumbRef() : nullptr;
+		if (DoAsyncCompile)
+		{
+			PSOCompilationDebugData.BreadcrumbRoot = FRHIBreadcrumbNode::GetNonNullRoot(RHICmdList.GetCurrentBreadcrumbRef());
+			PSOCompilationDebugData.BreadcrumbNode = PSOCompilationDebugData.BreadcrumbRoot
+				? RHICmdList.GetCurrentBreadcrumbRef()
+				: nullptr;
+		}
 #endif // WITH_RHI_BREADCRUMBS
 		
 		FGraphEventRef GraphEvent = OutCachedState->GetCompletionEvent();
@@ -3966,7 +3979,13 @@ FGraphicsPipelineState* PipelineStateCache::GetAndOrCreateGraphicsPipelineState(
 
 		FPSOCompilationDebugData PSOCompilationDebugData;
 #if WITH_RHI_BREADCRUMBS
-		PSOCompilationDebugData.BreadcrumbNode = DoAsyncCompile ? RHICmdList.GetCurrentBreadcrumbRef() : nullptr;
+		if (DoAsyncCompile)
+		{
+			PSOCompilationDebugData.BreadcrumbRoot = FRHIBreadcrumbNode::GetNonNullRoot(RHICmdList.GetCurrentBreadcrumbRef());
+			PSOCompilationDebugData.BreadcrumbNode = PSOCompilationDebugData.BreadcrumbRoot
+				? RHICmdList.GetCurrentBreadcrumbRef()
+				: nullptr;
+		}
 #endif // WITH_RHI_BREADCRUMBS
 
 		bool bPSOPrecache = Initializer.bFromPSOFileCache;
diff --git a/Engine/Source/Runtime/RHI/Private/RHIBreadcrumbs.cpp b/Engine/Source/Runtime/RHI/Private/RHIBreadcrumbs.cpp
--- a/Engine/Source/Runtime/RHI/Private/RHIBreadcrumbs.cpp
+++ b/Engine/Source/Runtime/RHI/Private/RHIBreadcrumbs.cpp
@@ -34,6 +34,21 @@ RHI_API uint32 FRHIBreadcrumbNode::GetLevel(FRHIBreadcrumbNode const* Node)
 	return Result;
 }
 
+RHI_API FRHIBreadcrumbNode const* FRHIBreadcrumbNode::GetNonNullRoot(FRHIBreadcrumbNode const* Node)
+{
+	if (Node == nullptr || Node == Sentinel)
+	{
+		return nullptr;
+	}
+
+	while (Node->GetParent() != nullptr && Node->GetParent() != Sentinel)
+	{
+		Node = Node->GetParent();
+	}
+
+	return Node;
+}
+
 RHI_API FRHIBreadcrumbNode const* FRHIBreadcrumbNode::FindCommonAncestor(FRHIBreadcrumbNode const* Node0, FRHIBreadcrumbNode const* Node1)
 {
 	uint32 Level0 = GetLevel(Node0);
diff --git a/Engine/Source/Runtime/RHI/Public/RHIBreadcrumbs.h b/Engine/Source/Runtime/RHI/Public/RHIBreadcrumbs.h
--- a/Engine/Source/Runtime/RHI/Public/RHIBreadcrumbs.h
+++ b/Engine/Source/Runtime/RHI/Public/RHIBreadcrumbs.h
@@ -269,10 +269,18 @@ namespace UE::RHI::GPUProfiler
 		inline void EndCPU() const;
 
 		// Calls BeginCPU() on all the breadcrumb nodes between the root and the specified node.
-		static inline void WalkIn(FRHIBreadcrumbNode* Node);
+		// Only valid to call from the bottom-of-pipe, after the dispatch thread has fixed up the breadcrumb tree.
+		static inline void WalkIn(FRHIBreadcrumbNode const* Node);
+
+		// Same as WalkIn, but the root node is specified, allowing it to be called from the top-of-pipe.
+		static inline void WalkInRange(FRHIBreadcrumbNode const* Leaf, FRHIBreadcrumbNode const* Root);
 
 		// Calls EndCPU() on all the breadcrumb nodes between the specified node and the root.
-		static inline void WalkOut(FRHIBreadcrumbNode* Node);
+		// Only valid to call from the bottom-of-pipe, after the dispatch thread has fixed up the breadcrumb tree.
+		static inline void WalkOut(FRHIBreadcrumbNode const* Node);
+
+		// Same as WalkOut, but the root node is specified, allowing it to be called from the top-of-pipe.
+		static inline void WalkOutRange(FRHIBreadcrumbNode const* Leaf, FRHIBreadcrumbNode const* Root);
 
 		// ----------------------------------------------------
 		// Debug logging / crash reporting
@@ -288,6 +296,8 @@ namespace UE::RHI::GPUProfiler
 		static RHI_API FRHIBreadcrumbNode const* FindCommonAncestor(FRHIBreadcrumbNode const* Node0, FRHIBreadcrumbNode const* Node1);
 		static RHI_API uint32 GetLevel(FRHIBreadcrumbNode const* Node);
 
+		static RHI_API FRHIBreadcrumbNode const* GetNonNullRoot(FRHIBreadcrumbNode const* Node);
+
 		// A constant pointer value representing an undefined node. Used as the parent pointer for nodes in sub-trees
 		// that haven't been attached to the root yet, specifically to be distinct from nullptr which is the root.
 		static RHI_API FRHIBreadcrumbNode* const Sentinel;
@@ -665,25 +675,46 @@ namespace UE::RHI::GPUProfiler
 	#endif
 	}
 
-	inline void FRHIBreadcrumbNode::WalkIn(FRHIBreadcrumbNode* Node)
+	inline void FRHIBreadcrumbNode::WalkIn(FRHIBreadcrumbNode const* Node)
 	{
 	#if RHI_BREADCRUMBS_EMIT_CPU
 		if (TRACE_CPUPROFILER_EVENT_MANUAL_IS_ENABLED())
 		{
-			auto Recurse = [](FRHIBreadcrumbNode* Current, auto& Recurse) -> void
+			auto Recurse = [](auto const& Recurse, FRHIBreadcrumbNode const* Current) -> void
 			{
 				if (!Current || Current == Sentinel)
 					return;
 
-				Recurse(Current->GetParent(), Recurse);
+				Recurse(Recurse, Current->GetParent());
+				Current->BeginCPU();
+			};
+			Recurse(Recurse, Node);
+		}
+	#endif
+	}
+
+	inline void FRHIBreadcrumbNode::WalkInRange(FRHIBreadcrumbNode const* Leaf, FRHIBreadcrumbNode const* Root)
+	{
+		check(Leaf && Root);
+
+	#if RHI_BREADCRUMBS_EMIT_CPU
+		if (TRACE_CPUPROFILER_EVENT_MANUAL_IS_ENABLED())
+		{
+			auto Recurse = [&](auto const& Recurse, FRHIBreadcrumbNode const* Current) -> void
+			{
+				if (Current != Root)
+				{
+					Recurse(Recurse, Current->GetParent());
+				}
+
 				Current->BeginCPU();
 			};
-			Recurse(Node, Recurse);
+			Recurse(Recurse, Leaf);
 		}
 	#endif
 	}
 
-	inline void FRHIBreadcrumbNode::WalkOut(FRHIBreadcrumbNode* Node)
+	inline void FRHIBreadcrumbNode::WalkOut(FRHIBreadcrumbNode const* Node)
 	{
 	#if RHI_BREADCRUMBS_EMIT_CPU
 		if (TRACE_CPUPROFILER_EVENT_MANUAL_IS_ENABLED())
@@ -697,6 +728,25 @@ namespace UE::RHI::GPUProfiler
 	#endif
 	}
 
+	inline void FRHIBreadcrumbNode::WalkOutRange(FRHIBreadcrumbNode const* Leaf, FRHIBreadcrumbNode const* Root)
+	{
+		check(Leaf && Root);
+
+	#if RHI_BREADCRUMBS_EMIT_CPU
+		if (TRACE_CPUPROFILER_EVENT_MANUAL_IS_ENABLED())
+		{
+			while (true)
+			{
+				Leaf->EndCPU();
+				if (Leaf == Root)
+					break;
+
+				Leaf = Leaf->GetParent();
+			}
+		}
+	#endif
+	}
+
 	namespace UE::RHI::Breadcrumbs::Private
 	{
 		// Replacement for std::remove_cvref, since this isn't available until C++20.
