Make FPipelineStateLoggedEvent use const for sanity and remove the mutable keyword from the lambda in FPipelineFileCacheManager::BroadcastNewPSOsDelegate

diff --git a/Engine/Source/Runtime/RHI/Private/PipelineFileCache.cpp b/Engine/Source/Runtime/RHI/Private/PipelineFileCache.cpp
--- a/Engine/Source/Runtime/RHI/Private/PipelineFileCache.cpp
+++ b/Engine/Source/Runtime/RHI/Private/PipelineFileCache.cpp
@@ -3546,11 +3546,11 @@ void FPipelineFileCacheManager::BroadcastNewPSOsDelegate()
 	if (!PSOs.IsEmpty() && ReportNewPSOs())
 	{
 		// It's not safe to touch UObjects-based delegates from the render thread.
-		ExecuteOnGameThread(TEXT("OnPipelineStateLoggedBroadcastGT"), [PSOs = MoveTemp(PSOs)]() mutable
+		ExecuteOnGameThread(TEXT("OnPipelineStateLoggedBroadcastGT"), [InPSOs = MoveTemp(PSOs)]()
 		{
 			if (PSOLoggedEvent.IsBound())
 			{
-				for (FPipelineCacheFileFormatPSO& PSO : PSOs)
+				for (const FPipelineCacheFileFormatPSO& PSO : InPSOs)
 				{
 					PSOLoggedEvent.Broadcast(PSO);
 				}
diff --git a/Engine/Source/Runtime/RHI/Public/PipelineFileCache.h b/Engine/Source/Runtime/RHI/Public/PipelineFileCache.h
--- a/Engine/Source/Runtime/RHI/Public/PipelineFileCache.h
+++ b/Engine/Source/Runtime/RHI/Public/PipelineFileCache.h
@@ -362,7 +362,7 @@ public:
 	/**
 	 * Event signature for being notified that a new PSO has been logged
 	 */
-	DECLARE_MULTICAST_DELEGATE_OneParam(FPipelineStateLoggedEvent, FPipelineCacheFileFormatPSO&);
+	DECLARE_MULTICAST_DELEGATE_OneParam(FPipelineStateLoggedEvent, const FPipelineCacheFileFormatPSO&);
 
 	/**
 	 * Gets the event delegate to register for pipeline state logging events.
