Stop RHI breadcrumbs from being included in the build with the old GPU profiler
- Was causing an unintended RHI thread perf regression when using the old GPU profiler.

diff --git a/Engine/Source/Runtime/RHI/Public/RHIBreadcrumbs.h b/Engine/Source/Runtime/RHI/Public/RHIBreadcrumbs.h
--- a/Engine/Source/Runtime/RHI/Public/RHIBreadcrumbs.h
+++ b/Engine/Source/Runtime/RHI/Public/RHIBreadcrumbs.h
@@ -21,7 +21,7 @@
 // (RHI breadcrumb allocators, platform RHI implemntation etc).
 //
 #ifndef WITH_RHI_BREADCRUMBS
-#define WITH_RHI_BREADCRUMBS (UE_BUILD_DEBUG || UE_BUILD_DEVELOPMENT || WITH_PROFILEGPU || HAS_GPU_STATS)
+#define WITH_RHI_BREADCRUMBS (UE_BUILD_DEBUG || UE_BUILD_DEVELOPMENT || WITH_PROFILEGPU || (HAS_GPU_STATS && RHI_NEW_GPU_PROFILER))
 #endif
 
 // Enables all RDG/RHI breadcrumb scopes, and features such as Insights markers.
