RHI: Fix RHI validation for parallel contexts
- Added wrapper RHICloseTranslateChain for RHIFinalizeContext to enable validation to track when deferring context finalize

diff --git a/Engine/Source/Runtime/RHI/Private/RHICommandList.cpp b/Engine/Source/Runtime/RHI/Private/RHICommandList.cpp
--- a/Engine/Source/Runtime/RHI/Private/RHICommandList.cpp
+++ b/Engine/Source/Runtime/RHI/Private/RHICommandList.cpp
@@ -1120,15 +1120,6 @@ void FRHICommandListExecutor::FTranslateState::Translate(FRHICommandListBase* Cm
 
 FGraphEventRef FRHICommandListExecutor::FTranslateState::Finalize()
 {
-	if(!bShouldFinalize)
-	{
-		// Don't finalize if this chain ended with the start of a parallel render pass,
-		// and the contexts will be forwarded to the chain that processes the rest of the parent.
-		ENamedThreads::Type NamedThread = ENamedThreads::AnyHiPriThreadHiPriTask;
-		FTaskPipe* Pipe = GetTranslateTaskPipe(NamedThread);
-		return Pipe->Close();
-	}
-	
 	FTaskPipe* Pipe = EnqueueTranslateTask({},
 		[this]()
 		{
@@ -1155,7 +1146,7 @@ FGraphEventRef FRHICommandListExecutor::FTranslateState::Finalize()
 				}
 				else
 				{
-					GDynamicRHI->RHIFinalizeContext(MoveTemp(FinalizeArgs), PlatformCommandLists);
+					GDynamicRHI->RHICloseTranslateChain(MoveTemp(FinalizeArgs), PlatformCommandLists, bShouldFinalize);
 
 					for (ERHIPipeline Pipeline : MakeFlagsRange(ERHIPipeline::All))
 					{
diff --git a/Engine/Source/Runtime/RHI/Private/RHIValidation.cpp b/Engine/Source/Runtime/RHI/Private/RHIValidation.cpp
--- a/Engine/Source/Runtime/RHI/Private/RHIValidation.cpp
+++ b/Engine/Source/Runtime/RHI/Private/RHIValidation.cpp
@@ -686,6 +686,71 @@ void FValidationRHI::RHIFinalizeContext(FRHIFinalizeContextArgs&& Args, TRHIPipe
 	}
 }
 
+IRHIComputeContext* FValidationRHI::RHIGetParallelCommandContext(FRHIParallelRenderPassInfo const& ParallelRenderPass, FRHIGPUMask GPUMask)
+{
+	IRHIComputeContext* InnerContext = RHI->RHIGetParallelCommandContext(ParallelRenderPass, GPUMask);
+	check(InnerContext);
+
+	FValidationContext* OuterContext = new FValidationContext(FValidationContext::EType::Parallel);
+	OuterContext->LinkToContext(static_cast<IRHICommandContext*>(InnerContext));
+
+	// Parallel contexts are always inside a renderpass
+	OuterContext->State.bInsideBeginRenderPass = true;
+	OuterContext->State.RenderPassInfo = ParallelRenderPass;
+	if (ParallelRenderPass.PassName)
+	{
+		OuterContext->State.RenderPassName = ParallelRenderPass.PassName;
+	}
+
+	return OuterContext;
+}
+
+void FValidationRHI::RHICloseTranslateChain(FRHIFinalizeContextArgs&& Args, TRHIPipelineArray<IRHIPlatformCommandList*>& Output, bool bShouldFinalize)
+{
+	// If we aren't finalizing the context we need to finalize the tracking
+	if(!bShouldFinalize)
+	{
+		for(IRHIComputeContext* Context : Args.Contexts)
+		{
+			if(Context)
+			{
+				check(Context->GetPipeline() == ERHIPipeline::Graphics);
+				
+				IRHIComputeContext& InnerContext = Context->GetLowestLevelContext();
+				
+				FValidationCommandList* OuterCommandList = new FValidationCommandList();
+				OuterCommandList->CompletedOpList = InnerContext.Tracker->Finalize();
+				OuterCommandList->Pipeline = ERHIPipeline::Graphics;
+				
+				Output[ERHIPipeline::Graphics] = OuterCommandList;
+			}
+		}
+		
+		return;
+	}
+
+	FDynamicRHI::RHICloseTranslateChain(MoveTemp(Args), Output, bShouldFinalize);
+}
+
+IRHIPlatformCommandList* FValidationRHI::RHIFinalizeParallelContext(IRHIComputeContext* Context)
+{
+	check(Context->GetPipeline() == ERHIPipeline::Graphics);
+
+	IRHIComputeContext& InnerContext = Context->GetLowestLevelContext();
+
+	FValidationCommandList* OuterCommandList = new FValidationCommandList();
+	OuterCommandList->CompletedOpList = InnerContext.Tracker->Finalize();
+	OuterCommandList->Pipeline = ERHIPipeline::Graphics;
+
+	IRHIPlatformCommandList* InnerCommandList = RHI->RHIFinalizeParallelContext(&InnerContext);
+	OuterCommandList->InnerCommandLists[ERHIPipeline::Graphics] = InnerCommandList;
+
+	check(static_cast<FValidationContext*>(Context)->Type == FValidationContext::EType::Parallel);
+	delete Context;
+
+	return OuterCommandList;
+}
+
 void FValidationRHI::RHISubmitCommandLists(FRHISubmitCommandListsArgs&& Args)
 {
 	FDynamicRHI::FRHISubmitCommandListsArgs InnerArgs;
diff --git a/Engine/Source/Runtime/RHI/Public/DynamicRHI.h b/Engine/Source/Runtime/RHI/Public/DynamicRHI.h
--- a/Engine/Source/Runtime/RHI/Public/DynamicRHI.h
+++ b/Engine/Source/Runtime/RHI/Public/DynamicRHI.h
@@ -912,6 +912,20 @@ public:
 		IRHIUploadContext* UploadContext;
 	};
 	
+	//
+	// Close the current translate chain or ignore if we are not finalizing
+	//
+	#if ENABLE_RHI_VALIDATION
+	virtual
+	#endif
+	void RHICloseTranslateChain(FRHIFinalizeContextArgs&& Args, TRHIPipelineArray<IRHIPlatformCommandList*>& Output, bool bShouldFinalize)
+	{
+		if(bShouldFinalize)
+		{
+			RHIFinalizeContext(MoveTemp(Args), Output);
+		}
+	}
+
 	virtual void RHIFinalizeContext(FRHIFinalizeContextArgs&& Args, TRHIPipelineArray<IRHIPlatformCommandList*>& Output) = 0;
 	
 	virtual IRHIPlatformCommandList* RHIFinalizeParallelContext(IRHIComputeContext* Context)
diff --git a/Engine/Source/Runtime/RHI/Public/RHIResources.h b/Engine/Source/Runtime/RHI/Public/RHIResources.h
--- a/Engine/Source/Runtime/RHI/Public/RHIResources.h
+++ b/Engine/Source/Runtime/RHI/Public/RHIResources.h
@@ -5469,7 +5469,7 @@ struct FRHIParallelRenderPassInfo : public FRHIRenderPassInfo
 	const TCHAR* PassName;
 	
 	// Contains platform specific data
-	void* RHIPlatformData;
+	void* RHIPlatformData = nullptr;
 
 	FRHIParallelRenderPassInfo(FRHIRenderPassInfo&& Info, const TCHAR* InPassName)
 		: FRHIRenderPassInfo(MoveTemp(Info)), PassName(InPassName)
diff --git a/Engine/Source/Runtime/RHI/Public/RHIValidation.h b/Engine/Source/Runtime/RHI/Public/RHIValidation.h
--- a/Engine/Source/Runtime/RHI/Public/RHIValidation.h
+++ b/Engine/Source/Runtime/RHI/Public/RHIValidation.h
@@ -997,9 +997,14 @@ public:
 		return RHI->RHIGetUploadContext();
 	}
 	
+	virtual void RHICloseTranslateChain(FRHIFinalizeContextArgs&& Args, TRHIPipelineArray<IRHIPlatformCommandList*>& Output, bool bShouldFinalize) override final;
+	
 	virtual void RHIFinalizeContext(FRHIFinalizeContextArgs&& Args, TRHIPipelineArray<IRHIPlatformCommandList*>& Output) override final;
 	virtual void RHISubmitCommandLists(FRHISubmitCommandListsArgs&& Args) override final;
 	
+	virtual IRHIComputeContext* RHIGetParallelCommandContext(FRHIParallelRenderPassInfo const& ParallelRenderPass, FRHIGPUMask GPUMask) override final;
+	virtual IRHIPlatformCommandList* RHIFinalizeParallelContext(IRHIComputeContext* Context) override final;
+
 	virtual void RHIProcessDeleteQueue() override final
 	{
 		RHI->RHIProcessDeleteQueue();
diff --git a/Engine/Source/Runtime/RHI/Public/RHIValidationContext.h b/Engine/Source/Runtime/RHI/Public/RHIValidationContext.h
--- a/Engine/Source/Runtime/RHI/Public/RHIValidationContext.h
+++ b/Engine/Source/Runtime/RHI/Public/RHIValidationContext.h
@@ -1167,6 +1167,16 @@ public:
 		State.PreviousRenderPassName = State.RenderPassName;
 	}
 
+	virtual void RHIBeginParallelRenderPass(TSharedPtr<FRHIParallelRenderPassInfo> InInfo, const TCHAR* InName) override final
+	{
+		RHIContext->RHIBeginParallelRenderPass(InInfo, InName);
+	}
+
+	virtual void RHIEndParallelRenderPass() override final
+	{
+		RHIContext->RHIEndParallelRenderPass();
+	}
+
 	virtual void RHIWriteGPUFence(FRHIGPUFence* FenceRHI) override final
 	{
 		RHIContext->RHIWriteGPUFence(FenceRHI);
