Make it possible to change at runtime the number of precached PSOs kept in memory.

Also simplify the circular buffer logic by using a TDeque.

diff --git a/Engine/Source/Runtime/RHI/Private/PipelineStateCache.cpp b/Engine/Source/Runtime/RHI/Private/PipelineStateCache.cpp
--- a/Engine/Source/Runtime/RHI/Private/PipelineStateCache.cpp
+++ b/Engine/Source/Runtime/RHI/Private/PipelineStateCache.cpp
@@ -7,6 +7,7 @@ PipelineStateCache.cpp: Pipeline state cache implementation.
 #include "PipelineStateCache.h"
 #include "Async/AsyncWork.h"
 #include "Async/TaskGraphInterfaces.h"
+#include "Containers/Deque.h"
 #include "PipelineFileCache.h"
 #include "Misc/ScopeRWLock.h"
 #include "Misc/App.h"
@@ -344,7 +345,7 @@ static FAutoConsoleVariableRef CVarPSOPrecacheKeepInMemoryGraphicsMaxNum(
 	GPSOPrecacheKeepInMemoryGraphicsMaxNum,
 	TEXT("If r.PSOPrecache.KeepInMemoryUntilUsed is enabled, this value will control the maximum number of precached graphics PSOs that are kept in memory at a time.\n")
 	TEXT("If set to 0, no limit will be applied (not recommended outside of testing, as it can cause unbounded memory usage)."),
-	ECVF_ReadOnly);
+	ECVF_RenderThreadSafe);
 
 int32 GPSOPrecacheKeepInMemoryComputeMaxNum = 200;
 static FAutoConsoleVariableRef CVarPSOPrecacheKeepInMemoryComputeMaxNum(
@@ -352,7 +353,7 @@ static FAutoConsoleVariableRef CVarPSOPrecacheKeepInMemoryComputeMaxNum(
 	GPSOPrecacheKeepInMemoryComputeMaxNum,
 	TEXT("If r.PSOPrecache.KeepInMemoryUntilUsed is enabled, this value will control the maximum number of precached compute PSOs that are kept in memory at a time.\n")
 	TEXT("If set to 0, no limit will be applied (not recommended outside of testing, as it can cause unbounded memory usage)."),
-	ECVF_ReadOnly);
+	ECVF_RenderThreadSafe);
 
 bool ShouldKeepPrecachedPSOsInMemory()
 {
@@ -2062,7 +2063,7 @@ public:
 	{
 		if (ShouldKeepPrecachedPSOsInMemory() && InMaxInMemoryPSOs > 0)
 		{
-			InMemoryPSOIndices.Init(kInvalidPSOIndex, InMaxInMemoryPSOs);
+			InMemoryPSOIndices.Reserve(InMaxInMemoryPSOs);
 		}
 	}
 
@@ -2071,6 +2072,30 @@ public:
 		// Wait for all precache tasks to have finished.
 		WaitTasksComplete();
 	}
+
+	// Sets a new maximum number of precached PSOs kept in memory. FIFO order of currently tracked PSOs is maintained.
+	// If the new maximum is smaller than the current maximum, the oldest PSOs are released.
+	void SetMaxInMemoryPSOs(uint32 NewMaxInMemoryPSOs)
+	{
+		if (!ShouldKeepPrecachedPSOsInMemory() || MaxInMemoryPSOs == NewMaxInMemoryPSOs)
+		{
+			return;
+		}
+
+		FRWScopeLock WriteLock(PrecachePSOsRWLock, SLT_Write);
+		MaxInMemoryPSOs = NewMaxInMemoryPSOs;
+		InMemoryPSOIndices.Reserve(MaxInMemoryPSOs);
+
+		// Release oldest PSOs.
+		while (static_cast<uint32>(InMemoryPSOIndices.Num()) > MaxInMemoryPSOs)
+		{
+			uint32 PSOIndex = InMemoryPSOIndices.First();
+			InMemoryPSOIndices.PopFirst();
+
+			// Enqueue the corresponding PSO for cleanup.
+			PrecachedPSOsToCleanup.Add(PrecachedPSOInitializers[PSOIndex]);
+		}
+	}
 	
 protected:
 
@@ -2161,21 +2186,19 @@ protected:
 			{
 				if (MaxInMemoryPSOs > 0)
 				{
-					uint32& PSOIndex = InMemoryPSOIndices[CurrentInMemoryPSOIndex];
+					check(static_cast<uint32>(InMemoryPSOIndices.Num()) <= MaxInMemoryPSOs);
 
-					// Check if we're evicting a valid index. In that case, enqueue the corresponding PSO for cleanup.
-					if (PSOIndex != kInvalidPSOIndex)
+					// Evict the oldest PSO if we're at maximum capacity.
+					if (InMemoryPSOIndices.Num() == MaxInMemoryPSOs)
 					{
-						PrecachedPSOsToCleanup.Add(PrecachedPSOInitializers[PSOIndex]);
-					}
-					PSOIndex = Result.RequestID.RequestID;
+						uint32 PSOIndex = InMemoryPSOIndices.First();
+						InMemoryPSOIndices.PopFirst();
 
-					// Manage the circular buffer index.
-					CurrentInMemoryPSOIndex++;
-					if (CurrentInMemoryPSOIndex >= MaxInMemoryPSOs)
-					{
-						CurrentInMemoryPSOIndex = 0;
+						// Enqueue the corresponding PSO for cleanup.
+						PrecachedPSOsToCleanup.Add(PrecachedPSOInitializers[PSOIndex]);
 					}
+					
+					InMemoryPSOIndices.PushLast(Result.RequestID.RequestID);
 				}
 				INC_DWORD_STAT(STAT_InMemoryPrecachedPSOCount);
 			}
@@ -2542,10 +2565,8 @@ protected:
 	TArray<uint64> PrecachedPSOsToCleanup;
 
 	// Circular buffer of indices of PSOs that are kept in memory if MaxInMemoryPSOs is non-zero.
-	TArray<uint32> InMemoryPSOIndices;
+	TDeque<uint32> InMemoryPSOIndices;
 	uint32 MaxInMemoryPSOs = 0;
-	uint32 CurrentInMemoryPSOIndex = 0;
-	static constexpr uint32 kInvalidPSOIndex = TNumericLimits<uint32>::Max();
 };
 
 struct FPrecacheComputeInitializer
@@ -3274,6 +3295,8 @@ void PipelineStateCache::FlushResources()
 	GGraphicsPipelineCache.FlushResources(bDiscardAndSwap);
 
 	check(GPrecacheGraphicsPipelineCache && GPrecacheComputePipelineCache);
+	GPrecacheGraphicsPipelineCache->SetMaxInMemoryPSOs(GPSOPrecacheKeepInMemoryGraphicsMaxNum);
+	GPrecacheComputePipelineCache->SetMaxInMemoryPSOs(GPSOPrecacheKeepInMemoryComputeMaxNum);
 	GPrecacheGraphicsPipelineCache->ProcessDelayedCleanup();
 	GPrecacheComputePipelineCache->ProcessDelayedCleanup();
 
