Enable the new GPU profiler
- This replaces the implementation of various GPU profiling tools and systems (stat unit, stat gpu, profilegpu, Unreal Insight's GPU track, GPU markers in external platform profiling tools).
- The new profiler is multi-GPU and multi-pipe aware, and uses the RHI breadcrumbs system to provide the same markers as seen on the CPU in Unreal Insights.

diff --git a/Engine/Source/Runtime/RHI/Public/GpuProfilerTrace.h b/Engine/Source/Runtime/RHI/Public/GpuProfilerTrace.h
--- a/Engine/Source/Runtime/RHI/Public/GpuProfilerTrace.h
+++ b/Engine/Source/Runtime/RHI/Public/GpuProfilerTrace.h
@@ -9,7 +9,7 @@
 #include "UObject/NameTypes.h"
 
 #ifndef RHI_NEW_GPU_PROFILER
-#define RHI_NEW_GPU_PROFILER 0
+#define RHI_NEW_GPU_PROFILER 1
 #endif
 
 #if !defined(GPUPROFILERTRACE_ENABLED)
