Fix build error in GPUProfiler.cpp

diff --git a/Engine/Source/Runtime/RHI/Private/GPUProfiler.cpp b/Engine/Source/Runtime/RHI/Private/GPUProfiler.cpp
--- a/Engine/Source/Runtime/RHI/Private/GPUProfiler.cpp
+++ b/Engine/Source/Runtime/RHI/Private/GPUProfiler.cpp
@@ -1549,19 +1549,17 @@ namespace UE::RHI::GPUProfiler
 				Other.Inclusive = {};
 			}
 
-			void EmitResults(FQueue Queue
-			#if STATS && HAS_GPU_STATS
+		#if HAS_GPU_STATS
+			void EmitResults(FQueue Queue, FGPUStat& GPUStat
+			#if STATS
 				, FEndOfPipeStats* Stats
 			#endif
-			#if HAS_GPU_STATS
-				, FGPUStat& GPUStat
-			#endif
 			#if CSV_PROFILER_STATS
 				, FCsvProfiler* CsvProfiler
 			#endif
 			) const
 			{
-			#if STATS && HAS_GPU_STATS
+			#if STATS
 				Stats->AddMessage(GPUStat.GetStatId(Queue, FGPUStat::EType::Busy).GetName(), EStatOperation::Set, FPlatformTime::ToMilliseconds64(Inclusive.BusyCycles));
 				Stats->AddMessage(GPUStat.GetStatId(Queue, FGPUStat::EType::Idle).GetName(), EStatOperation::Set, FPlatformTime::ToMilliseconds64(Inclusive.IdleCycles));
 				Stats->AddMessage(GPUStat.GetStatId(Queue, FGPUStat::EType::Wait).GetName(), EStatOperation::Set, FPlatformTime::ToMilliseconds64(Inclusive.WaitCycles));
@@ -1594,6 +1592,7 @@ namespace UE::RHI::GPUProfiler
 				}
 			#endif
 			}
+		#endif
 		};
 
 		struct FQueueTimestamps
@@ -2363,36 +2362,32 @@ namespace UE::RHI::GPUProfiler
 
 			for (auto const& [Queue, QueueTimestamps] : FrameState)
 			{
-			#if WITH_RHI_BREADCRUMBS
+			#if WITH_RHI_BREADCRUMBS && HAS_GPU_STATS
 				// Compute the individual GPU stats
 				for (auto const& [Stat, StatState] : QueueTimestamps.Stats)
 				{
-					StatState.EmitResults(Queue
-					#if STATS && HAS_GPU_STATS
+					StatState.EmitResults(Queue, *Stat.GPUStat
+					#if STATS
 						, Stats
 					#endif
-					#if HAS_GPU_STATS
-						, *Stat.GPUStat
-					#endif
 					#if CSV_PROFILER_STATS
 						, bCsvStatsEnabled ? CsvProfiler : nullptr
 					#endif
 					);
 				}
-			#endif // WITH_RHI_BREADCRUMBS
+			#endif // WITH_RHI_BREADCRUMBS && HAS_GPU_STATS
 
 				// Set the whole-frame per queue stat
-				QueueTimestamps.WholeQueueStat.EmitResults(Queue
-				#if STATS && HAS_GPU_STATS
+			#if HAS_GPU_STATS
+				QueueTimestamps.WholeQueueStat.EmitResults(Queue, GPUStat_Total
+				#if STATS
 					, Stats
 				#endif
-				#if HAS_GPU_STATS
-					, GPUStat_Total
-				#endif
 				#if CSV_PROFILER_STATS
 					, bCsvStatsEnabled ? CsvProfiler : nullptr
 				#endif
 				);
+			#endif
 
 				if (QueueTimestamps.TotalBusyCycles.IsSet())
 				{
