Moving bCanUsePreviewPlatformForMaterialValidation to the EditorOnly DDSPI.

diff --git a/Engine/Source/Runtime/RHI/Private/DataDrivenShaderPlatformInfo.cpp b/Engine/Source/Runtime/RHI/Private/DataDrivenShaderPlatformInfo.cpp
--- a/Engine/Source/Runtime/RHI/Private/DataDrivenShaderPlatformInfo.cpp
+++ b/Engine/Source/Runtime/RHI/Private/DataDrivenShaderPlatformInfo.cpp
@@ -26,6 +26,7 @@ struct FDataDrivenShaderPlatformInfoEditorOnly
 {
 	FText FriendlyName;
 	EShaderPlatform PreviewShaderPlatformParent = SP_NumPlatforms;
+	uint32 bCanUsePreviewPlatformForMaterialValidation : 1;
 };
 static FDataDrivenShaderPlatformInfoEditorOnly DataDrivenShaderPlatformInfoEditorOnlyInfos[SP_NumPlatforms];
 
@@ -348,9 +349,6 @@ void FGenericDataDrivenShaderPlatformInfo::ParseDataDrivenShaderInfo(const FConf
 	GET_SECTION_BOOL_HELPER(bSupportsIndependentSamplers);
 	GET_SECTION_BOOL_HELPER(bSupportsWorkGraphs);
 	GET_SECTION_BOOL_HELPER(bSupportsWorkGraphsTier1_1);
-#if WITH_EDITOR
-	GET_SECTION_BOOL_HELPER(bCanUsePreviewPlatformForMaterialValidation);
-#endif
 	GET_SECTION_BOOL_HELPER(bSupportsDLSSShaders);
 #undef GET_SECTION_BOOL_HELPER
 #undef GET_SECTION_INT_HELPER
@@ -363,6 +361,8 @@ void FGenericDataDrivenShaderPlatformInfo::ParseDataDrivenShaderInfo(const FConf
 #if WITH_EDITOR
 	FDataDrivenShaderPlatformInfoEditorOnly& EditorInfo = DataDrivenShaderPlatformInfoEditorOnlyInfos[Index];
 	FTextStringHelper::ReadFromBuffer(*GetSectionString(Section, FName("FriendlyName")), EditorInfo.FriendlyName);
+
+	EditorInfo.bCanUsePreviewPlatformForMaterialValidation = GetSectionBool(Section, "bCanUsePreviewPlatformForMaterialValidation", false);
 #endif
 }
 
@@ -616,6 +616,13 @@ const EShaderPlatform FGenericDataDrivenShaderPlatformInfo::GetPreviewShaderPlat
 	check(IsValid(Platform));
 	return DataDrivenShaderPlatformInfoEditorOnlyInfos[Platform].PreviewShaderPlatformParent;
 }
+
+const bool FGenericDataDrivenShaderPlatformInfo::CanUseForMaterialValidation(const FStaticShaderPlatform Platform)
+{
+	check(IsValid(Platform));
+	return DataDrivenShaderPlatformInfoEditorOnlyInfos[Platform].bCanUsePreviewPlatformForMaterialValidation || !Infos[Platform].bIsPreviewPlatform;
+}
+
 #endif  // WITH_EDITOR
 
 const EShaderPlatform FGenericDataDrivenShaderPlatformInfo::GetShaderPlatformFromName(const FName ShaderPlatformName)
diff --git a/Engine/Source/Runtime/RHI/Public/DataDrivenShaderPlatformInfo.h b/Engine/Source/Runtime/RHI/Public/DataDrivenShaderPlatformInfo.h
--- a/Engine/Source/Runtime/RHI/Public/DataDrivenShaderPlatformInfo.h
+++ b/Engine/Source/Runtime/RHI/Public/DataDrivenShaderPlatformInfo.h
@@ -130,13 +130,10 @@ class FGenericDataDrivenShaderPlatformInfo
 	uint32 bSupportsIndependentSamplers : 1;
 	uint32 bSupportsWorkGraphs : 1;
 	uint32 bSupportsWorkGraphsTier1_1 : 1;
-#if WITH_EDITOR
-	uint32 bCanUsePreviewPlatformForMaterialValidation : 1;
-#endif
+	uint32 bSupportsDLSSShaders : 1;
 
 	// NOTE: When adding fields, you must also add to ParseDataDrivenShaderInfo!
 	uint32 bContainsValidPlatformInfo : 1;
-	uint32 bSupportsDLSSShaders : 1;
 
 	FGenericDataDrivenShaderPlatformInfo()
 	{
@@ -851,14 +848,6 @@ public:
 		return Infos[Platform].bSupportsWorkGraphsTier1_1;
 	}
 
-#if WITH_EDITOR
-	static FORCEINLINE_DEBUGGABLE const bool CanUseForMaterialValidation(const FStaticShaderPlatform Platform)
-	{
-		check(IsValid(Platform));
-		return Infos[Platform].bCanUsePreviewPlatformForMaterialValidation || !Infos[Platform].bIsPreviewPlatform;
-	}
-#endif
-
 	static FORCEINLINE_DEBUGGABLE const bool GetSupportsDLSSShaders(const FStaticShaderPlatform Platform)
 	{
 		check(IsValid(Platform));
@@ -874,6 +863,7 @@ public:
 	RHI_API static void UpdatePreviewPlatforms();
 	RHI_API static FText GetFriendlyName(const FStaticShaderPlatform Platform);
 	RHI_API static const EShaderPlatform GetPreviewShaderPlatformParent(const FStaticShaderPlatform Platform);
+	RHI_API static const bool CanUseForMaterialValidation(const FStaticShaderPlatform Platform);
 
 	RHI_API static TMap<FString, TFunction<bool(const FStaticShaderPlatform Platform)>> PropertyToShaderPlatformFunctionMap;
 #endif
