Removing deprecated RHI code.

diff --git a/Engine/Plugins/Runtime/OpenXR/Source/OpenXRHMD/Internal/OpenXRHMD_RenderBridge.h b/Engine/Plugins/Runtime/OpenXR/Source/OpenXRHMD/Internal/OpenXRHMD_RenderBridge.h
--- a/Engine/Plugins/Runtime/OpenXR/Source/OpenXRHMD/Internal/OpenXRHMD_RenderBridge.h
+++ b/Engine/Plugins/Runtime/OpenXR/Source/OpenXRHMD/Internal/OpenXRHMD_RenderBridge.h
@@ -49,7 +49,6 @@ public:
 
 	/** FRHICustomPresent */
 	virtual bool Present(IRHICommandContext& RHICmdContext, int32& InOutSyncInterval) override;
-	virtual bool Present(int32& InOutSyncInterval) override { check(false); return false; }
 
 	virtual bool Support10BitSwapchain() const { return false; }
 
diff --git a/Engine/Plugins/Runtime/nDisplay/Source/DisplayCluster/Private/Render/Presentation/DisplayClusterPresentationBase.cpp b/Engine/Plugins/Runtime/nDisplay/Source/DisplayCluster/Private/Render/Presentation/DisplayClusterPresentationBase.cpp
--- a/Engine/Plugins/Runtime/nDisplay/Source/DisplayCluster/Private/Render/Presentation/DisplayClusterPresentationBase.cpp
+++ b/Engine/Plugins/Runtime/nDisplay/Source/DisplayCluster/Private/Render/Presentation/DisplayClusterPresentationBase.cpp
@@ -41,7 +41,7 @@ void FDisplayClusterPresentationBase::OnBackBufferResize()
 {
 }
 
-bool FDisplayClusterPresentationBase::Present(int32& InOutSyncInterval)
+bool FDisplayClusterPresentationBase::Present(IRHICommandContext& RHICmdContext, int32& InOutSyncInterval)
 {
 	TRACE_CPUPROFILER_EVENT_SCOPE(nDisplay PresentationBase);
 
diff --git a/Engine/Plugins/Runtime/nDisplay/Source/DisplayCluster/Private/Render/Presentation/DisplayClusterPresentationBase.h b/Engine/Plugins/Runtime/nDisplay/Source/DisplayCluster/Private/Render/Presentation/DisplayClusterPresentationBase.h
--- a/Engine/Plugins/Runtime/nDisplay/Source/DisplayCluster/Private/Render/Presentation/DisplayClusterPresentationBase.h
+++ b/Engine/Plugins/Runtime/nDisplay/Source/DisplayCluster/Private/Render/Presentation/DisplayClusterPresentationBase.h
@@ -32,7 +32,7 @@ protected:
 	virtual bool NeedsAdvanceBackbuffer() override
 	{ return true; };
 
-	virtual bool Present(int32& InOutSyncInterval) override;
+	virtual bool Present(IRHICommandContext& RHICmdContext, int32& InOutSyncInterval) override;
 	virtual void PostPresent() override;
 	virtual void OnBackBufferResize() override;
 
diff --git a/Engine/Source/Runtime/Apple/MetalRHI/Private/MetalRayTracing.cpp b/Engine/Source/Runtime/Apple/MetalRHI/Private/MetalRayTracing.cpp
--- a/Engine/Source/Runtime/Apple/MetalRHI/Private/MetalRayTracing.cpp
+++ b/Engine/Source/Runtime/Apple/MetalRHI/Private/MetalRayTracing.cpp
@@ -696,11 +696,6 @@ void FMetalRHICommandContext::RHIBindAccelerationStructureMemory(FRHIRayTracingS
 	MetalScene->BindBuffer(InBuffer, InBufferOffset);
 }
 
-void FMetalRHICommandContext::RHIClearRayTracingBindings(FRHIRayTracingScene* Scene)
-{
-	// TODO:
-}
-
 void FMetalRHICommandContext::RHIClearShaderBindingTable(FRHIShaderBindingTable* SBT)
 {
 	// TODO:
diff --git a/Engine/Source/Runtime/Apple/MetalRHI/Public/MetalRHIContext.h b/Engine/Source/Runtime/Apple/MetalRHI/Public/MetalRHIContext.h
--- a/Engine/Source/Runtime/Apple/MetalRHI/Public/MetalRHIContext.h
+++ b/Engine/Source/Runtime/Apple/MetalRHI/Public/MetalRHIContext.h
@@ -175,7 +175,6 @@ public:
 	virtual void RHIBindAccelerationStructureMemory(FRHIRayTracingScene* Scene, FRHIBuffer* Buffer, uint32 BufferOffset) final override;
 	virtual void RHIBuildAccelerationStructures(TConstArrayView<FRayTracingGeometryBuildParams> Params, const FRHIBufferRange& ScratchBufferRange) final override;
 	virtual void RHIBuildAccelerationStructures(TConstArrayView<FRayTracingSceneBuildParams> Params) final override;
-	virtual void RHIClearRayTracingBindings(FRHIRayTracingScene* Scene) final override;
 	virtual void RHIClearShaderBindingTable(FRHIShaderBindingTable* SBT) final override;
 	virtual void RHIRayTraceDispatch(FRHIRayTracingPipelineState* RayTracingPipelineState, FRHIRayTracingShader* RayGenShader,
 		FRHIRayTracingScene* Scene,
diff --git a/Engine/Source/Runtime/D3D12RHI/Private/D3D12CommandContext.h b/Engine/Source/Runtime/D3D12RHI/Private/D3D12CommandContext.h
--- a/Engine/Source/Runtime/D3D12RHI/Private/D3D12CommandContext.h
+++ b/Engine/Source/Runtime/D3D12RHI/Private/D3D12CommandContext.h
@@ -738,8 +738,6 @@ public:
 #endif
 	virtual void RHIBuildAccelerationStructures(TConstArrayView<FRayTracingGeometryBuildParams> Params, const FRHIBufferRange& ScratchBufferRange) final override;
 	virtual void RHIBuildAccelerationStructures(TConstArrayView<FRayTracingSceneBuildParams> Params) final override;
-	virtual void RHIClearRayTracingBindings(FRHIRayTracingScene* Scene) final override;
-	virtual void RHICommitRayTracingBindings(FRHIRayTracingScene* Scene) final override;
 	virtual void RHIClearShaderBindingTable(FRHIShaderBindingTable* SBT) final override;
 	virtual void RHICommitShaderBindingTable(FRHIShaderBindingTable* SBT, FRHIBuffer* InlineBindingDataBuffer) final override;
 
@@ -1123,11 +1121,6 @@ public:
 		ContextRedirect(RHISetBindingsOnShaderBindingTable(SBT, Pipeline, NumBindings, Bindings, BindingType));
 	}
 
-	virtual void RHIClearRayTracingBindings(FRHIRayTracingScene* Scene) final override
-	{
-		ContextRedirect(RHIClearRayTracingBindings(Scene));
-	}
-
 	virtual void RHIClearShaderBindingTable(FRHIShaderBindingTable* SBT) final override
 	{
 		ContextRedirect(RHIClearShaderBindingTable(SBT));
@@ -1143,11 +1136,6 @@ public:
 		ContextRedirect(RHICommitShaderBindingTable(SBT, InlineBindingDataBuffer));
 	}
 
-	virtual void RHICommitRayTracingBindings(FRHIRayTracingScene* Scene) final override
-	{
-		ContextRedirect(RHICommitRayTracingBindings(Scene));
-	}
-
 #endif // D3D12_RHI_RAYTRACING
 
 	virtual void RHISetGPUMask(FRHIGPUMask InGPUMask) final override
diff --git a/Engine/Source/Runtime/D3D12RHI/Private/D3D12RayTracing.cpp b/Engine/Source/Runtime/D3D12RHI/Private/D3D12RayTracing.cpp
--- a/Engine/Source/Runtime/D3D12RHI/Private/D3D12RayTracing.cpp
+++ b/Engine/Source/Runtime/D3D12RHI/Private/D3D12RayTracing.cpp
@@ -2794,12 +2794,6 @@ FD3D12RayTracingPipelineState::FD3D12RayTracingPipelineState(FD3D12Device* Devic
 		return ShaderCacheEntry;
 	};
 
-
-	// If no custom hit groups were provided, then disable SBT indexing and force default shader on all primitives
-	PRAGMA_DISABLE_DEPRECATION_WARNINGS
-	bAllowHitGroupIndexing = Initializer.GetHitGroupTable().Num() ? Initializer.bAllowHitGroupIndexing : false;
-	PRAGMA_ENABLE_DEPRECATION_WARNINGS
-
 	// Add ray generation shaders
 
 	TArray<FD3D12RayTracingPipelineCache::FEntry*> RayGenShaderEntries;
@@ -4299,9 +4293,6 @@ FD3D12RayTracingScene::FD3D12RayTracingScene(FD3D12Adapter* Adapter, FRayTracing
 {
 	INC_DWORD_STAT(STAT_D3D12RayTracingAllocatedTLAS);
 
-	PRAGMA_DISABLE_DEPRECATION_WARNINGS
-	checkf(Initializer.NumMissShaderSlots >= 1, TEXT("Need at least 1 miss shader slot."));
-	PRAGMA_ENABLE_DEPRECATION_WARNINGS
 	checkf(Initializer.Lifetime == RTSL_SingleFrame, TEXT("Only single-frame ray tracing scenes are currently implemented."));
 
 	// Get maximum buffer sizes for all GPUs in the system
@@ -4536,7 +4527,7 @@ void PrepareAccelerationStructureBuild(
 		uint32 InstanceBufferStride = GRHIRayTracingInstanceDescriptorSize;
 		PRAGMA_DISABLE_DEPRECATION_WARNINGS
 		// TODO: Validation related to SBT needs to be done somewhere else since SBT is not known when in BuildAccelerationStructure
-		uint32 TotalHitGroupSlots = Scene.Initializer.NumTotalSegments * Scene.Initializer.ShaderSlotsPerGeometrySegment;
+			uint32 TotalHitGroupSlots = Scene.Initializer.NumTotalSegments;
 		PRAGMA_ENABLE_DEPRECATION_WARNINGS
 		FRayTracingValidateSceneBuildParamsCS::Dispatch(RHICmdList,
 			TotalHitGroupSlots, NumInstances,
@@ -4597,47 +4588,6 @@ void FD3D12RayTracingScene::UpdateResidency(FD3D12CommandContext& CommandContext
 #endif // ENABLE_RESIDENCY_MANAGEMENT
 }
 
-FD3D12RayTracingShaderBindingTable* FD3D12RayTracingScene::FindExistingShaderTable(const FD3D12RayTracingPipelineState* Pipeline) const
-{
-	if (const TRefCountPtr<FD3D12RayTracingShaderBindingTable>* FoundShaderTable = ShaderTables.Find(Pipeline))
-	{
-		return *FoundShaderTable;
-	}
-	else
-	{
-		return nullptr;
-	}
-}
-
-FRHIShaderBindingTable* FD3D12RayTracingScene::FindOrCreateShaderBindingTable(const FRHIRayTracingPipelineState* InPipeline)
-{
-	UE::TScopeLock Lock(Mutex);
-
-	const FD3D12RayTracingPipelineState* Pipeline = FD3D12DynamicRHI::ResourceCast(InPipeline);
-
-	if (FD3D12RayTracingShaderBindingTable* FoundShaderTable = FindExistingShaderTable(Pipeline))
-	{
-		return FoundShaderTable;
-	}
-
-	PRAGMA_DISABLE_DEPRECATION_WARNINGS
-	FRayTracingShaderBindingTableInitializer SBTInitializer;
-	SBTInitializer.NumGeometrySegments = Initializer.NumTotalSegments;
-	SBTInitializer.NumShaderSlotsPerGeometrySegment = Initializer.ShaderSlotsPerGeometrySegment;
-	SBTInitializer.NumCallableShaderSlots = Initializer.NumCallableShaderSlots;
-	SBTInitializer.NumMissShaderSlots = Initializer.NumMissShaderSlots;
-	SBTInitializer.HitGroupIndexingMode = Pipeline->bAllowHitGroupIndexing ? ERayTracingHitGroupIndexingMode::Allow : ERayTracingHitGroupIndexingMode::Disallow;
-	SBTInitializer.ShaderBindingMode = ERayTracingShaderBindingMode::RTPSO;
-	PRAGMA_ENABLE_DEPRECATION_WARNINGS
-
-	// Create new table (use FRHICommandListExecutor::GetImmediateCommandList() directly for now this is deprecated code)
-	FD3D12RayTracingShaderBindingTable* CreatedShaderTable = new FD3D12RayTracingShaderBindingTable(FRHICommandListExecutor::GetImmediateCommandList(), GetParentAdapter(), MoveTemp(SBTInitializer));
-
-	ShaderTables.Add(Pipeline, CreatedShaderTable);
-
-	return CreatedShaderTable;
-}
-
 void FD3D12CommandContext::BuildAccelerationStructuresInternal(TConstArrayView<D3D12_BUILD_RAYTRACING_ACCELERATION_STRUCTURE_DESC> BuildDescs)
 {
 	for (const D3D12_BUILD_RAYTRACING_ACCELERATION_STRUCTURE_DESC& Desc : BuildDescs)
@@ -4873,33 +4823,6 @@ void FD3D12CommandContext::RHIBindAccelerationStructureMemory(FRHIRayTracingScen
 	Scene->BindBuffer(InBuffer, InBufferOffset);
 }
 
-void FD3D12CommandContext::RHICommitRayTracingBindings(FRHIRayTracingScene* InScene)
-{
-	FD3D12RayTracingScene* Scene = FD3D12DynamicRHI::ResourceCast(InScene);
-	check(Scene);
-
-	for (auto Item : Scene->ShaderTables)
-	{
-		FD3D12RayTracingShaderBindingTable* ShaderTable = Item.Value;
-		FD3D12RayTracingShaderBindingTableInternal* ShaderTableForDevice = ShaderTable->GetTableForDevice(GetParentDevice());
-		if (ShaderTableForDevice->bIsDirty)
-		{
-			ShaderTableForDevice->Commit(*this, nullptr);
-		}
-	}
-}
-
-void FD3D12CommandContext::RHIClearRayTracingBindings(FRHIRayTracingScene* InScene)
-{
-	FD3D12RayTracingScene* Scene = FD3D12DynamicRHI::ResourceCast(InScene);
-	check(Scene);
-
-	for (auto Table : Scene->ShaderTables)
-	{
-		Table.Value->ReleaseForDevice(GetParentDevice());
-	}
-}
-
 void FD3D12CommandContext::RHICommitShaderBindingTable(FRHIShaderBindingTable* InSBT, FRHIBuffer* InlineBindingDataBuffer)
 {
 	FD3D12RayTracingShaderBindingTable* SBT = FD3D12DynamicRHI::ResourceCast(InSBT);
diff --git a/Engine/Source/Runtime/D3D12RHI/Private/D3D12RayTracing.h b/Engine/Source/Runtime/D3D12RHI/Private/D3D12RayTracing.h
--- a/Engine/Source/Runtime/D3D12RHI/Private/D3D12RayTracing.h
+++ b/Engine/Source/Runtime/D3D12RHI/Private/D3D12RayTracing.h
@@ -103,9 +103,6 @@ public:
 	// State objects with raygen shaders grouped by occupancy
 	TArray<TRefCountPtr<ID3D12StateObject>> SpecializedStateObjects;
 
-	UE_DEPRECATED(5.5, "bAllowHitGroupIndexing is now stored in the ShaderBindingTable.")
-	bool bAllowHitGroupIndexing = true;
-
 	uint32 MaxLocalRootSignatureSize = 0;
 	uint32 MaxHitGroupViewDescriptors = 0;
 
@@ -271,12 +268,6 @@ public:
 
 	void UpdateResidency(FD3D12CommandContext& CommandContext) const;
 
-	// #dxr_todo UE-68230: shader tables should be explicitly registered and unregistered with the scene
-	FRHIShaderBindingTable* FindOrCreateShaderBindingTable(const FRHIRayTracingPipelineState* Pipeline) override;
-	FD3D12RayTracingShaderBindingTable* FindExistingShaderTable(const FD3D12RayTracingPipelineState* Pipeline) const;
-
-	TMap<const FD3D12RayTracingPipelineState*, TRefCountPtr<FD3D12RayTracingShaderBindingTable>> ShaderTables;
-
 	bool bBuilt = false;
 
 private:
diff --git a/Engine/Source/Runtime/D3D12RHI/Public/D3D12RHIBridge.h b/Engine/Source/Runtime/D3D12RHI/Public/D3D12RHIBridge.h
deleted file mode 100644
diff --git a/Engine/Source/Runtime/RHI/Private/RHICommandList.cpp b/Engine/Source/Runtime/RHI/Private/RHICommandList.cpp
--- a/Engine/Source/Runtime/RHI/Private/RHICommandList.cpp
+++ b/Engine/Source/Runtime/RHI/Private/RHICommandList.cpp
@@ -2326,12 +2326,6 @@ void FRHICommandListExecutor::CleanupGraphEvents()
 	RHIThreadPipe.CleanupGraphEvents();
 }
 
-UE_DEPRECATED(5.5, "GenerateMips on RHI command lists is deprecated and no longer functions. Use the FGenerateMips helper class from the RenderCore module to generate mips on textures.")
-void FRHICommandList::GenerateMips(FRHITexture*)
-{
-	UE_LOG(LogRHI, Fatal, TEXT("GenerateMips on RHI command lists is deprecated and no longer functions. Use the FGenerateMips helper class from the RenderCore module to generate mips on textures."));
-}
-
 void FRHICommandListBase::InvalidBufferFatalError(const FRHIBufferCreateDesc& CreateDesc)
 {
 	UE_LOG(LogRHI, Fatal, TEXT("Attempt to create zero-sized buffer '%s', owner '%s', usage 0x%x, stride %u."),
diff --git a/Engine/Source/Runtime/RHI/Private/RHIResources.cpp b/Engine/Source/Runtime/RHI/Private/RHIResources.cpp
--- a/Engine/Source/Runtime/RHI/Private/RHIResources.cpp
+++ b/Engine/Source/Runtime/RHI/Private/RHIResources.cpp
@@ -117,11 +117,6 @@ void FRHIResource::GatherResourcesToDelete(TArray<FRHIResource*>& OutResources,
 	});
 }
 
-bool FRHIResource::Bypass()
-{
-	return GRHICommandList.Bypass();
-}
-
 FRHITexture::FRHITexture(const FRHITextureCreateDesc& InDesc)
 	: FRHIViewableResource(RRT_Texture, InDesc.InitialState, InDesc.GetSubresourceCount())
 #if ENABLE_RHI_VALIDATION
diff --git a/Engine/Source/Runtime/RHI/Public/DynamicRHI.h b/Engine/Source/Runtime/RHI/Public/DynamicRHI.h
--- a/Engine/Source/Runtime/RHI/Public/DynamicRHI.h
+++ b/Engine/Source/Runtime/RHI/Public/DynamicRHI.h
@@ -525,11 +525,6 @@ public:
 	* @returns a reference to a 2D texture resource
 	*/
 	// FlushType: Thread safe
-	UE_DEPRECATED(5.4, "Use the RHIAsyncCreateTexture2D function that takes a DebugName.")
-	FTextureRHIRef RHIAsyncCreateTexture2D(uint32 SizeX, uint32 SizeY, uint8 Format, uint32 NumMips, ETextureCreateFlags Flags, ERHIAccess InResourceState, void** InitialMipData, uint32 NumInitialMips, FGraphEventRef& OutCompletionEvent)
-	{
-		return RHIAsyncCreateTexture2D(SizeX, SizeY, Format, NumMips, Flags, InResourceState, InitialMipData, NumInitialMips, TEXT("RHIAsyncCreateTexture2D"), OutCompletionEvent);
-	}
 	virtual FTextureRHIRef RHIAsyncCreateTexture2D(uint32 SizeX, uint32 SizeY, uint8 Format, uint32 NumMips, ETextureCreateFlags Flags, ERHIAccess InResourceState, void** InitialMipData, uint32 NumInitialMips, const TCHAR* DebugName, FGraphEventRef& OutCompletionEvent) = 0;
 
 		
@@ -1248,130 +1243,6 @@ FORCEINLINE FRHICalcTextureSizeResult RHICalcTexturePlatformSize(FRHITextureDesc
 	return GDynamicRHI->RHICalcTexturePlatformSize(Desc, FirstMipIndex);
 }
 
-PRAGMA_DISABLE_DEPRECATION_WARNINGS
-
-UE_DEPRECATED(5.5, "The separate RHICalcTexture... functions for each texture type are deprecated. Use RHICalcTexturePlatformSize instead.")
-FORCEINLINE uint64 RHICalcTexture2DPlatformSize(uint32 SizeX, uint32 SizeY, uint8 Format, uint32 NumMips, uint32 NumSamples, ETextureCreateFlags Flags, const FRHIResourceCreateInfo& CreateInfo, uint32& OutAlign)
-{
-	const uint16 Depth = 1;
-	const uint16 ArraySize = 1;
-
-	FRHITextureDesc Desc(
-		ETextureDimension::Texture2D,
-		Flags,
-		(EPixelFormat)Format,
-		CreateInfo.ClearValueBinding,
-		{ (int32)SizeX, (int32)SizeY },
-		Depth,
-		ArraySize,
-		(uint8)NumMips,
-		(uint8)NumSamples,
-		CreateInfo.ExtData
-	);
-
-	FRHICalcTextureSizeResult Result = RHICalcTexturePlatformSize(Desc, 0);
-	OutAlign = Result.Align;
-	return Result.Size;
-}
-
-UE_DEPRECATED(5.5, "The separate RHICalcTexture... functions for each texture type are deprecated. Use RHICalcTexturePlatformSize instead.")
-FORCEINLINE uint64 RHICalcTexture2DArrayPlatformSize(uint32 SizeX, uint32 SizeY, uint32 ArraySize, uint8 Format, uint32 NumMips, uint32 NumSamples, ETextureCreateFlags Flags, const FRHIResourceCreateInfo& CreateInfo, uint32& OutAlign)
-{
-	const uint16 Depth = 1;
-
-	FRHITextureDesc Desc(
-		ETextureDimension::Texture2DArray,
-		Flags,
-		(EPixelFormat)Format,
-		CreateInfo.ClearValueBinding,
-		{ (int32)SizeX, (int32)SizeY },
-		Depth,
-		(uint16)ArraySize,
-		(uint8)NumMips,
-		(uint8)NumSamples,
-		CreateInfo.ExtData
-	);
-
-	FRHICalcTextureSizeResult Result = RHICalcTexturePlatformSize(Desc, 0);
-	OutAlign = Result.Align;
-	return Result.Size;
-}
-
-UE_DEPRECATED(5.5, "The separate RHICalcTexture... functions for each texture type are deprecated. Use RHICalcTexturePlatformSize instead.")
-FORCEINLINE uint64 RHICalcVMTexture2DPlatformSize(uint32 Mip0Width, uint32 Mip0Height, uint8 Format, uint32 NumMips, uint32 FirstMipIdx, uint32 NumSamples, ETextureCreateFlags Flags, uint32& OutAlign)
-{
-	const uint16 Depth     = 1;
-	const uint16 ArraySize = 1;
-	const uint32 ExtData   = 0;
-
-	FRHITextureDesc Desc(
-		ETextureDimension::Texture2D,
-		Flags,
-		(EPixelFormat)Format,
-		FClearValueBinding::None,
-		{ (int32)Mip0Width, (int32)Mip0Height },
-		Depth,
-		ArraySize,
-		(uint8)NumMips,
-		(uint8)NumSamples,
-		ExtData
-	);
-
-	FRHICalcTextureSizeResult Result = RHICalcTexturePlatformSize(Desc, FirstMipIdx);
-	OutAlign = Result.Align;
-	return Result.Size;
-}
-
-UE_DEPRECATED(5.5, "The separate RHICalcTexture... functions for each texture type are deprecated. Use RHICalcTexturePlatformSize instead.")
-FORCEINLINE uint64 RHICalcTexture3DPlatformSize(uint32 SizeX, uint32 SizeY, uint32 SizeZ, uint8 Format, uint32 NumMips, ETextureCreateFlags Flags, const FRHIResourceCreateInfo& CreateInfo, uint32& OutAlign)
-{
-	const uint16 ArraySize  = 1;
-	const uint8 NumSamples = 1;
-
-	FRHITextureDesc Desc(
-		ETextureDimension::Texture3D,
-		Flags,
-		(EPixelFormat)Format,
-		CreateInfo.ClearValueBinding,
-		{ (int32)SizeX, (int32)SizeY },
-		(uint16)SizeZ,
-		ArraySize,
-		(uint8)NumMips,
-		NumSamples,
-		CreateInfo.ExtData
-	);
-
-	FRHICalcTextureSizeResult Result = RHICalcTexturePlatformSize(Desc, 0);
-	OutAlign = Result.Align;
-	return Result.Size;
-}
-
-UE_DEPRECATED(5.5, "The separate RHICalcTexture... functions for each texture type are deprecated. Use RHICalcTexturePlatformSize instead.")
-FORCEINLINE uint64 RHICalcTextureCubePlatformSize(uint32 Size, uint8 Format, uint32 NumMips, ETextureCreateFlags Flags, const FRHIResourceCreateInfo& CreateInfo, uint32& OutAlign)
-{
-	const uint16 Depth      = 1;
-	const uint16 ArraySize  = 1;
-	const uint8 NumSamples = 1;
-
-	FRHITextureDesc Desc(
-		ETextureDimension::TextureCube,
-		Flags,
-		(EPixelFormat)Format,
-		CreateInfo.ClearValueBinding,
-		{ (int32)Size, (int32)Size },
-		Depth,
-		ArraySize,
-		(uint8)NumMips,
-		NumSamples,
-		CreateInfo.ExtData
-	);
-
-	FRHICalcTextureSizeResult Result = RHICalcTexturePlatformSize(Desc, 0);
-	OutAlign = Result.Align;
-	return Result.Size;
-}
-PRAGMA_ENABLE_DEPRECATION_WARNINGS
-
 FORCEINLINE uint64 RHIGetMinimumAlignmentForBufferBackedSRV(EPixelFormat Format)
 {
 	return GDynamicRHI->RHIGetMinimumAlignmentForBufferBackedSRV(Format);
@@ -1478,12 +1349,6 @@ FORCEINLINE void RHIGetSupportedResolution(uint32& Width, uint32& Height)
 	GDynamicRHI->RHIGetSupportedResolution(Width, Height);
 }
 
-UE_DEPRECATED(5.5, "RHIRequiresComputeGenerateMips is deprecated and always returns true. Use FGenerateMips helper class from the RenderCore module to generate mips on textures.")
-FORCEINLINE bool RHIRequiresComputeGenerateMips()
-{
-	return true;
-}
-
 FORCEINLINE class IRHICommandContext* RHIGetDefaultContext()
 {
 	return GDynamicRHI->RHIGetDefaultContext();
@@ -1528,24 +1393,8 @@ FORCEINLINE FRayTracingAccelerationStructureSize RHICalcRayTracingSceneSize(cons
 	return GDynamicRHI->RHICalcRayTracingSceneSize(Initializer);
 }
 
-UE_DEPRECATED(5.5, "Provide FRayTracingSceneInitializer instead.")
-FORCEINLINE FRayTracingAccelerationStructureSize RHICalcRayTracingSceneSize(uint32 MaxInstances, ERayTracingAccelerationStructureFlags Flags)
-{
-	FRayTracingSceneInitializer Initializer;
-	Initializer.MaxNumInstances = MaxInstances;
-	Initializer.BuildFlags = Flags;
-
-	return GDynamicRHI->RHICalcRayTracingSceneSize(Initializer);
-}
-
 FORCEINLINE FRayTracingSceneRHIRef RHICreateRayTracingScene(FRayTracingSceneInitializer Initializer)
 {
-	PRAGMA_DISABLE_DEPRECATION_WARNINGS
-	if (Initializer.NumNativeInstancesPerLayer.Num() >= 1)
-	{
-		Initializer.MaxNumInstances = Initializer.NumNativeInstancesPerLayer[0];
-	}
-	PRAGMA_ENABLE_DEPRECATION_WARNINGS
 	return GDynamicRHI->RHICreateRayTracingScene(MoveTemp(Initializer));
 }
 
diff --git a/Engine/Source/Runtime/RHI/Public/RHI.h b/Engine/Source/Runtime/RHI/Public/RHI.h
--- a/Engine/Source/Runtime/RHI/Public/RHI.h
+++ b/Engine/Source/Runtime/RHI/Public/RHI.h
@@ -279,8 +279,6 @@ struct FRasterizerStateInitializerRHI
 	float SlopeScaleDepthBias = 0.0f;
 	ERasterizerDepthClipMode DepthClipMode = ERasterizerDepthClipMode::DepthClip;
 	bool bAllowMSAA = false;
-	UE_DEPRECATED(5.4, "bEnableLineAA is unsupported")
-	bool bEnableLineAA = false;
 
 	PRAGMA_DISABLE_DEPRECATION_WARNINGS
 	FRasterizerStateInitializerRHI() = default;
@@ -306,25 +304,6 @@ struct FRasterizerStateInitializerRHI
 	{
 	}
 
-	UE_DEPRECATED(5.4, "bEnableLineAA is unsupported")
-	FRasterizerStateInitializerRHI(ERasterizerFillMode InFillMode, ERasterizerCullMode InCullMode, bool bInAllowMSAA, bool bInEnableLineAA)
-		: FillMode(InFillMode)
-		, CullMode(InCullMode)
-		, bAllowMSAA(bInAllowMSAA)
-	{
-	}
-
-	UE_DEPRECATED(5.4, "bEnableLineAA is unsupported")
-	FRasterizerStateInitializerRHI(ERasterizerFillMode InFillMode, ERasterizerCullMode InCullMode, float InDepthBias, float InSlopeScaleDepthBias, ERasterizerDepthClipMode InDepthClipMode, bool bInAllowMSAA, bool bInEnableLineAA)
-		: FillMode(InFillMode)
-		, CullMode(InCullMode)
-		, DepthBias(InDepthBias)
-		, SlopeScaleDepthBias(InSlopeScaleDepthBias)
-		, DepthClipMode(InDepthClipMode)
-		, bAllowMSAA(bInAllowMSAA)
-	{
-	}
-
 	friend FArchive& operator<<(FArchive& Ar, FRasterizerStateInitializerRHI& RasterizerStateInitializer)
 	{
 		Ar << RasterizerStateInitializer.FillMode;
diff --git a/Engine/Source/Runtime/RHI/Public/RHICommandList.h b/Engine/Source/Runtime/RHI/Public/RHICommandList.h
--- a/Engine/Source/Runtime/RHI/Public/RHICommandList.h
+++ b/Engine/Source/Runtime/RHI/Public/RHICommandList.h
@@ -275,21 +275,6 @@ struct FRayTracingLocalShaderBindings
 	uint16 LooseParameterDataSize = 0;
 	FRHIUniformBuffer** UniformBuffers = nullptr;
 	uint8* LooseParameterData = nullptr;
-	
-	UE_DEPRECATED(5.5, "Provide Geometry and RecordIndex instead.")
-	uint32 InstanceIndex = 0;
-
-	UE_DEPRECATED(5.5, "Provide Geometry and RecordIndex instead.")
-	uint32 ShaderSlot = 0;
-
-	PRAGMA_DISABLE_DEPRECATION_WARNINGS
-	FRayTracingLocalShaderBindings() = default;
-	FRayTracingLocalShaderBindings(const FRayTracingLocalShaderBindings&) = default;
-	FRayTracingLocalShaderBindings& operator=(const FRayTracingLocalShaderBindings&) = default;
-	FRayTracingLocalShaderBindings(FRayTracingLocalShaderBindings&&) = default;
-	FRayTracingLocalShaderBindings& operator=(FRayTracingLocalShaderBindings&&) = default;
-	~FRayTracingLocalShaderBindings() = default;
-	PRAGMA_ENABLE_DEPRECATION_WARNINGS
 };
 
 enum class ERayTracingBindingType : uint8
@@ -530,9 +515,6 @@ public:
 	//
 	RHI_API void FinishRecording();
 
-	UE_DEPRECATED(5.5, "SetCurrentStat is deprecated and there is no replacement. Consider marking up rendering code with RDG event scopes or RHI breadcrumbs.")
-	inline void SetCurrentStat(TStatId Stat) {}
-
 	FORCEINLINE_DEBUGGABLE void* Alloc(int64 AllocSize, int64 Alignment)
 	{
 		return MemManager.Alloc(AllocSize, Alignment);
@@ -1171,12 +1153,6 @@ PRAGMA_ENABLE_DEPRECATION_WARNINGS
 		return GDynamicRHI->RHICreateShaderBindingTable(*this, Initializer);
 	}
 	
-	UE_DEPRECATED(5.5, "Use the global scope RHICalcRayTracingGeometrySize function instead.")
-	FORCEINLINE FRayTracingAccelerationStructureSize CalcRayTracingGeometrySize(const FRayTracingGeometryInitializer& Initializer)
-	{
-		return GDynamicRHI->RHICalcRayTracingGeometrySize(Initializer);
-	}
-
 	FORCEINLINE void ReplaceResources(TArray<FRHIResourceReplaceInfo>&& ReplaceInfos)
 	{
 		if (ReplaceInfos.Num() == 0)
@@ -2660,17 +2636,6 @@ FRHICOMMAND_UNNAMED(FRHICommandBuildSceneAccelerationStructures)
 	RHI_EXECUTE_API void Execute(FRHICommandListBase& CmdList);
 };
 
-FRHICOMMAND_MACRO(FRHICommandCommitRayTracingBindings)
-{
-	FRHIRayTracingScene* Scene;
-
-	explicit FRHICommandCommitRayTracingBindings(FRHIRayTracingScene* InScene)
-		: Scene(InScene)
-	{}
-
-	RHI_EXECUTE_API void Execute(FRHICommandListBase& CmdList);
-};
-
 FRHICOMMAND_MACRO(FRHICommandCommitShaderBindingTable)
 {
 	FRHIShaderBindingTable* SBT;
@@ -2683,17 +2648,6 @@ FRHICOMMAND_MACRO(FRHICommandCommitShaderBindingTable)
 	RHI_EXECUTE_API void Execute(FRHICommandListBase& CmdList);
 };
 
-FRHICOMMAND_MACRO(FRHICommandClearRayTracingBindings)
-{
-	FRHIRayTracingScene* Scene;
-
-	explicit FRHICommandClearRayTracingBindings(FRHIRayTracingScene* InScene)
-		: Scene(InScene)
-	{}
-
-	RHI_EXECUTE_API void Execute(FRHICommandListBase & CmdList);
-};
-
 FRHICOMMAND_MACRO(FRHICommandClearShaderBindingTable)
 {
 	FRHIShaderBindingTable* SBT;
@@ -3508,16 +3462,6 @@ public:
 	}
 #endif // WITH_RHI_BREADCRUMBS
 
-	UE_DEPRECATED(5.5, "RHIPushEvent is deprecated. All events and markers now use the RHI breadcrumb system. Use RDG_EVENT_SCOPE or SCOPED_DRAW_EVENT macros to mark up rendering code, rather than calling this function directly.")
-	FORCEINLINE_DEBUGGABLE void PushEvent(const TCHAR* Name, FColor Color)
-	{
-	}
-
-	UE_DEPRECATED(5.5, "RHIPopEvent is deprecated. All events and markers now use the RHI breadcrumb system. Use RDG_EVENT_SCOPE or SCOPED_DRAW_EVENT macros to mark up rendering code, rather than calling this function directly.")
-	FORCEINLINE_DEBUGGABLE void PopEvent()
-	{
-	}
-
 	//UE_DEPRECATED(5.1, "SubmitCommandsHint is deprecated, and has no effect if called on a non-immediate RHI command list. Consider calling ImmediateFlush(EImmediateFlushType::DispatchToRHIThread) on the immediate command list instead.")
 	inline void SubmitCommandsHint();
 
@@ -3809,19 +3753,6 @@ public:
 #endif // WITH_MGPU
 	}
 	
-	UE_DEPRECATED(5.5, "Provide FRHIShaderBindingTable instead of FRayTracingScene.")
-	FORCEINLINE_DEBUGGABLE void RayTraceDispatch(FRayTracingPipelineState* Pipeline, FRHIRayTracingShader* RayGenShader, FRHIRayTracingScene* Scene, const FRayTracingShaderBindings& GlobalResourceBindings, uint32 Width, uint32 Height)
-	{
-		if (Bypass())
-		{
-			GetComputeContext().RHIRayTraceDispatch(GetRHIRayTracingPipelineState(Pipeline), RayGenShader, Scene, GlobalResourceBindings, Width, Height);
-		}
-		else
-		{
-			ALLOC_COMMAND(FRHICommandRayTraceDispatch)(Pipeline, RayGenShader, Scene, GlobalResourceBindings, Width, Height);
-		}
-	}
-
 	FORCEINLINE_DEBUGGABLE void RayTraceDispatch(FRayTracingPipelineState* Pipeline, FRHIRayTracingShader* RayGenShader, FRHIShaderBindingTable* SBT, const FRayTracingShaderBindings& GlobalResourceBindings, uint32 Width, uint32 Height)
 	{
 		check(SBT != nullptr);
@@ -3841,23 +3772,6 @@ public:
 	*/
 	RHI_API void RayTraceDispatch(FRayTracingPipelineState* Pipeline, FRHIRayTracingShader* RayGenShader, FRHIShaderBindingTable* SBT, FRHIBatchedShaderParameters& GlobalResourceBindings, uint32 Width, uint32 Height);
 
-	/**
-	 * Trace rays using dimensions from a GPU buffer containing uint[3], interpreted as number of rays in X, Y and Z dimensions.
-	 * ArgumentBuffer must be in IndirectArgs|SRVCompute state.
-	 */
-	UE_DEPRECATED(5.5, "Provide FRHIShaderBindingTable instead of FRayTracingScene.")
-	FORCEINLINE_DEBUGGABLE void RayTraceDispatchIndirect(FRayTracingPipelineState* Pipeline, FRHIRayTracingShader* RayGenShader, FRHIRayTracingScene* Scene, const FRayTracingShaderBindings& GlobalResourceBindings, FRHIBuffer* ArgumentBuffer, uint32 ArgumentOffset)
-	{
-		if (Bypass())
-		{
-			GetComputeContext().RHIRayTraceDispatchIndirect(GetRHIRayTracingPipelineState(Pipeline), RayGenShader, Scene, GlobalResourceBindings, ArgumentBuffer, ArgumentOffset);
-		}
-		else
-		{
-			ALLOC_COMMAND(FRHICommandRayTraceDispatch)(Pipeline, RayGenShader, Scene, GlobalResourceBindings, ArgumentBuffer, ArgumentOffset);
-		}
-	}
-
 	FORCEINLINE_DEBUGGABLE void RayTraceDispatchIndirect(FRayTracingPipelineState* Pipeline, FRHIRayTracingShader* RayGenShader, FRHIShaderBindingTable* SBT, const FRayTracingShaderBindings& GlobalResourceBindings, FRHIBuffer* ArgumentBuffer, uint32 ArgumentOffset)
 	{
 		check(SBT != nullptr);
@@ -4301,18 +4215,6 @@ public:
 		ALLOC_COMMAND(FRHICommandDrawPrimitiveIndirect)(ArgumentBuffer, ArgumentOffset);
 	}
 
-	UE_DEPRECATED(5.4, "Use DrawIndexedPrimitiveIndirect.")
-	FORCEINLINE_DEBUGGABLE void DrawIndexedIndirect(FRHIBuffer* IndexBufferRHI, FRHIBuffer* ArgumentsBufferRHI, uint32 DrawArgumentsIndex, uint32 NumInstances)
-	{
-		//check(IsOutsideRenderPass());
-		if (Bypass())
-		{
-			GetContext().RHIDrawIndexedIndirect(IndexBufferRHI, ArgumentsBufferRHI, DrawArgumentsIndex, NumInstances);
-			return;
-		}
-		ALLOC_COMMAND(FRHICommandDrawIndexedIndirect)(IndexBufferRHI, ArgumentsBufferRHI, DrawArgumentsIndex, NumInstances);
-	}
-
 	FORCEINLINE_DEBUGGABLE void DrawIndexedPrimitiveIndirect(FRHIBuffer* IndexBuffer, FRHIBuffer* ArgumentsBuffer, uint32 ArgumentOffset)
 	{
 		//check(IsOutsideRenderPass());
@@ -4432,11 +4334,6 @@ public:
 		ALLOC_COMMAND(FRHICommandCalibrateTimers)(CalibrationQuery);
 	}
 
-	UE_DEPRECATED(5.5, "PollOcclusionQueries is no longer necessary. It has been deprecated and there is no replacement. Remove any remaining calls to PollOcclusionQueries.")
-	FORCEINLINE_DEBUGGABLE void PollOcclusionQueries()
-	{
-	}
-	
 	FORCEINLINE_DEBUGGABLE void BeginRenderPass(const FRHIRenderPassInfo& InInfo, const TCHAR* Name)
 	{
 		check(!IsInsideRenderPass());
@@ -4529,21 +4426,6 @@ public:
 		IncrementSubpass();
 	}
 
-	UE_DEPRECATED(5.5, "RHIInvalidateCachedState is not implemented by any platform RHI.")
-	FORCEINLINE_DEBUGGABLE void RHIInvalidateCachedState()
-	{}
-
-	UE_DEPRECATED(5.5, "DiscardRenderTargets is deprecated and will be removed in future engine versions.")
-	FORCEINLINE void DiscardRenderTargets(bool Depth, bool Stencil, uint32 ColorBitMask)
-	{
-		if (Bypass())
-		{
-			GetContext().RHIDiscardRenderTargets(Depth, Stencil, ColorBitMask);
-			return;
-		}
-		ALLOC_COMMAND(FRHICommandDiscardRenderTargets)(Depth, Stencil, ColorBitMask);
-	}
-	
 	FORCEINLINE_DEBUGGABLE void CopyBufferRegion(FRHIBuffer* DestBuffer, uint64 DstOffset, FRHIBuffer* SourceBuffer, uint64 SrcOffset, uint64 NumBytes)
 	{
 		// No copy/DMA operation inside render passes
@@ -4563,28 +4445,8 @@ public:
 		}
 	}
 
-	UE_DEPRECATED(5.5, "GenerateMips on RHI command lists is deprecated and no longer functions. Use the FGenerateMips helper class from the RenderCore module to generate mips on textures.")
-	RHI_API void GenerateMips(FRHITexture*);
-
 	// Ray tracing API
 
-	UE_DEPRECATED(5.5, "Use FRHIShaderBindingTable instead.")
-	FORCEINLINE_DEBUGGABLE void CommitRayTracingBindings(FRHIRayTracingScene* Scene)
-	{
-		if (Bypass())
-		{
-			GetContext().RHICommitRayTracingBindings(Scene);
-		}
-		else
-		{
-			ALLOC_COMMAND(FRHICommandCommitRayTracingBindings)(Scene);
-
-			// This RHI command modifies members of the FRHIRayTracingScene inside platform RHI implementations.
-			// It therefore needs the RHI lock fence to prevent races on those members.
-			RHIThreadFence(true);
-		}
-	}
-		
 	void CommitShaderBindingTable(FRHIShaderBindingTable* SBT)
 	{
 		checkf(!EnumHasAnyFlags(SBT->GetInitializer().ShaderBindingMode, ERayTracingShaderBindingMode::Inline), TEXT("Use CommitShaderBindingTable function which also provides the InlineBindingDataBuffer when SBT has inline binding mode set"));		
@@ -4607,23 +4469,6 @@ public:
 		}
 	}
 
-	UE_DEPRECATED(5.5, "Use FRHIShaderBindingTable instead.")
-	FORCEINLINE_DEBUGGABLE void ClearRayTracingBindings(FRHIRayTracingScene* Scene)
-	{
-		if (Bypass())
-		{
-			GetContext().RHIClearRayTracingBindings(Scene);
-		}
-		else
-		{
-			ALLOC_COMMAND(FRHICommandClearRayTracingBindings)(Scene);
-
-			// This RHI command modifies members of the FRHIRayTracingScene inside platform RHI implementations.
-			// It therefore needs the RHI lock fence to prevent races on those members.
-			RHIThreadFence(true);
-		}
-	}
-
 	FORCEINLINE_DEBUGGABLE void ClearShaderBindingTable(FRHIShaderBindingTable* SBT)
 	{
 		if (Bypass())
@@ -4640,58 +4485,6 @@ public:
 		}
 	}
 	
-	UE_DEPRECATED(5.5, "Use FRHIShaderBindingTable instead.")
-	FORCEINLINE_DEBUGGABLE void SetRayTracingBindings(
-		FRHIRayTracingScene* Scene, FRayTracingPipelineState* Pipeline,
-		uint32 NumBindings, const FRayTracingLocalShaderBindings* InBindings,
-		ERayTracingBindingType BindingType,
-		bool bCopyDataToInlineStorage = true)
-	{
-		// Upgrade bindings to new code path (duplicate bindings array to command list memory and patch to use Geometry + RecordIndex instead of the deprecated InstanceIndex + ShaderSlot)
-		FRayTracingLocalShaderBindings* Bindings = nullptr;
-		{
-			if (NumBindings)
-			{
-				uint32 Size = sizeof(FRayTracingLocalShaderBindings) * NumBindings;
-				Bindings = (FRayTracingLocalShaderBindings*)Alloc(Size, alignof(FRayTracingLocalShaderBindings));
-				FMemory::Memcpy(Bindings, InBindings, Size);
-			}
-
-			for (uint32 i = 0; i < NumBindings; ++i)
-			{
-				PRAGMA_DISABLE_DEPRECATION_WARNINGS
-				Bindings[i].Geometry = Scene->GetInitializer().PerInstanceGeometries[Bindings[i].InstanceIndex];
-				Bindings[i].RecordIndex = (Scene->GetInitializer().SegmentPrefixSum[Bindings[i].InstanceIndex] + Bindings[i].SegmentIndex) * Scene->GetInitializer().ShaderSlotsPerGeometrySegment + Bindings[i].ShaderSlot;
-				Bindings[i].InstanceIndex = INDEX_NONE;
-				PRAGMA_ENABLE_DEPRECATION_WARNINGS
-
-				if (Bindings[i].NumUniformBuffers)
-				{
-					Bindings[i].UniformBuffers = (FRHIUniformBuffer**)Alloc(sizeof(FRHIUniformBuffer*) * Bindings[i].NumUniformBuffers, alignof(FRHIUniformBuffer*));
-					for (uint32 Index = 0; Index < Bindings[i].NumUniformBuffers; ++Index)
-					{
-						Bindings[i].UniformBuffers[Index] = InBindings[i].UniformBuffers[Index];
-					}
-				}
-
-				if (Bindings[i].LooseParameterDataSize)
-				{
-					Bindings[i].LooseParameterData = (uint8*)Alloc(Bindings[i].LooseParameterDataSize, 16);
-					FMemory::Memcpy(Bindings[i].LooseParameterData, InBindings[i].LooseParameterData, Bindings[i].LooseParameterDataSize);
-				}
-			}
-		}
-
-		if (Bypass())
-		{
-			GetContext().RHISetRayTracingBindings(Scene, GetRHIRayTracingPipelineState(Pipeline), NumBindings, Bindings, BindingType);
-		}
-		else
-		{
-			ALLOC_COMMAND(FRHICommandSetBindingsOnShaderBindingTable)(Scene, Pipeline, NumBindings, Bindings, BindingType);
-		}
-	}
-
 	FORCEINLINE_DEBUGGABLE void SetBindingsOnShaderBindingTable(
 		FRHIShaderBindingTable* SBT, FRayTracingPipelineState* Pipeline,
 		uint32 NumBindings, const FRayTracingLocalShaderBindings* Bindings,
@@ -4748,17 +4541,6 @@ public:
 		}
 	}
 
-	UE_DEPRECATED(5.5, "Use FRHIShaderBindingTable instead.")
-	FORCEINLINE_DEBUGGABLE void SetRayTracingHitGroups(
-		FRHIRayTracingScene* Scene, FRayTracingPipelineState* Pipeline,
-		uint32 NumBindings, const FRayTracingLocalShaderBindings* Bindings,
-		bool bCopyDataToInlineStorage = true)
-	{
-		PRAGMA_DISABLE_DEPRECATION_WARNINGS
-		SetRayTracingBindings(Scene, Pipeline, NumBindings, Bindings, ERayTracingBindingType::HitGroup, bCopyDataToInlineStorage);
-		PRAGMA_ENABLE_DEPRECATION_WARNINGS
-	}
-
 	FORCEINLINE_DEBUGGABLE void SetRayTracingHitGroups(
 		FRHIShaderBindingTable* SBT, FRayTracingPipelineState* Pipeline,
 		uint32 NumBindings, const FRayTracingLocalShaderBindings* Bindings,
@@ -4767,17 +4549,6 @@ public:
 		SetBindingsOnShaderBindingTable(SBT, Pipeline, NumBindings, Bindings, ERayTracingBindingType::HitGroup, bCopyDataToInlineStorage);
 	}
 
-	UE_DEPRECATED(5.5, "Use FRHIShaderBindingTable instead.")
-	FORCEINLINE_DEBUGGABLE void SetRayTracingCallableShaders(
-		FRHIRayTracingScene* Scene, FRayTracingPipelineState* Pipeline,
-		uint32 NumBindings, const FRayTracingLocalShaderBindings* Bindings,
-		bool bCopyDataToInlineStorage = true)
-	{
-		PRAGMA_DISABLE_DEPRECATION_WARNINGS
-		SetRayTracingBindings(Scene, Pipeline, NumBindings, Bindings, ERayTracingBindingType::CallableShader, bCopyDataToInlineStorage);
-		PRAGMA_ENABLE_DEPRECATION_WARNINGS
-	}
-
 	FORCEINLINE_DEBUGGABLE void SetRayTracingCallableShaders(
 		FRHIShaderBindingTable* SBT, FRayTracingPipelineState* Pipeline,
 		uint32 NumBindings, const FRayTracingLocalShaderBindings* Bindings,
@@ -4786,17 +4557,6 @@ public:
 		SetBindingsOnShaderBindingTable(SBT, Pipeline, NumBindings, Bindings, ERayTracingBindingType::CallableShader, bCopyDataToInlineStorage);
 	}
 
-	UE_DEPRECATED(5.5, "Use FRHIShaderBindingTable instead.")
-	FORCEINLINE_DEBUGGABLE void SetRayTracingMissShaders(
-		FRHIRayTracingScene* Scene, FRayTracingPipelineState* Pipeline,
-		uint32 NumBindings, const FRayTracingLocalShaderBindings* Bindings,
-		bool bCopyDataToInlineStorage = true)
-	{
-		PRAGMA_DISABLE_DEPRECATION_WARNINGS
-		SetRayTracingBindings(Scene, Pipeline, NumBindings, Bindings, ERayTracingBindingType::MissShader, bCopyDataToInlineStorage);
-		PRAGMA_ENABLE_DEPRECATION_WARNINGS
-	}
-
 	FORCEINLINE_DEBUGGABLE void SetRayTracingMissShaders(
 		FRHIShaderBindingTable* SBT, FRayTracingPipelineState* Pipeline,
 		uint32 NumBindings, const FRayTracingLocalShaderBindings* Bindings,
@@ -4805,50 +4565,6 @@ public:
 		SetBindingsOnShaderBindingTable(SBT, Pipeline, NumBindings, Bindings, ERayTracingBindingType::MissShader, bCopyDataToInlineStorage);
 	}
 
-	UE_DEPRECATED(5.5, "Use FRHIShaderBindingTable instead.")
-	FORCEINLINE_DEBUGGABLE void SetRayTracingHitGroup(
-		FRHIRayTracingScene* Scene, uint32 InstanceIndex, uint32 SegmentIndex, uint32 ShaderSlot,
-		FRayTracingPipelineState* Pipeline, uint32 HitGroupIndex,
-		uint32 NumUniformBuffers, FRHIUniformBuffer* const* UniformBuffers,
-		uint32 LooseParameterDataSize, const void* LooseParameterData,
-		uint32 UserData)
-	{
-		check(NumUniformBuffers <= UINT16_MAX);
-		check(LooseParameterDataSize <= UINT16_MAX);
-
-		FRayTracingLocalShaderBindings* InlineBindings = new(Alloc<FRayTracingLocalShaderBindings>()) FRayTracingLocalShaderBindings();
-		PRAGMA_DISABLE_DEPRECATION_WARNINGS
-		InlineBindings->Geometry = Scene->GetInitializer().PerInstanceGeometries[InstanceIndex];
-		PRAGMA_ENABLE_DEPRECATION_WARNINGS
-		InlineBindings->SegmentIndex = SegmentIndex;
-		PRAGMA_DISABLE_DEPRECATION_WARNINGS
-		InlineBindings->RecordIndex = (Scene->GetInitializer().SegmentPrefixSum[InstanceIndex] + SegmentIndex) * Scene->GetInitializer().ShaderSlotsPerGeometrySegment + ShaderSlot;
-		PRAGMA_ENABLE_DEPRECATION_WARNINGS
-		InlineBindings->ShaderIndexInPipeline = HitGroupIndex;
-		InlineBindings->UserData = UserData;
-		InlineBindings->NumUniformBuffers = (uint16)NumUniformBuffers;
-		InlineBindings->LooseParameterDataSize = (uint16)LooseParameterDataSize;
-
-		if (NumUniformBuffers)
-		{
-			InlineBindings->UniformBuffers = (FRHIUniformBuffer**)Alloc(sizeof(FRHIUniformBuffer*) * NumUniformBuffers, alignof(FRHIUniformBuffer*));
-			for (uint32 Index = 0; Index < NumUniformBuffers; ++Index)
-			{
-				InlineBindings->UniformBuffers[Index] = UniformBuffers[Index];
-			}
-		}
-
-		if (LooseParameterDataSize)
-		{
-			InlineBindings->LooseParameterData = (uint8*)Alloc(LooseParameterDataSize, 16);
-			FMemory::Memcpy(InlineBindings->LooseParameterData, LooseParameterData, LooseParameterDataSize);
-		}
-
-		PRAGMA_DISABLE_DEPRECATION_WARNINGS
-		SetRayTracingBindings(Scene, Pipeline, 1, InlineBindings, ERayTracingBindingType::HitGroup, /*bCopyDataToInlineStorage*/ false);
-		PRAGMA_ENABLE_DEPRECATION_WARNINGS
-	}
-
 	FORCEINLINE_DEBUGGABLE void SetRayTracingHitGroup(
 		FRHIShaderBindingTable* SBT, uint32 RecordIndex, FRHIRayTracingGeometry* Geometry, uint32 GeometrySegmentIndex,
 		FRayTracingPipelineState* Pipeline, uint32 HitGroupIndex,
@@ -4896,33 +4612,6 @@ public:
 		SetBindingsOnShaderBindingTable(SBT, Pipeline, 1, InlineBindings, ERayTracingBindingType::HitGroup, /*bCopyDataToInlineStorage*/ false);
 	}
 
-	UE_DEPRECATED(5.5, "Use FRHIShaderBindingTable instead.")
-	FORCEINLINE_DEBUGGABLE void SetRayTracingCallableShader(
-		FRHIRayTracingScene* Scene, uint32 ShaderSlotInScene,
-		FRayTracingPipelineState* Pipeline, uint32 ShaderIndexInPipeline,
-		uint32 NumUniformBuffers, FRHIUniformBuffer* const* UniformBuffers,
-		uint32 UserData)
-	{
-		FRayTracingLocalShaderBindings* InlineBindings = new(Alloc<FRayTracingLocalShaderBindings>()) FRayTracingLocalShaderBindings();
-		InlineBindings->RecordIndex = ShaderSlotInScene;
-		InlineBindings->ShaderIndexInPipeline = ShaderIndexInPipeline;
-		InlineBindings->UserData = UserData;
-		InlineBindings->NumUniformBuffers = (uint16)NumUniformBuffers;
-
-		if (NumUniformBuffers)
-		{
-			InlineBindings->UniformBuffers = (FRHIUniformBuffer**)Alloc(sizeof(FRHIUniformBuffer*) * NumUniformBuffers, alignof(FRHIUniformBuffer*));
-			for (uint32 Index = 0; Index < NumUniformBuffers; ++Index)
-			{
-				InlineBindings->UniformBuffers[Index] = UniformBuffers[Index];
-			}
-		}
-
-		PRAGMA_DISABLE_DEPRECATION_WARNINGS
-		SetRayTracingBindings(Scene, Pipeline, 1, InlineBindings, ERayTracingBindingType::CallableShader, /*bCopyDataToInlineStorage*/ false);
-		PRAGMA_ENABLE_DEPRECATION_WARNINGS
-	}
-
 	FORCEINLINE_DEBUGGABLE void SetRayTracingCallableShader(
 		FRHIShaderBindingTable* SBT, uint32 RecordIndex,
 		FRayTracingPipelineState* Pipeline, uint32 ShaderIndexInPipeline,
@@ -4947,33 +4636,6 @@ public:
 		SetBindingsOnShaderBindingTable(SBT, Pipeline, 1, InlineBindings, ERayTracingBindingType::CallableShader, /*bCopyDataToInlineStorage*/ false);
 	}
 
-	UE_DEPRECATED(5.5, "Use FRHIShaderBindingTable instead.")
-	FORCEINLINE_DEBUGGABLE void SetRayTracingMissShader(
-		FRHIRayTracingScene* Scene, uint32 ShaderSlotInScene,
-		FRayTracingPipelineState* Pipeline, uint32 ShaderIndexInPipeline,
-		uint32 NumUniformBuffers, FRHIUniformBuffer* const* UniformBuffers,
-		uint32 UserData)
-	{
-		FRayTracingLocalShaderBindings* InlineBindings = new(Alloc<FRayTracingLocalShaderBindings>()) FRayTracingLocalShaderBindings();
-		InlineBindings->RecordIndex = ShaderSlotInScene;
-		InlineBindings->ShaderIndexInPipeline = ShaderIndexInPipeline;
-		InlineBindings->UserData = UserData;
-		InlineBindings->NumUniformBuffers = (uint16)NumUniformBuffers;
-
-		if (NumUniformBuffers)
-		{
-			InlineBindings->UniformBuffers = (FRHIUniformBuffer**)Alloc(sizeof(FRHIUniformBuffer*) * NumUniformBuffers, alignof(FRHIUniformBuffer*));
-			for (uint32 Index = 0; Index < NumUniformBuffers; ++Index)
-			{
-				InlineBindings->UniformBuffers[Index] = UniformBuffers[Index];
-			}
-		}
-
-		PRAGMA_DISABLE_DEPRECATION_WARNINGS
-		SetRayTracingBindings(Scene, Pipeline, 1, InlineBindings, ERayTracingBindingType::MissShader, /*bCopyDataToInlineStorage*/ false);
-		PRAGMA_ENABLE_DEPRECATION_WARNINGS
-	}
-
 	FORCEINLINE_DEBUGGABLE void SetRayTracingMissShader(
 		FRHIShaderBindingTable* SBT, uint32 RecordIndex,
 		FRayTracingPipelineState* Pipeline, uint32 ShaderIndexInPipeline,
@@ -5087,18 +4749,9 @@ public:
 		return static_cast<FRHICommandListImmediate&>(RHICmdList);
 	}
 
-	UE_DEPRECATED(5.5, "FRHICommandListImmediate::BeginScene() is deprecated and is no longer necessary. Remove any remaining calls to BeginScene(). There is no replacement.")
-	void BeginScene() {}
-
-	UE_DEPRECATED(5.5, "FRHICommandListImmediate::EndScene() is deprecated and is no longer necessary. Remove any remaining calls to EndScene(). There is no replacement.")
-	void EndScene() {}
-
 	RHI_API void BeginDrawingViewport(FRHIViewport* Viewport, FRHITexture* RenderTargetRHI);
 	RHI_API void EndDrawingViewport(FRHIViewport* Viewport, bool bPresent, bool bLockToVsync);
 
-	UE_DEPRECATED(5.5, "FRHICommandListImmediate::BeginFrame() is deprecated and is no longer necessary. Remove any remaining calls to BeginFrame(). There is no replacement.")
-	void BeginFrame() {}
-
 	RHI_API void EndFrame();
 
 	struct FQueuedCommandList
@@ -5174,13 +4827,6 @@ public:
 		GDynamicRHI->UnlockStagingBuffer_RenderThread(*this, StagingBuffer);
 	}
 
-	UE_DEPRECATED(5.5, "CopyBuffer is deprecated. Use CopyBufferRegion.")
-	FORCEINLINE void CopyBuffer(FRHIBuffer* SourceBuffer, FRHIBuffer* DestBuffer)
-	{
-		uint64 Size = FMath::Min(SourceBuffer->GetSize(), DestBuffer->GetSize());
-		CopyBufferRegion(DestBuffer, 0, SourceBuffer, 0, Size);
-	}
-
 	FORCEINLINE bool GetTextureMemoryVisualizeData(FColor* TextureData,int32 SizeX,int32 SizeY,int32 Pitch,int32 PixelSize)
 	{
 		QUICK_SCOPE_CYCLE_COUNTER(STAT_RHIMETHOD_GetTextureMemoryVisualizeData_Flush);
@@ -5354,14 +5000,6 @@ public:
 		GDynamicRHI->RHIRead3DSurfaceFloatData(Texture,Rect,ZMinMax,OutData,Flags);
 	}
 	
-	UE_DEPRECATED(5.5, "FRHICommandListImmediate::AcquireThreadOwnership() is deprecated. Thread ownership of the RHI is automatic, so this function is redundant.")
-	FORCEINLINE void AcquireThreadOwnership()
-	{}
-	
-	UE_DEPRECATED(5.5, "FRHICommandListImmediate::ReleaseThreadOwnership() is deprecated. Thread ownership of the RHI is automatic, so this function is redundant.")
-	FORCEINLINE void ReleaseThreadOwnership()
-	{}
-	
 	FORCEINLINE void FlushResources()
 	{
 		QUICK_SCOPE_CYCLE_COUNTER(STAT_RHIMETHOD_FlushResources_Flush);
@@ -5370,9 +5008,6 @@ public:
 		return GDynamicRHI->RHIFlushResources();
 	}
 
-	UE_DEPRECATED(5.5, "FlushPendingDeletes is deprecated and removed. Use RHICmdList.ImmediateFlush(EImmediateFlushType::FlushRHIThreadFlushResources)")
-	inline int32 FlushPendingDeletes() { return 0; }
-	
 	UE_DEPRECATED(5.6, "GetGPUFrameCycles is deprecated. Use the global scope RHIGetGPUFrameCycles() function.")
 	FORCEINLINE uint32 GetGPUFrameCycles()
 	{
@@ -5471,9 +5106,6 @@ public:
 		return GDynamicRHI->RHIGetNativeCommandBuffer();
 	}
 
-	UE_DEPRECATED(5.5, "RHIPollRenderQueryResults is deprecated. Platform RHIs that require query polling now do this automatically as part of RHI command list submission.")
-	FORCEINLINE void PollRenderQueryResults() {}
-
 	//UE_DEPRECATED(5.1, "SubmitCommandsHint is deprecated. Consider calling ImmediateFlush(EImmediateFlushType::DispatchToRHIThread) instead.")
 	FORCEINLINE_DEBUGGABLE void SubmitCommandsHint()
 	{
@@ -5764,15 +5396,6 @@ public:
 	//
 	RHI_API void WaitForTasks(FGraphEventArray& OutstandingTasks);
 
-	//
-	// Blocks the calling thread until the RHI thread is idle.
-	//
-	UE_DEPRECATED(5.5, "FRHICommandListExecutor::WaitForRHIThreadTasks() is deprecated. Call FRHICommandListImmediate::ImmediateFlush(EImmediateFlushType::FlushRHIThread) instead.")
-	inline void WaitForRHIThreadTasks()
-	{
-		FRHICommandListImmediate::Get().ImmediateFlush(EImmediateFlushType::FlushRHIThread);
-	}
-
 	// Global graph events must be destroyed explicitly to avoid undefined order of static destruction, as they can be destroyed after their allocator.
 	void CleanupGraphEvents();
 
@@ -5802,18 +5425,6 @@ public:
 	//
 	RHI_API static bool AreRHITasksActive();
 
-	UE_DEPRECATED(5.5, "FRHICommandListExecutor::IsRHIThreadActive() is deprecated. Use FRHICommandListExecutor::AreRHITasksActive() instead.")
-	static inline bool IsRHIThreadActive()
-	{
-		return AreRHITasksActive();
-	}
-
-	UE_DEPRECATED(5.5, "FRHICommandListExecutor::IsRHIThreadCompletelyFlushed() is deprecated. Use FRHICommandListExecutor::AreRHITasksActive() instead.")
-	static inline bool IsRHIThreadCompletelyFlushed()
-	{
-		return !AreRHITasksActive();
-	}
-
 	//
 	// Adds a prerequisite for subsequent Submit dispatch tasks.
 	// 
@@ -6005,10 +5616,6 @@ extern RHI_API FRHICommandListExecutor GRHICommandList;
 
 extern RHI_API FAutoConsoleTaskPriority CPrio_SceneRenderingTask;
 
-/** Used to separate which command list is used for ray tracing operations. */
-UE_DEPRECATED(5.5, "Use FRHICommandList instead.")
-typedef FRHICommandListImmediate FRHIRayTracingCommandList;
-
 class FRenderTask
 {
 public:
@@ -6046,14 +5653,6 @@ FORCEINLINE FTextureRHIRef RHICreateTexture(const FRHITextureCreateDesc& CreateD
 	return FRHICommandListImmediate::Get().CreateTexture(CreateDesc);
 }
 
-UE_DEPRECATED(5.4, "Use the RHIAsyncCreateTexture2D function that takes an DebugName.")
-FORCEINLINE FTextureRHIRef RHIAsyncCreateTexture2D(uint32 SizeX, uint32 SizeY, uint8 Format, uint32 NumMips, ETextureCreateFlags Flags, ERHIAccess InResourceState, void** InitialMipData, uint32 NumInitialMips, FGraphEventRef& OutCompletionEvent)
-{
-	LLM_SCOPE(EnumHasAnyFlags(Flags, TexCreate_RenderTargetable | TexCreate_DepthStencilTargetable) ? ELLMTag::RenderTargets : ELLMTag::Textures);
-	const ERHIAccess ResourceState = InResourceState == ERHIAccess::Unknown ? RHIGetDefaultResourceState((ETextureCreateFlags)Flags, InitialMipData != nullptr) : InResourceState;
-	return GDynamicRHI->RHIAsyncCreateTexture2D(SizeX, SizeY, Format, NumMips, Flags, ResourceState, InitialMipData, NumInitialMips, TEXT("RHIAsyncCreateTexture2D"), OutCompletionEvent);
-}
-
 FORCEINLINE FTextureRHIRef RHIAsyncCreateTexture2D(uint32 SizeX, uint32 SizeY, uint8 Format, uint32 NumMips, ETextureCreateFlags Flags, ERHIAccess InResourceState, void** InitialMipData, uint32 NumInitialMips, const TCHAR* DebugName, FGraphEventRef& OutCompletionEvent)
 {
 	LLM_SCOPE(EnumHasAnyFlags(Flags, TexCreate_RenderTargetable | TexCreate_DepthStencilTargetable) ? ELLMTag::RenderTargets : ELLMTag::Textures);
@@ -6061,14 +5660,6 @@ FORCEINLINE FTextureRHIRef RHIAsyncCreateTexture2D(uint32 SizeX, uint32 SizeY, u
 	return GDynamicRHI->RHIAsyncCreateTexture2D(SizeX, SizeY, Format, NumMips, Flags, ResourceState, InitialMipData, NumInitialMips, DebugName, OutCompletionEvent);
 }
 
-UE_DEPRECATED(5.4, "Use the RHIAsyncCreateTexture2D function that takes an InResourceState and DebugName.")
-FORCEINLINE FTextureRHIRef RHIAsyncCreateTexture2D(uint32 SizeX, uint32 SizeY, uint8 Format, uint32 NumMips, ETextureCreateFlags Flags, void** InitialMipData, uint32 NumInitialMips, FGraphEventRef& OutCompletionEvent)
-{
-PRAGMA_DISABLE_DEPRECATION_WARNINGS
-	return RHIAsyncCreateTexture2D(SizeX, SizeY, Format, NumMips, Flags, ERHIAccess::Unknown, InitialMipData, NumInitialMips, OutCompletionEvent);
-PRAGMA_ENABLE_DEPRECATION_WARNINGS
-}
-
 FORCEINLINE FTextureRHIRef RHIAsyncReallocateTexture2D(FRHITexture* Texture2D, int32 NewMipCount, int32 NewSizeX, int32 NewSizeY, FThreadSafeCounter* RequestStatus)
 {
 	return FRHICommandListExecutor::GetImmediateCommandList().AsyncReallocateTexture2D(Texture2D, NewMipCount, NewSizeX, NewSizeY, RequestStatus);
@@ -6147,14 +5738,6 @@ FORCEINLINE void RHIUpdateTexture3D(FRHITexture* Texture, uint32 MipIndex, const
 	 FRHICommandListExecutor::GetImmediateCommandList().UpdateTexture3D(Texture, MipIndex, UpdateRegion, SourceRowPitch, SourceDepthPitch, SourceData);
 }
 
-UE_DEPRECATED(5.5, "RHIAcquireThreadOwnership() is deprecated. Thread ownership of the RHI is automatic, so this function is redundant.")
-FORCEINLINE void RHIAcquireThreadOwnership()
-{}
-
-UE_DEPRECATED(5.5, "RHIReleaseThreadOwnership() is deprecated. Thread ownership of the RHI is automatic, so this function is redundant.")
-FORCEINLINE void RHIReleaseThreadOwnership()
-{}
-
 FORCEINLINE void RHIFlushResources()
 {
 	return FRHICommandListExecutor::GetImmediateCommandList().FlushResources();
diff --git a/Engine/Source/Runtime/RHI/Public/RHICommandListCommandExecutes.inl b/Engine/Source/Runtime/RHI/Public/RHICommandListCommandExecutes.inl
--- a/Engine/Source/Runtime/RHI/Public/RHICommandListCommandExecutes.inl
+++ b/Engine/Source/Runtime/RHI/Public/RHICommandListCommandExecutes.inl
@@ -48,7 +48,6 @@ struct FRHICommandTransitionTextures;
 struct FRHICommandTransitionTexturesPipeline;
 struct FRHICommandTransitionTexturesDepth;
 struct FRHICommandTransitionTexturesArray;
-struct FRHICommandClearRayTracingBindings;
 struct FRHICommandClearBindings;
 struct FRHICommandRayTraceDispatch;
 struct FRHICommandSetBindings;
@@ -539,18 +538,6 @@ void FRHICommandBuildSceneAccelerationStructures::Execute(FRHICommandListBase& C
 	INTERNAL_DECORATOR_COMPUTE(RHIBuildAccelerationStructures)(Params);
 }
 
-void FRHICommandCommitRayTracingBindings::Execute(FRHICommandListBase& CmdList)
-{
-	RHISTAT(CommitRayTracingBindings);
-	INTERNAL_DECORATOR(RHICommitRayTracingBindings)(Scene);
-}
-
-void FRHICommandClearRayTracingBindings::Execute(FRHICommandListBase& CmdList)
-{
-	RHISTAT(ClearRayTracingBindings);
-	INTERNAL_DECORATOR(RHIClearRayTracingBindings)(Scene);
-}
-
 void FRHICommandCommitShaderBindingTable::Execute(FRHICommandListBase& CmdList)
 {
 	RHISTAT(CommitShaderBindingTable);
@@ -576,13 +563,6 @@ void FRHICommandRayTraceDispatch::Execute(FRHICommandListBase& CmdList)
 	
 	FRHIRayTracingPipelineState* RayTracingPipelineState = GetRHIRayTracingPipelineState(Pipeline);
 
-	PRAGMA_DISABLE_DEPRECATION_WARNINGS
-	if (SBT == nullptr)
-	{
-		SBT = Scene->FindOrCreateShaderBindingTable(RayTracingPipelineState);
-	}
-	PRAGMA_ENABLE_DEPRECATION_WARNINGS
-
 	if (ArgumentBuffer)
 	{
 		INTERNAL_DECORATOR_COMPUTE(RHIRayTraceDispatchIndirect)(RayTracingPipelineState, RayGenShader, SBT, GlobalResourceBindings, ArgumentBuffer, ArgumentOffset);
@@ -600,13 +580,6 @@ void FRHICommandSetBindingsOnShaderBindingTable::Execute(FRHICommandListBase& Cm
 	
 	FRHIRayTracingPipelineState* RayTracingPipelineState = GetRHIRayTracingPipelineState(Pipeline);
 
-	PRAGMA_DISABLE_DEPRECATION_WARNINGS
-	if (SBT == nullptr)
-	{
-		SBT = Scene->FindOrCreateShaderBindingTable(RayTracingPipelineState);
-	}
-	PRAGMA_ENABLE_DEPRECATION_WARNINGS
-		
 	INTERNAL_DECORATOR(RHISetBindingsOnShaderBindingTable)(SBT, RayTracingPipelineState, NumBindings, Bindings, BindingType);
 }
 
diff --git a/Engine/Source/Runtime/RHI/Public/RHIContext.h b/Engine/Source/Runtime/RHI/Public/RHIContext.h
--- a/Engine/Source/Runtime/RHI/Public/RHIContext.h
+++ b/Engine/Source/Runtime/RHI/Public/RHIContext.h
@@ -275,19 +275,6 @@ public:
 		checkNoEntry();
 	}
 
-	void RHIRayTraceDispatch(FRHIRayTracingPipelineState* RayTracingPipelineState, FRHIRayTracingShader* RayGenShader,
-		FRHIRayTracingScene* Scene,
-		const FRayTracingShaderBindings& GlobalResourceBindings,
-		uint32 Width, uint32 Height)
-	{
-		PRAGMA_DISABLE_DEPRECATION_WARNINGS
-			RHIRayTraceDispatch(RayTracingPipelineState, RayGenShader,
-				Scene->FindOrCreateShaderBindingTable(RayTracingPipelineState),
-				GlobalResourceBindings,
-				Width, Height);
-		PRAGMA_ENABLE_DEPRECATION_WARNINGS
-	}
-
 	virtual void RHIRayTraceDispatchIndirect(FRHIRayTracingPipelineState* RayTracingPipelineState, FRHIRayTracingShader* RayGenShader,
 		FRHIShaderBindingTable* SBT, const FRayTracingShaderBindings& GlobalResourceBindings,
 		FRHIBuffer* ArgumentBuffer, uint32 ArgumentOffset)
@@ -295,19 +282,6 @@ public:
 		checkNoEntry();
 	}
 
-	void RHIRayTraceDispatchIndirect(FRHIRayTracingPipelineState* RayTracingPipelineState, FRHIRayTracingShader* RayGenShader,
-		FRHIRayTracingScene* Scene,
-		const FRayTracingShaderBindings& GlobalResourceBindings,
-		FRHIBuffer* ArgumentBuffer, uint32 ArgumentOffset)
-	{
-		PRAGMA_DISABLE_DEPRECATION_WARNINGS
-			RHIRayTraceDispatchIndirect(RayTracingPipelineState, RayGenShader,
-				Scene->FindOrCreateShaderBindingTable(RayTracingPipelineState),
-				GlobalResourceBindings,
-				ArgumentBuffer, ArgumentOffset);
-		PRAGMA_ENABLE_DEPRECATION_WARNINGS
-	}
-
 	virtual void RHIDispatchIndirectComputeShader(FRHIBuffer* ArgumentBuffer, uint32 ArgumentOffset) = 0;
 
 	virtual void RHISetAsyncComputeBudget(EAsyncComputeBudget Budget) {}
@@ -837,11 +811,6 @@ public:
 
 	virtual void RHICopyBufferRegion(FRHIBuffer* DestBuffer, uint64 DstOffset, FRHIBuffer* SourceBuffer, uint64 SrcOffset, uint64 NumBytes) = 0;
 
-	virtual void RHIClearRayTracingBindings(FRHIRayTracingScene* Scene)
-	{
-		checkNoEntry();
-	}
-
 	virtual void RHIClearShaderBindingTable(FRHIShaderBindingTable* SBT)
 	{
 		checkNoEntry();
@@ -857,18 +826,6 @@ public:
 		checkNoEntry();
 	}
 
-	void RHISetRayTracingBindings(FRHIRayTracingScene* Scene, FRHIRayTracingPipelineState* Pipeline, uint32 NumBindings, const FRayTracingLocalShaderBindings* Bindings, ERayTracingBindingType BindingType)
-	{
-		PRAGMA_DISABLE_DEPRECATION_WARNINGS
-		RHISetBindingsOnShaderBindingTable(Scene->FindOrCreateShaderBindingTable(Pipeline), Pipeline, NumBindings, Bindings, BindingType);
-		PRAGMA_ENABLE_DEPRECATION_WARNINGS
-	}
-
-	virtual void RHICommitRayTracingBindings(FRHIRayTracingScene* Scene)
-	{
-		checkNoEntry();
-	}
-
 protected:
 	FRHIRenderPassInfo RenderPassInfo;
 };
diff --git a/Engine/Source/Runtime/RHI/Public/RHIDefinitions.h b/Engine/Source/Runtime/RHI/Public/RHIDefinitions.h
--- a/Engine/Source/Runtime/RHI/Public/RHIDefinitions.h
+++ b/Engine/Source/Runtime/RHI/Public/RHIDefinitions.h
@@ -1100,12 +1100,9 @@ enum class ETextureCreateFlags : uint64
     Tiling3D                          = 1ull << 11,
     // This texture has no GPU or CPU backing. It only exists in tile memory on TBDR GPUs (i.e., mobile).
     Memoryless                        = 1ull << 12,
-    // Create the texture with the flag that allows mip generation later, only applicable to D3D11
-    GenerateMipCapable UE_DEPRECATED(5.5,
-		"The GenerateMipCapable flag is no longer used."
-		" Use ETextureCreateFlags::UAV to make a texture compatible with compute-based mip generation,"
-		" or use ETextureCreateFlags::RenderTargetable for raster-based mip generation.")
-		                              = 1ull << 13,
+
+    //UNUSED_BIT_13                   = 1ull << 13,
+
     // The texture can be partially allocated in fastvram
     FastVRAMPartialAlloc              = 1ull << 14,
     // Do not create associated shader resource view, only applicable to D3D11 and D3D12
diff --git a/Engine/Source/Runtime/RHI/Public/RHIFwd.h b/Engine/Source/Runtime/RHI/Public/RHIFwd.h
--- a/Engine/Source/Runtime/RHI/Public/RHIFwd.h
+++ b/Engine/Source/Runtime/RHI/Public/RHIFwd.h
@@ -139,13 +139,3 @@ using FViewportRHIRef                  = TRefCountPtr<FRHIViewport>;
 using FWorkGraphPipelineStateRHIRef    = TRefCountPtr<FRHIWorkGraphPipelineState>;
 using FStreamSourceSlotRHIRef          = TRefCountPtr<FRHIStreamSourceSlot>;
 using FWorkGraphShaderRHIRef           = TRefCountPtr<FRHIWorkGraphShader>;
-
-// Deprecated typenames
-using FRHITexture2D         UE_DEPRECATED(5.5, "The separate RHI texture types were unified in UE 5.1. FRHITexture2D is now deprecated. Use FRHITexture instead."           ) = FRHITexture;
-using FRHITexture2DArray    UE_DEPRECATED(5.5, "The separate RHI texture types were unified in UE 5.1. FRHITexture2DArray is now deprecated. Use FRHITexture instead."      ) = FRHITexture;
-using FRHITexture3D         UE_DEPRECATED(5.5, "The separate RHI texture types were unified in UE 5.1. FRHITexture3D is now deprecated. Use FRHITexture instead."           ) = FRHITexture;
-using FRHITextureCube       UE_DEPRECATED(5.5, "The separate RHI texture types were unified in UE 5.1. FRHITextureCube is now deprecated. Use FRHITexture instead."         ) = FRHITexture;
-using FTexture2DRHIRef      UE_DEPRECATED(5.5, "The separate RHI texture types were unified in UE 5.1. FTexture2DRHIRef is now deprecated. Use FTextureRHIRef instead."     ) = FTextureRHIRef;
-using FTexture2DArrayRHIRef UE_DEPRECATED(5.5, "The separate RHI texture types were unified in UE 5.1. FTexture2DArrayRHIRef is now deprecated. Use FTextureRHIRef instead.") = FTextureRHIRef;
-using FTexture3DRHIRef      UE_DEPRECATED(5.5, "The separate RHI texture types were unified in UE 5.1. FTexture3DRHIRef is now deprecated. Use FTextureRHIRef instead."     ) = FTextureRHIRef;
-using FTextureCubeRHIRef    UE_DEPRECATED(5.5, "The separate RHI texture types were unified in UE 5.1. FTextureCubeRHIRef is now deprecated. Use FTextureRHIRef instead."   ) = FTextureRHIRef;
diff --git a/Engine/Source/Runtime/RHI/Public/RHIGlobals.h b/Engine/Source/Runtime/RHI/Public/RHIGlobals.h
--- a/Engine/Source/Runtime/RHI/Public/RHIGlobals.h
+++ b/Engine/Source/Runtime/RHI/Public/RHIGlobals.h
@@ -540,14 +540,6 @@ struct FRHIGlobals
 
 	struct FVariableRateShading
 	{
-		/** Whether VRS (in all flavors) is currently enabled (separate from whether it's supported/available). */
-		UE_DEPRECATED(5.5, "GRHIGlobals.VariableRateShading.Enabled has been deprecated - please use GRHIGlobals.VariableRateShading.SupportsPipeline or GVRSImageManager.IsPipelineVRSEnabled() instead.")
-		bool Enabled = true;
-
-		/** Whether attachment (image-based) VRS is currently enabled (separate from whether it's supported/available). */
-		UE_DEPRECATED(5.5, "GRHIGlobals.VariableRateShading.AttachmentEnabled has been deprecated - please use GRHIGlobals.VariableRateShading.SupportsAttachment or GVRSImageManager.IsAttachmentVRSEnabled() instead.")
-		bool AttachmentEnabled = true;
-
 		/** Whether or not the RHI can support per-draw Variable Rate Shading. */
 		bool SupportsPipeline = false;
 
@@ -706,10 +698,6 @@ struct FRHIGlobals
 	/** True if the RHI supports shaders with barycentrics */
 	bool SupportsBarycentricsSemantic = false;
 
-	/** True if HDR requires vendor specific extensions */
-	UE_DEPRECATED(5.5, "GRHIGlobals.HDRNeedsVendorExtensions has been deprecated - Vendor extensions are not used anymore")
-	bool HDRNeedsVendorExtensions = false;
-
 	/** True if RHI supports MSAA resolve with a custom shader */
 	bool SupportsMSAAShaderResolve = false;
 
@@ -879,8 +867,6 @@ extern RHI_API uint32 GGPUFrameTime;
 #define GRHIIsHDREnabled                                       GRHIGlobals.IsHDREnabled
 #define GRHISupportsHDROutput                                  GRHIGlobals.SupportsHDROutput
 #define GRHIMaxDispatchThreadGroupsPerDimension                GRHIGlobals.MaxDispatchThreadGroupsPerDimension
-#define GRHIVariableRateShadingEnabled                         GRHIGlobals.VariableRateShading.Enabled UE_DEPRECATED_MACRO(5.5, "GRHIVariableRateShadingEnabled has been deprecated - please use GRHISupportsPipelineVariableRateShading or GVRSImageManager.IsPipelineVRSEnabled() instead.")
-#define GRHIAttachmentVariableRateShadingEnabled               GRHIGlobals.VariableRateShading.AttachmentEnabled UE_DEPRECATED_MACRO(5.5, "GRHIAttachmentVariableRateShadingEnabled has been deprecated - please use GRHISupportsAttachmentVariableRateShading or GVRSImageManager.IsAttachmentVRSEnabled() instead.")
 #define GRHISupportsPipelineVariableRateShading                GRHIGlobals.VariableRateShading.SupportsPipeline
 #define GRHISupportsLargerVariableRateShadingSizes             GRHIGlobals.VariableRateShading.SupportsLargerSizes
 #define GRHISupportsAttachmentVariableRateShading              GRHIGlobals.VariableRateShading.SupportsAttachment
@@ -914,10 +900,8 @@ extern RHI_API uint32 GGPUFrameTime;
 #define GRHIDefaultMSAASampleOffsets                           GRHIGlobals.DefaultMSAASampleOffsets
 #define GRHISupportsAsyncPipelinePrecompile                    GRHIGlobals.SupportsAsyncPipelinePrecompile
 #define GRHIBindlessSupport                                    GRHIGlobals.BindlessSupport
-#define GRHISupportsReservedResources                          GRHIGlobals.ReservedResources.Supported UE_DEPRECATED_MACRO(5.4, "GRHISupportsReservedResources has been deprecated - please use GRHIGlobals.ReservedResources.Supported instead.")
 #define GShaderPlatformForFeatureLevel                         GRHIGlobals.ShaderPlatformForFeatureLevel
 #define GRHIIsDebugLayerEnabled                                GRHIGlobals.IsDebugLayerEnabled
-#define GRHIHDRNeedsVendorExtensions						   GRHIGlobals.HDRNeedsVendorExtensions
 #define GRHISupportsMSAAShaderResolve						   GRHIGlobals.SupportsMSAAShaderResolve
 #define GRHISupportsDepthStencilResolve                        GRHIGlobals.SupportsDepthStencilResolve
 #define GRHISupportsLossyFramebufferCompression                GRHIGlobals.SupportsLossyFramebufferCompression
diff --git a/Engine/Source/Runtime/RHI/Public/RHIPipeline.h b/Engine/Source/Runtime/RHI/Public/RHIPipeline.h
--- a/Engine/Source/Runtime/RHI/Public/RHIPipeline.h
+++ b/Engine/Source/Runtime/RHI/Public/RHIPipeline.h
@@ -49,23 +49,6 @@ inline ERHIPipeline GetEnabledRHIPipelines()
 		: ERHIPipeline::Graphics;
 }
 
-UE_DEPRECATED(5.5, "GetRHIPipelines is deprecated. Prefer ranged-for iteration over pipelines using 'for (ERHIPipeline Pipeline : MakeFlagsRange(Pipelines))'.")
-inline TArrayView<const ERHIPipeline> GetRHIPipelines()
-{
-	static const ERHIPipeline Pipelines[] = { ERHIPipeline::Graphics, ERHIPipeline::AsyncCompute };
-	return Pipelines;
-}
-
-template <typename FunctionType>
-UE_DEPRECATED(5.5, "EnumerateRHIPipelines is deprecated. Prefer ranged-for iteration over pipelines using 'for (ERHIPipeline Pipeline : MakeFlagsRange(Pipelines))'.")
-inline void EnumerateRHIPipelines(ERHIPipeline PipelineMask, FunctionType Function)
-{
-	for (ERHIPipeline Pipeline : MakeFlagsRange(PipelineMask))
-	{
-		Function(Pipeline);
-	}
-}
-
 /** Array of pass handles by RHI pipeline, with overloads to help with enum conversion. */
 template <typename ElementType>
 class TRHIPipelineArray : public TStaticArray<ElementType, GetRHIPipelineCount()>
diff --git a/Engine/Source/Runtime/RHI/Public/RHIResources.h b/Engine/Source/Runtime/RHI/Public/RHIResources.h
--- a/Engine/Source/Runtime/RHI/Public/RHIResources.h
+++ b/Engine/Source/Runtime/RHI/Public/RHIResources.h
@@ -97,9 +97,6 @@ public:
 		return uint32(CurrentValue);
 	}
 
-	UE_DEPRECATED(5.5, "Don't call Bypass() on an FRHIResource. Use the Bypass() function on an FRHICommandList instance, or the FRHICommmandListExecutor.")
-	RHI_API static bool Bypass();
-
 	bool IsValid() const
 	{
 		return AtomicFlags.IsValid(std::memory_order_relaxed);
@@ -382,57 +379,26 @@ PRAGMA_DISABLE_DEPRECATION_WARNINGS
 		check(InDebugName);
 	}
 
-	// for CreateTexture calls
-	FRHIResourceCreateInfo(const TCHAR* InDebugName, FResourceBulkDataInterface* InBulkData)
-		: FRHIResourceCreateInfo(InDebugName)
-	{
-		BulkData = InBulkData;
-	}
-
 	// for CreateBuffer calls
 	FRHIResourceCreateInfo(const TCHAR* InDebugName, FResourceArrayUploadInterface* InResourceArray)
 		: FRHIResourceCreateInfo(InDebugName)
 	{
 		ResourceArray = InResourceArray;
 	}
-
-	FRHIResourceCreateInfo(const TCHAR* InDebugName, const FClearValueBinding& InClearValueBinding)
-		: FRHIResourceCreateInfo(InDebugName)
-	{
-		ClearValueBinding = InClearValueBinding;
-	}
-
-	FRHIResourceCreateInfo(uint32 InExtData)
-		: FRHIResourceCreateInfo(TEXT(""))
-	{
-		ExtData = InExtData;
-	}
 PRAGMA_ENABLE_DEPRECATION_WARNINGS
 
 	FName GetTraceClassName() const									{ const static FLazyName FRHIBufferName(TEXT("FRHIBuffer")); return (ClassName == NAME_None) ? FRHIBufferName : ClassName; }
 
-	// for CreateTexture calls
-	UE_DEPRECATED(5.5, "Please use FRHITextureCreateDesc for creating Textures with Bulk Data")
-	FResourceBulkDataInterface* BulkData = nullptr;
-
 	// for CreateBuffer calls
 	UE_DEPRECATED(5.6, "Please use FRHIBufferCreateDesc for creating buffers with Resource Arrays")
 	FResourceArrayUploadInterface* ResourceArray = nullptr;
 
-	// for binding clear colors to render targets.
-	UE_DEPRECATED(5.5, "Please use FRHITextureCreateDesc for creating Textures with a Clear Value Binding")
-	FClearValueBinding ClearValueBinding = FClearValueBinding(FLinearColor::Transparent);
-
 	// set of GPUs on which to create the resource
 	FRHIGPUMask GPUMask = FRHIGPUMask::All();
 
 	// whether to create an RHI object with no underlying resource
 	bool bWithoutNativeResource = false;
 
-	// optional data that would have come from an offline cooker or whatever - general purpose
-	UE_DEPRECATED(5.5, "Please use FRHITextureCreateDesc for creating Textures with Ext Data")
-	uint32 ExtData = 0;
-
 	const TCHAR* DebugName;
 
 	FName ClassName = NAME_None;	// The owner class of FRHIBuffer used for Insight asset metadata tracing
@@ -3330,15 +3296,6 @@ class FRHIRayTracingGeometry;
 */
 struct FRayTracingGeometryInstance
 {
-	PRAGMA_DISABLE_DEPRECATION_WARNINGS
-	FRayTracingGeometryInstance() = default;
-	FRayTracingGeometryInstance(const FRayTracingGeometryInstance&) = default;
-	FRayTracingGeometryInstance& operator=(const FRayTracingGeometryInstance&) = default;
-	FRayTracingGeometryInstance(FRayTracingGeometryInstance&&) = default;
-	FRayTracingGeometryInstance& operator=(FRayTracingGeometryInstance&&) = default;
-	~FRayTracingGeometryInstance() = default;
-	PRAGMA_ENABLE_DEPRECATION_WARNINGS
-
 	FRHIRayTracingGeometry* GeometryRHI = nullptr;
 
 	int32 InstanceContributionToHitGroupIndex = INDEX_NONE;
@@ -3352,10 +3309,6 @@ struct FRayTracingGeometryInstance
 	int32 BaseInstanceSceneDataOffset = -1;
 	TArrayView<const uint32> InstanceSceneDataOffsets;
 
-	// Optional buffer that stores GPU transforms. Used instead of CPU-side transform data.
-	UE_DEPRECATED(5.5, "GPUTransformsSRV has been deprecated. GPU Scene should be used instead.")
-	FShaderResourceViewRHIRef GPUTransformsSRV = nullptr;
-
 	// Conservative number of instances. Some of the actual instances may be made inactive if GPU transforms are used.
 	// Must be less or equal to number of entries in Transforms view if CPU transform data is used.
 	uint32 NumTransforms = 0;
@@ -3376,9 +3329,6 @@ struct FRayTracingGeometryInstance
 	// If binary AND of instance mask with ray mask is zero, then the instance is considered not intersected / invisible.
 	uint8 Mask = 0xFF;
 
-	UE_DEPRECATED(5.5, "Specify layer when adding instances to FRayTracingScene instead.")
-	uint8 LayerIndex = 0;
-
 	// Flags to control triangle back face culling, whether to allow any-hit shaders, etc.
 	ERayTracingInstanceFlags Flags = ERayTracingInstanceFlags::None;
 };
@@ -3656,40 +3606,12 @@ struct FRayTracingShaderBindingTableInitializer
 
 struct FRayTracingSceneInitializer
 {
-	// Unique list of geometries referenced by all instances in this scene.
-	// Any referenced geometry is kept alive while the scene is alive.
-	UE_DEPRECATED(5.5, "ReferencedGeometries should be provided in FRayTracingSceneBuildParams.")
-	TArray<TRefCountPtr<FRHIRayTracingGeometry>> ReferencedGeometries;
-
-	UE_DEPRECATED(5.5, "No longer necessary.")
-	TArray<FRHIRayTracingGeometry*> PerInstanceGeometries;
-
-	// Exclusive prefix sum of `Instance.NumTransforms` for all instances in this scene. Used to emulate SV_InstanceID in hit shaders.
-	UE_DEPRECATED(5.5, "Providing BaseInstancePrefixSum is no longer necessary.")
-	TArray<uint32> BaseInstancePrefixSum;
-	// Exclusive prefix sum of instance geometry segments is used to calculate SBT record address from instance and segment indices.
-	UE_DEPRECATED(5.5, "Providing SegmentPrefixSum is no longer necessary.")
-	TArray<uint32> SegmentPrefixSum;
-
-	// Total flattened number of ray tracing geometry instances (a single FRayTracingGeometryInstance may represent many) per layer.
-	UE_DEPRECATED(5.5, "FRHIRayTracingScene layers are deprecated. Use MaxNumInstances and create one FRHIRayTracingScene per layer instead.")
-	TArray<uint32> NumNativeInstancesPerLayer;
-
 	// Maximum number of instances in this scene. Actual number of instances is specified in FRayTracingSceneBuildParams.
 	uint32 MaxNumInstances = 0;
 
 	UE_DEPRECATED(5.6, "Use FRayTracingShaderBindingTableInitializer instead.")
 	uint32 NumTotalSegments = 0;
 
-	UE_DEPRECATED(5.5, "Use FRayTracingShaderBindingTableInitializer instead.")
-	uint32 ShaderSlotsPerGeometrySegment = 1;
-
-	UE_DEPRECATED(5.5, "Use FRayTracingShaderBindingTableInitializer instead.")
-	uint32 NumCallableShaderSlots = 0;
-
-	UE_DEPRECATED(5.5, "Use FRayTracingShaderBindingTableInitializer instead.")
-	uint32 NumMissShaderSlots = 1;
-
 	// Defines whether data in this scene should persist between frames.
 	// Currently only single-frame lifetime is supported.
 	ERayTracingSceneLifetime Lifetime = RTSL_SingleFrame;
@@ -3706,8 +3628,6 @@ struct FRayTracingSceneInitializer
 	PRAGMA_ENABLE_DEPRECATION_WARNINGS
 };
 
-using FRayTracingSceneInitializer2 UE_DEPRECATED(5.5, "Use FRayTracingSceneInitializer instead") = FRayTracingSceneInitializer;
-
 struct FRayTracingAccelerationStructureSize
 {
 	uint64 ResultSize = 0;
@@ -3774,22 +3694,6 @@ class FRHIRayTracingScene
 {
 public:
 	virtual const FRayTracingSceneInitializer& GetInitializer() const = 0;
-
-	UE_DEPRECATED(5.5, "Use GetOrCreateInlineBufferSRV on the FRHIShaderBindingTable instead of the Scene")
-	virtual FRHIShaderResourceView* GetOrCreateMetadataBufferSRV(FRHICommandListImmediate& RHICmdList)
-	{
-		return nullptr;
-	}
-
-	UE_DEPRECATED(5.5, "FRHIRayTracingScene layers are deprecated. Create one FRHIRayTracingScene per layer instead.")
-	virtual uint32 GetLayerBufferOffset(uint32 LayerIndex) const
-	{
-		checkf(LayerIndex == 0, TEXT("FRHIRayTracingScene layers are deprecated. Create one FRHIRayTracingScene per layer instead."));
-		return 0;
-	}
-
-	UE_DEPRECATED(5.5, "Create standalone FRHIShaderBindingTable instead.")
-	virtual FRHIShaderBindingTable* FindOrCreateShaderBindingTable(const FRHIRayTracingPipelineState* Pipeline) = 0;
 };
 
 class FRHIShaderBindingTable
@@ -4255,15 +4159,7 @@ public:
 	// @return	true if native Present should be also be performed; false otherwise. If it returns
 	// true, then InOutSyncInterval could be modified to switch between VSync/NoVSync for the normal 
 	// Present.  Must match value previously returned by NeedsNativePresent for this frame.
-	virtual bool Present(IRHICommandContext& RHICmdContext, int32& InOutSyncInterval)
-	{
-	 	PRAGMA_DISABLE_DEPRECATION_WARNINGS
-	 	return Present(InOutSyncInterval); 
-		PRAGMA_ENABLE_DEPRECATION_WARNINGS
-	 };
-
-	UE_DEPRECATED(5.5, "Please replace with Present(IRHICommandContext& RHICmdContext, int32& InOutSyncInterval).")
-	virtual bool Present(int32& InOutSyncInterval) { check(false); return true; };
+	virtual bool Present(IRHICommandContext& RHICmdContext, int32& InOutSyncInterval) { check(false); return true; };
 
 	// Called from RHI thread after native Present has been called
 	virtual void PostPresent() {};
@@ -4929,9 +4825,6 @@ public:
 	uint32 MaxAttributeSizeInBytes = 8; // sizeof FRayTracingIntersectionAttributes declared in RayTracingCommon.ush
 	uint32 MaxPayloadSizeInBytes = 24; // sizeof FDefaultPayload declared in RayTracingCommon.ush
 
-	UE_DEPRECATED(5.5, "Set bAllowHitGroupIndexing in FRayTracingShaderBindingTableInitializer.")
-	bool bAllowHitGroupIndexing = true;
-
 	// NOTE: GetTypeHash(const FRayTracingPipelineStateInitializer& Initializer) should also be updated when changing this function
 	bool operator==(const FRayTracingPipelineStateSignature& rhs) const
 	{
@@ -4940,10 +4833,7 @@ public:
 			&& RayGenHash == rhs.RayGenHash
 			&& MissHash == rhs.MissHash
 			&& HitGroupHash == rhs.HitGroupHash
-			&& CallableHash == rhs.CallableHash
-		PRAGMA_DISABLE_DEPRECATION_WARNINGS
-			&& bAllowHitGroupIndexing == rhs.bAllowHitGroupIndexing;
-		PRAGMA_ENABLE_DEPRECATION_WARNINGS
+			&& CallableHash == rhs.CallableHash;
 	}
 
 	friend uint32 GetTypeHash(const FRayTracingPipelineStateSignature& Initializer)
@@ -4953,10 +4843,7 @@ public:
 			GetTypeHash(Initializer.GetRayGenHash()) ^
 			GetTypeHash(Initializer.GetRayMissHash()) ^
 			GetTypeHash(Initializer.GetHitGroupHash()) ^
-			GetTypeHash(Initializer.GetCallableHash()) ^
-		PRAGMA_DISABLE_DEPRECATION_WARNINGS
-			GetTypeHash(Initializer.bAllowHitGroupIndexing);
-		PRAGMA_ENABLE_DEPRECATION_WARNINGS
+			GetTypeHash(Initializer.GetCallableHash());
 	}
 
 	uint64 GetHitGroupHash() const { return HitGroupHash; }
@@ -4964,15 +4851,6 @@ public:
 	uint64 GetRayMissHash()  const { return MissHash; }
 	uint64 GetCallableHash() const { return CallableHash; }
 
-	PRAGMA_DISABLE_DEPRECATION_WARNINGS
-	FRayTracingPipelineStateSignature() = default;
-	FRayTracingPipelineStateSignature(const FRayTracingPipelineStateSignature&) = default;
-	FRayTracingPipelineStateSignature& operator=(const FRayTracingPipelineStateSignature&) = default;
-	FRayTracingPipelineStateSignature(FRayTracingPipelineStateSignature&&) = default;
-	FRayTracingPipelineStateSignature& operator=(FRayTracingPipelineStateSignature&&) = default;
-	~FRayTracingPipelineStateSignature() = default;
-	PRAGMA_ENABLE_DEPRECATION_WARNINGS
-
 protected:
 
 	uint64 RayGenHash = 0;
diff --git a/Engine/Source/Runtime/RHI/Public/RHIValidationContext.h b/Engine/Source/Runtime/RHI/Public/RHIValidationContext.h
--- a/Engine/Source/Runtime/RHI/Public/RHIValidationContext.h
+++ b/Engine/Source/Runtime/RHI/Public/RHIValidationContext.h
@@ -1217,16 +1217,6 @@ public:
 		RHIContext->RHICopyBufferRegion(DestBuffer, DstOffset, SourceBuffer, SrcOffset, NumBytes);
 	}
 
-	void RHIClearRayTracingBindings(FRHIRayTracingScene* Scene)
-	{
-		RHIContext->RHIClearRayTracingBindings(Scene);
-	}
-	
-	void RHICommitRayTracingBindings(FRHIRayTracingScene* Scene)
-	{
-		RHIContext->RHICommitRayTracingBindings(Scene);
-	}
-
 	void RHIClearShaderBindingTable(FRHIShaderBindingTable* SBT)
 	{
 		SBT->Clear();
diff --git a/Engine/Source/Runtime/Renderer/Private/MeshPassProcessor.cpp b/Engine/Source/Runtime/Renderer/Private/MeshPassProcessor.cpp
--- a/Engine/Source/Runtime/Renderer/Private/MeshPassProcessor.cpp
+++ b/Engine/Source/Runtime/Renderer/Private/MeshPassProcessor.cpp
@@ -368,23 +368,6 @@ FRayTracingLocalShaderBindings* FMeshDrawShaderBindings::SetRayTracingShaderBind
 	return &Bindings;
 }
 
-FRayTracingLocalShaderBindings* FMeshDrawShaderBindings::SetRayTracingShaderBindingsForHitGroup(
-	FRayTracingLocalShaderBindingWriter* BindingWriter,
-	uint32 InstanceIndex, 
-	uint32 SegmentIndex,
-	uint32 HitGroupIndexInPipeline,
-	uint32 ShaderSlot) const
-{
-	FRayTracingLocalShaderBindings* Bindings = SetRayTracingShaderBindings(BindingWriter, HitGroupIndexInPipeline, INDEX_NONE, 0, ERayTracingLocalShaderBindingType::Transient);
-	PRAGMA_DISABLE_DEPRECATION_WARNINGS
-	Bindings->InstanceIndex = InstanceIndex;
-	Bindings->ShaderSlot = ShaderSlot;
-	PRAGMA_ENABLE_DEPRECATION_WARNINGS
-	Bindings->SegmentIndex = SegmentIndex;
-
-	return Bindings;
-}
-
 FRayTracingLocalShaderBindings* FMeshDrawShaderBindings::SetRayTracingShaderBindingsForHitGroup(
 	FRayTracingLocalShaderBindingWriter* BindingWriter,
 	uint32 RecordIndex,
@@ -400,56 +383,6 @@ FRayTracingLocalShaderBindings* FMeshDrawShaderBindings::SetRayTracingShaderBind
 	return Bindings;
 }
 
-void FMeshDrawShaderBindings::SetRayTracingShaderBindingsForMissShader(
-	FRHICommandList& RHICmdList,
-	FRHIRayTracingScene* Scene,
-	FRayTracingPipelineState* PipelineState,
-	uint32 ShaderIndexInPipeline,
-	uint32 ShaderSlot) const
-{
-	check(ShaderLayouts.Num() == 1);
-
-	FReadOnlyMeshDrawSingleShaderBindings SingleShaderBindings(ShaderLayouts[0], GetData());
-
-	FShaderBindingState BindingState;
-
-	const int32 MaxUniformBuffers = UE_ARRAY_COUNT(BindingState.UniformBuffers);
-
-	FRHIUniformBuffer* const* RESTRICT UniformBufferBindings = SingleShaderBindings.GetUniformBufferStart();
-	const FShaderUniformBufferParameterInfo* RESTRICT UniformBufferParameters = SingleShaderBindings.ParameterMapInfo.UniformBuffers.GetData();
-	const int32 NumUniformBufferParameters = SingleShaderBindings.ParameterMapInfo.UniformBuffers.Num();
-
-	checkf(SingleShaderBindings.ParameterMapInfo.TextureSamplers.Num() == 0, TEXT("Texture sampler parameters are not supported for ray tracing. UniformBuffers must be used for all resource binding."));
-	checkf(SingleShaderBindings.ParameterMapInfo.SRVs.Num() == 0, TEXT("SRV parameters are not supported for ray tracing. UniformBuffers must be used for all resource binding."));
-
-	// Measure parameter memory requirements
-
-	int32 MaxUniformBufferUsed = -1;
-	for (int32 UniformBufferIndex = 0; UniformBufferIndex < NumUniformBufferParameters; UniformBufferIndex++)
-	{
-		FShaderUniformBufferParameterInfo Parameter = UniformBufferParameters[UniformBufferIndex];
-		checkSlow(Parameter.BaseIndex < MaxUniformBuffers);
-		FRHIUniformBuffer* UniformBuffer = UniformBufferBindings[UniformBufferIndex];
-		if (Parameter.BaseIndex < MaxUniformBuffers)
-		{
-			BindingState.UniformBuffers[Parameter.BaseIndex] = UniformBuffer;
-			MaxUniformBufferUsed = FMath::Max((int32)Parameter.BaseIndex, MaxUniformBufferUsed);
-		}
-	}
-
-	checkf(SingleShaderBindings.ParameterMapInfo.TextureSamplers.Num() == 0, TEXT("Texture sampler parameters are not supported for ray tracing. UniformBuffers must be used for all resource binding."));
-	checkf(SingleShaderBindings.ParameterMapInfo.SRVs.Num() == 0, TEXT("SRV parameters are not supported for ray tracing. UniformBuffers must be used for all resource binding."));
-	checkf(SingleShaderBindings.ParameterMapInfo.LooseParameterBuffers.Num() == 0, TEXT("Ray tracing miss shaders may not have loose parameters"));
-
-	uint32 NumUniformBuffersToSet = MaxUniformBufferUsed + 1;
-	const uint32 UserData = 0; // UserData could be used to store material ID or any other kind of per-material constant. This can be retrieved in hit shaders via GetHitGroupUserData().
-	PRAGMA_DISABLE_DEPRECATION_WARNINGS
-	RHICmdList.SetRayTracingMissShader(Scene, ShaderSlot, PipelineState, ShaderIndexInPipeline,
-		NumUniformBuffersToSet, BindingState.UniformBuffers,
-		UserData);
-	PRAGMA_ENABLE_DEPRECATION_WARNINGS
-}
-
 void FMeshDrawShaderBindings::SetRayTracingShaderBindingsForMissShader(
 	FRHICommandList& RHICmdList,
 	FRHIShaderBindingTable* SBT, uint32 RecordIndex,
diff --git a/Engine/Source/Runtime/Renderer/Private/RayTracing/RayTracingMeshDrawCommands.cpp b/Engine/Source/Runtime/Renderer/Private/RayTracing/RayTracingMeshDrawCommands.cpp
--- a/Engine/Source/Runtime/Renderer/Private/RayTracing/RayTracingMeshDrawCommands.cpp
+++ b/Engine/Source/Runtime/Renderer/Private/RayTracing/RayTracingMeshDrawCommands.cpp
@@ -34,22 +34,6 @@ void FDynamicRayTracingMeshCommandContext::FinalizeCommand(FRayTracingMeshComman
 			ShaderBindings.Add(ShaderBinding);
 		}
 	}
-
-	PRAGMA_DISABLE_DEPRECATION_WARNINGS
-		if (RayTracingInstanceIndex != INDEX_NONE)
-		{
-			const bool bHidden = RayTracingMeshCommand.bDecal;
-			FRayTracingShaderBindingData ShaderBinding(&RayTracingMeshCommand, RayTracingInstanceIndex, bHidden);
-			ShaderBindings.Add(ShaderBinding);
-		}
-
-	if (RayTracingDecalInstanceIndex != INDEX_NONE)
-	{
-		const bool bHidden = !RayTracingMeshCommand.bDecal;
-		FRayTracingShaderBindingData ShaderBinding(&RayTracingMeshCommand, RayTracingDecalInstanceIndex, bHidden);
-		ShaderBindings.Add(ShaderBinding);
-	}
-	PRAGMA_ENABLE_DEPRECATION_WARNINGS
 }
 
 void FRayTracingMeshCommand::SetRayTracingShaderBindingsForHitGroup(
@@ -84,39 +68,6 @@ void FRayTracingMeshCommand::SetRayTracingShaderBindingsForHitGroup(
 	}
 }
 
-void FRayTracingMeshCommand::SetRayTracingShaderBindingsForHitGroup(
-	FRayTracingLocalShaderBindingWriter* BindingWriter,
-	const TUniformBufferRef<FViewUniformShaderParameters>& ViewUniformBuffer,
-	FRHIUniformBuffer* SceneUniformBuffer,
-	FRHIUniformBuffer* NaniteUniformBuffer,
-	uint32 InstanceIndex,
-	uint32 SegmentIndex,
-	uint32 HitGroupIndexInPipeline,
-	uint32 ShaderSlot) const
-{
-	PRAGMA_DISABLE_DEPRECATION_WARNINGS
-	FRayTracingLocalShaderBindings* Bindings = ShaderBindings.SetRayTracingShaderBindingsForHitGroup(BindingWriter, InstanceIndex, SegmentIndex, HitGroupIndexInPipeline, ShaderSlot);
-	PRAGMA_ENABLE_DEPRECATION_WARNINGS
-
-	if (ViewUniformBufferParameter.IsBound())
-	{
-		check(ViewUniformBuffer);
-		Bindings->UniformBuffers[ViewUniformBufferParameter.GetBaseIndex()] = ViewUniformBuffer;
-	}
-
-	if (SceneUniformBufferParameter.IsBound())
-	{
-		check(SceneUniformBuffer);
-		Bindings->UniformBuffers[SceneUniformBufferParameter.GetBaseIndex()] = SceneUniformBuffer;
-	}
-
-	if (NaniteUniformBufferParameter.IsBound())
-	{
-		check(NaniteUniformBuffer);
-		Bindings->UniformBuffers[NaniteUniformBufferParameter.GetBaseIndex()] = NaniteUniformBuffer;
-	}
-}
-
 void FRayTracingMeshCommand::SetShader(const TShaderRef<FShader>& Shader)
 {
 	check(Shader.IsValid());
@@ -135,11 +86,6 @@ void FRayTracingMeshCommand::SetShader(const TShaderRef<FShader>& Shader)
 	}
 }
 
-void FRayTracingMeshCommand::SetShaders(const FMeshProcessorShaders& Shaders)
-{
-	SetShader(Shaders.RayTracingShader);
-}
-
 bool FRayTracingMeshCommand::IsUsingNaniteRayTracing() const
 {
 	return bNaniteRayTracing;
diff --git a/Engine/Source/Runtime/Renderer/Private/VariableRateShading/VariableRateShadingImageManager.cpp b/Engine/Source/Runtime/Renderer/Private/VariableRateShading/VariableRateShadingImageManager.cpp
--- a/Engine/Source/Runtime/Renderer/Private/VariableRateShading/VariableRateShadingImageManager.cpp
+++ b/Engine/Source/Runtime/Renderer/Private/VariableRateShading/VariableRateShadingImageManager.cpp
@@ -41,12 +41,6 @@ void CVarEnableVRSCallback(IConsoleVariable* Var)
 {
 	const int32 Value = Var->GetInt();
 
-	// Maintain deprecated globals
-	PRAGMA_DISABLE_DEPRECATION_WARNINGS
-	GRHIGlobals.VariableRateShading.Enabled = (Value != 0) ? GRHISupportsPipelineVariableRateShading : false;
-	GRHIGlobals.VariableRateShading.AttachmentEnabled = (Value != 0) ? GRHISupportsAttachmentVariableRateShading : false;
-	PRAGMA_ENABLE_DEPRECATION_WARNINGS
-
 	// If pipeline VRS is enabled, we need to update static meshes using per-material rates to reflect a change
 	if (GRHISupportsPipelineVariableRateShading)
 	{
diff --git a/Engine/Source/Runtime/Renderer/Public/MeshPassProcessor.h b/Engine/Source/Runtime/Renderer/Public/MeshPassProcessor.h
--- a/Engine/Source/Runtime/Renderer/Public/MeshPassProcessor.h
+++ b/Engine/Source/Runtime/Renderer/Public/MeshPassProcessor.h
@@ -1004,13 +1004,9 @@ public:
 
 #if RHI_RAYTRACING
 	RENDERER_API FRayTracingLocalShaderBindings* SetRayTracingShaderBindings(FRayTracingLocalShaderBindingWriter* BindingWriter, uint32 ShaderIndexInPipeline, uint32 RecordIndex, uint32 UserData, ERayTracingLocalShaderBindingType BindingType) const;
-	UE_DEPRECATED(5.5, "Use version with explicit FRHIRayTracingGeometry.")
-	RENDERER_API FRayTracingLocalShaderBindings* SetRayTracingShaderBindingsForHitGroup(FRayTracingLocalShaderBindingWriter* BindingWriter, uint32 InstanceIndex, uint32 SegmentIndex, uint32 HitGroupIndexInPipeline, uint32 ShaderSlot) const;
 	RENDERER_API FRayTracingLocalShaderBindings* SetRayTracingShaderBindingsForHitGroup(FRayTracingLocalShaderBindingWriter* BindingWriter, uint32 RecordIndex, const FRHIRayTracingGeometry* Geometry, uint32 GeometrySegmentIndex, uint32 HitGroupIndexInPipeline, ERayTracingLocalShaderBindingType BindingType) const;
 
 	// TODO: should these move to a binding writer too? should we introduce a different class to do these bindings since they aren't mesh related? rename this class entirely?
-	UE_DEPRECATED(5.5, "Use version with explicit FRHIShaderBindingTable.")
-	void SetRayTracingShaderBindingsForMissShader(FRHICommandList& RHICmdList, FRHIRayTracingScene* Scene, FRayTracingPipelineState* Pipeline, uint32 ShaderIndexInPipeline, uint32 ShaderSlot) const;
 	void SetRayTracingShaderBindingsForMissShader(FRHICommandList& RHICmdList, FRHIShaderBindingTable* SBT, uint32 RecordIndex, FRayTracingPipelineState* Pipeline, uint32 ShaderIndexInPipeline) const;
 #endif // RHI_RAYTRACING
 
diff --git a/Engine/Source/Runtime/Renderer/Public/RayTracingMeshDrawCommands.h b/Engine/Source/Runtime/Renderer/Public/RayTracingMeshDrawCommands.h
--- a/Engine/Source/Runtime/Renderer/Public/RayTracingMeshDrawCommands.h
+++ b/Engine/Source/Runtime/Renderer/Public/RayTracingMeshDrawCommands.h
@@ -96,23 +96,9 @@ public:
 		uint32 HitGroupIndexInPipeline,
 		ERayTracingLocalShaderBindingType BindingType) const;
 
-	UE_DEPRECATED(5.5, "Provide RayTracingGeometry and SBTSegmentIndex instead of InstanceIndex")
-	RENDERER_API void SetRayTracingShaderBindingsForHitGroup(
-		FRayTracingLocalShaderBindingWriter* BindingWriter,
-		const TUniformBufferRef<FViewUniformShaderParameters>& ViewUniformBuffer,
-		FRHIUniformBuffer* SceneUniformBuffer,
-		FRHIUniformBuffer* NaniteUniformBuffer,
-		uint32 InstanceIndex,
-		uint32 SegmentIndex,
-		uint32 HitGroupIndexInPipeline,
-		uint32 ShaderSlot) const;
-
 	/** Sets ray hit group shaders on the mesh command and allocates room for the shader bindings. */
 	RENDERER_API void SetShader(const TShaderRef<FShader>& Shader);
 
-	UE_DEPRECATED(5.4, "Use SetShader")
-	RENDERER_API void SetShaders(const FMeshProcessorShaders& Shaders);
-
 	RENDERER_API bool IsUsingNaniteRayTracing() const;
 
 	RENDERER_API void UpdateFlags(FRayTracingCachedMeshCommandFlags& Flags) const;
@@ -145,32 +131,15 @@ public:
 		, RayTracingGeometry(InRayTracingGeometry)
 		, SBTRecordIndex(InSBTRecordIndex)
 		, BindingType(InBindingType)
-		, InstanceIndex(INDEX_NONE)
 		, bHidden(bInHidden)
 	{
 		check(RayTracingGeometry != nullptr);
 	}
 
-	UE_DEPRECATED(5.5, "Provide RayTracingGeometry and SBTSegmentIndex instead of InstanceIndex")
-	FRayTracingShaderBindingData(const FRayTracingMeshCommand* InRayTracingMeshCommand, uint32 InInstanceIndex, bool bInHidden = false)
-		: RayTracingMeshCommand(InRayTracingMeshCommand)
-		, RayTracingGeometry(nullptr)
-		, SBTRecordIndex(INDEX_NONE)
-		, BindingType(ERayTracingLocalShaderBindingType::Transient)
-		, InstanceIndex(InInstanceIndex)
-		, bHidden(bInHidden)
-	{
-PRAGMA_DISABLE_DEPRECATION_WARNINGS
-		check(InstanceIndex != INDEX_NONE);
-PRAGMA_ENABLE_DEPRECATION_WARNINGS
-	}
-	
 	const FRayTracingMeshCommand* RayTracingMeshCommand;
 	const FRHIRayTracingGeometry* RayTracingGeometry;
 	uint32 SBTRecordIndex;
 	ERayTracingLocalShaderBindingType BindingType;
-	UE_DEPRECATED(5.5, "Provide RayTracingGeometry and SBTSegmentIndex instead of InstanceIndex")
-	uint32 InstanceIndex;
 	bool bHidden;
 	
 	PRAGMA_DISABLE_DEPRECATION_WARNINGS
@@ -190,12 +159,6 @@ struct TUseBitwiseSwap<FRayTracingShaderBindingData>
 
 typedef TArray<FRayTracingShaderBindingData> FRayTracingShaderBindingDataOneFrameArray;
 
-UE_DEPRECATED(5.5, "Use FRayTracingShaderBindingData instead of FVisibleRayTracingMeshCommand");
-using FVisibleRayTracingMeshCommand = FRayTracingShaderBindingData;
-
-UE_DEPRECATED(5.5, "Use FRayTracingShaderBindingDataOneFrameArray instead of FRayTracingMeshCommandOneFrameArray");
-using FRayTracingMeshCommandOneFrameArray = FRayTracingShaderBindingDataOneFrameArray;
-
 class FRayTracingMeshCommandContext
 {
 public:
@@ -248,27 +211,7 @@ public:
 		ShaderBindings(InShaderBindings),
 		RayTracingGeometry(InRayTracingGeometry),
 		GeometrySegmentIndex(InGeometrySegmentIndex),
-		SBTAllocation(InSBTAllocation),
-		RayTracingInstanceIndex(INDEX_NONE),
-		RayTracingDecalInstanceIndex(INDEX_NONE)
-	{}
-
-	UE_DEPRECATED(5.5, "Provide RayTracingGeometry and GlobalGeometrySegmentIndex and DecalGlobalGeometrySegmentIndex instead")
-	FDynamicRayTracingMeshCommandContext
-	(
-		FDynamicRayTracingMeshCommandStorage& InDynamicCommandStorage,
-		FRayTracingShaderBindingDataOneFrameArray& InShaderBindings,
-		uint32 InGeometrySegmentIndex,
-		uint32 InRayTracingInstanceIndex,
-		uint32 InRayTracingDecalInstanceIndex = INDEX_NONE
-	) :
-		DynamicCommandStorage(InDynamicCommandStorage),
-		ShaderBindings(InShaderBindings),
-		RayTracingGeometry(nullptr),
-		GeometrySegmentIndex(InGeometrySegmentIndex),
-		SBTAllocation(nullptr),
-		RayTracingInstanceIndex(InRayTracingInstanceIndex),
-		RayTracingDecalInstanceIndex(InRayTracingDecalInstanceIndex)
+		SBTAllocation(InSBTAllocation)
 	{}
 
 	virtual FRayTracingMeshCommand& AddCommand(const FRayTracingMeshCommand& Initializer) override final
@@ -289,11 +232,6 @@ private:
 	uint32 GeometrySegmentIndex;
 
 	FRayTracingSBTAllocation* SBTAllocation;
-
-	UE_DEPRECATED(5.5, "Provide RayTracingGeometry and SBTSegmentRange instead of RayTracingInstanceIndex")
-	uint32 RayTracingInstanceIndex;
-	UE_DEPRECATED(5.5, "Provide RayTracingGeometry and DecalGlobalGeometrySegmentIndex instead of RayTracingDecalInstanceIndex")
-	uint32 RayTracingDecalInstanceIndex;
 };
 
 class FRayTracingShaderCommand
diff --git a/Engine/Source/Runtime/VulkanRHI/Private/VulkanContext.h b/Engine/Source/Runtime/VulkanRHI/Private/VulkanContext.h
--- a/Engine/Source/Runtime/VulkanRHI/Private/VulkanContext.h
+++ b/Engine/Source/Runtime/VulkanRHI/Private/VulkanContext.h
@@ -327,8 +327,6 @@ public:
 	virtual void RHIBeginParallelRenderPass(TSharedPtr<FRHIParallelRenderPassInfo> InInfo, const TCHAR* InName) final override;
 	virtual void RHIEndParallelRenderPass() final override;
 
-	virtual void RHIClearRayTracingBindings(FRHIRayTracingScene* Scene) final override;
-	virtual void RHICommitRayTracingBindings(FRHIRayTracingScene* Scene) final override;
 	virtual void RHIClearShaderBindingTable(FRHIShaderBindingTable* SBT) final override;
 	virtual void RHICommitShaderBindingTable(FRHIShaderBindingTable* SBT, FRHIBuffer* InlineBindingDataBuffer) final override;
 	virtual void RHIBindAccelerationStructureMemory(FRHIRayTracingScene* Scene, FRHIBuffer* Buffer, uint32 BufferOffset) final override;
diff --git a/Engine/Source/Runtime/VulkanRHI/Private/VulkanRayTracing.cpp b/Engine/Source/Runtime/VulkanRHI/Private/VulkanRayTracing.cpp
--- a/Engine/Source/Runtime/VulkanRHI/Private/VulkanRayTracing.cpp
+++ b/Engine/Source/Runtime/VulkanRHI/Private/VulkanRayTracing.cpp
@@ -706,11 +706,6 @@ FVulkanRayTracingScene::FVulkanRayTracingScene(FRayTracingSceneInitializer InIni
 
 FVulkanRayTracingScene::~FVulkanRayTracingScene()
 {
-	for (auto& Item : ShaderTables)
-	{
-		delete Item.Value;
-	}
-
 	if (AccelerationStructureBuffer)
 	{
 		DEC_MEMORY_STAT_BY(STAT_VulkanRayTracingUsedVideoMemory, AccelerationStructureBuffer->GetSize());
@@ -1089,38 +1084,6 @@ void FVulkanRayTracingShaderTable::Commit(FVulkanCommandListContext& Context, FR
 	VulkanRHI::vkCmdPipelineBarrier(CommandBuffer.GetHandle(), VK_PIPELINE_STAGE_TRANSFER_BIT, VK_PIPELINE_STAGE_RAY_TRACING_SHADER_BIT_KHR, 0, 1, &BarrierAfter, 0, nullptr, 0, nullptr);
 }
 
-FRHIShaderBindingTable* FVulkanRayTracingScene::FindOrCreateShaderBindingTable(const FRHIRayTracingPipelineState* InPipeline)
-{
-	UE::TScopeLock Lock(Mutex);
-
-	const FVulkanRayTracingPipelineState* Pipeline = ResourceCast(InPipeline);
-
-	// Find existing table
-	{
-		TRefCountPtr<FVulkanRayTracingShaderTable>* FoundShaderTable = ShaderTables.Find(Pipeline);
-		if (FoundShaderTable)
-		{
-			return FoundShaderTable->GetReference();
-		}
-	}
-
-	PRAGMA_DISABLE_DEPRECATION_WARNINGS
-	FRayTracingShaderBindingTableInitializer SBTInitializer;
-	SBTInitializer.NumGeometrySegments = Initializer.NumTotalSegments;
-	SBTInitializer.NumShaderSlotsPerGeometrySegment = Initializer.ShaderSlotsPerGeometrySegment;
-	SBTInitializer.NumCallableShaderSlots = Initializer.NumCallableShaderSlots;
-	SBTInitializer.NumMissShaderSlots = Initializer.NumMissShaderSlots;
-	SBTInitializer.HitGroupIndexingMode = Pipeline->bAllowHitGroupIndexing ? ERayTracingHitGroupIndexingMode::Allow : ERayTracingHitGroupIndexingMode::Disallow;
-	SBTInitializer.ShaderBindingMode = ERayTracingShaderBindingMode::RTPSO;
-	PRAGMA_ENABLE_DEPRECATION_WARNINGS
-
-	// Create new table (use FRHICommandListExecutor::GetImmediateCommandList() directly for now this is deprecated code)
-	FVulkanRayTracingShaderTable* CreatedShaderTable = new FVulkanRayTracingShaderTable(FRHICommandListExecutor::GetImmediateCommandList(), Device, MoveTemp(SBTInitializer));
-	ShaderTables.Add(Pipeline, CreatedShaderTable);
-
-	return CreatedShaderTable;
-}
-
 FRayTracingAccelerationStructureSize FVulkanDynamicRHI::RHICalcRayTracingSceneSize(const FRayTracingSceneInitializer& Initializer)
 {
 	FVkRtTLASBuildData BuildData;
@@ -1177,23 +1140,6 @@ FShaderBindingTableRHIRef FVulkanDynamicRHI::RHICreateShaderBindingTable(FRHICom
 	return new FVulkanRayTracingShaderTable(RHICmdList, GetDevice(), Initializer);
 }
 
-void FVulkanCommandListContext::RHIClearRayTracingBindings(FRHIRayTracingScene* InScene)
-{
-	FVulkanRayTracingScene* Scene = ResourceCast(InScene);
-	check(Scene);
-
-	for (auto& Table : Scene->ShaderTables)
-	{
-		Table.Value->ReleaseLocalBuffers();
-	}
-}
-
-void FVulkanCommandListContext::RHICommitRayTracingBindings(FRHIRayTracingScene* InScene)
-{
-	FVulkanRayTracingScene* Scene = ResourceCast(InScene);
-	Scene->CommitShaderTables(*this);
-}
-
 void FVulkanCommandListContext::RHIClearShaderBindingTable(FRHIShaderBindingTable* InSBT)
 {
 	FVulkanRayTracingShaderTable* SBT = ResourceCast(InSBT);
@@ -1719,11 +1665,6 @@ FVulkanRayTracingPipelineState::FVulkanRayTracingPipelineState(FVulkanDevice* co
 		Callable.ShaderHandles = FetchShaderHandles(Pipeline, InitializerCallableShaders.Num());
 	}
 
-	// If no custom hit groups were provided, then disable SBT indexing and force default shader on all primitives
-	PRAGMA_DISABLE_DEPRECATION_WARNINGS
-	bAllowHitGroupIndexing = Initializer.GetHitGroupTable().Num() ? Initializer.bAllowHitGroupIndexing : false;
-	PRAGMA_ENABLE_DEPRECATION_WARNINGS
-
 	INC_DWORD_STAT(STAT_VulkanRayTracingCreatedPipelines);
 }
 
diff --git a/Engine/Source/Runtime/VulkanRHI/Private/VulkanRayTracing.h b/Engine/Source/Runtime/VulkanRHI/Private/VulkanRayTracing.h
--- a/Engine/Source/Runtime/VulkanRHI/Private/VulkanRayTracing.h
+++ b/Engine/Source/Runtime/VulkanRHI/Private/VulkanRayTracing.h
@@ -223,16 +223,6 @@ public:
 
 	void BindBuffer(FRHIBuffer* InBuffer, uint32 InBufferOffset);
 
-	void CommitShaderTables(FVulkanCommandListContext& Context)
-	{
-		for (auto& Pair : ShaderTables)
-		{
-			Pair.Value->Commit(Context, nullptr);
-		}
-	}
-
-	FRHIShaderBindingTable* FindOrCreateShaderBindingTable(const FRHIRayTracingPipelineState* Pipeline);
-
 	inline bool IsBuilt() const
 	{
 		return bBuilt;
@@ -260,8 +250,6 @@ public:
 	
 	TRefCountPtr<FVulkanBuffer> AccelerationStructureBuffer;
 	
-	TMap<const FVulkanRayTracingPipelineState*, TRefCountPtr<FVulkanRayTracingShaderTable>> ShaderTables;
-
 	bool bBuilt = false;
 
 private:
@@ -312,9 +300,6 @@ private:
 	const bool bIsPartialPipeline;
 
 public:
-	UE_DEPRECATED(5.5, "bAllowHitGroupIndexing is now stored in the ShaderBindingTable.")
-	bool bAllowHitGroupIndexing = true;
-
 	VulkanResourceFrameCounter FrameCounter;
 
 	friend FVulkanCommandListContext;
