VulkanRHI: Add different message if Vulkan is supported but not at the required feature level.

#rnx

diff --git a/Engine/Source/Runtime/RHI/Private/Linux/LinuxDynamicRHI.cpp b/Engine/Source/Runtime/RHI/Private/Linux/LinuxDynamicRHI.cpp
--- a/Engine/Source/Runtime/RHI/Private/Linux/LinuxDynamicRHI.cpp
+++ b/Engine/Source/Runtime/RHI/Private/Linux/LinuxDynamicRHI.cpp
@@ -54,6 +54,7 @@ FDynamicRHI* PlatformCreateDynamicRHI()
 	}
 
 	bool bVulkanFailed = false;
+	bool bVulkanSMFailed = false;
 	bool bOpenGLFailed = false;
 
 	IDynamicRHIModule* DynamicRHIModule = nullptr;
@@ -113,6 +114,7 @@ FDynamicRHI* PlatformCreateDynamicRHI()
 							UE_LOG(LogRHI, Display, TEXT("Skipping %s..."), *ShaderFormatName.ToString());
 
 							bVulkanFailed = true;
+							bVulkanSMFailed = true;
 						}
 					}
 
@@ -171,7 +173,14 @@ FDynamicRHI* PlatformCreateDynamicRHI()
 		{
 			if (bVulkanFailed)
 			{
-				FMessageDialog::Open(EAppMsgType::Ok, NSLOCTEXT("LinuxDynamicRHI", "RequiredVulkan", "Vulkan Driver is required to run the engine."));
+				if (bVulkanSMFailed)
+				{
+					FMessageDialog::Open(EAppMsgType::Ok, NSLOCTEXT("LinuxDynamicRHI", "RequiredVulkan", "Vulkan device could not be created at the project's specified feature levels (see log for details)."));
+				}
+				else
+				{
+					FMessageDialog::Open(EAppMsgType::Ok, NSLOCTEXT("LinuxDynamicRHI", "RequiredVulkan", "Vulkan Driver is required to run the engine."));
+				}
 			}
 			else
 			{
@@ -207,7 +216,14 @@ FDynamicRHI* PlatformCreateDynamicRHI()
 			{
 				if (bVulkanFailed)
 				{
-					FMessageDialog::Open(EAppMsgType::Ok, NSLOCTEXT("LinuxDynamicRHI", "NoVulkanDriver", "Failed to load Vulkan Driver which is required to run the engine.\nThe engine no longer fallbacks to OpenGL4 which has been deprecated."));
+					if (bVulkanSMFailed)
+					{
+						FMessageDialog::Open(EAppMsgType::Ok, NSLOCTEXT("LinuxDynamicRHI", "RequiredVulkan", "Vulkan device could not be created at the project's specified feature levels (see log for details)."));
+					}
+					else
+					{
+						FMessageDialog::Open(EAppMsgType::Ok, NSLOCTEXT("LinuxDynamicRHI", "NoVulkanDriver", "Failed to load Vulkan Driver which is required to run the engine.\nThe engine no longer fallbacks to OpenGL4 which has been deprecated."));
+					}
 				}
 				else if (bOpenGLFailed)
 				{
