RHI: Add pass through in validation context for RHINextSubpass and RHISetUniformBufferDynamicOffset.

#rnx

diff --git a/Engine/Source/Runtime/RHI/Public/RHIValidationContext.h b/Engine/Source/Runtime/RHI/Public/RHIValidationContext.h
--- a/Engine/Source/Runtime/RHI/Public/RHIValidationContext.h
+++ b/Engine/Source/Runtime/RHI/Public/RHIValidationContext.h
@@ -924,6 +924,11 @@ public:
 		RHIContext->RHISetStaticUniformBuffer(Slot, UniformBuffer);
 	}
 
+	virtual void RHISetUniformBufferDynamicOffset(FUniformBufferStaticSlot Slot, uint32 Offset) override final
+	{
+		RHIContext->RHISetUniformBufferDynamicOffset(Slot, Offset);
+	}
+
 	virtual void RHISetStencilRef(uint32 StencilRef) override final
 	{
 		//checkf(State.bGfxPSOSet, TEXT("A Graphics PSO has to be set to change stencil ref!"));
@@ -1171,6 +1176,11 @@ public:
 		State.PreviousRenderPassName = State.RenderPassName;
 	}
 
+	virtual void RHINextSubpass() override final
+	{
+		RHIContext->RHINextSubpass();
+	}
+
 	virtual void RHIBeginParallelRenderPass(TSharedPtr<FRHIParallelRenderPassInfo> InInfo, const TCHAR* InName) override final
 	{
 		RHIContext->RHIBeginParallelRenderPass(InInfo, InName);
