VulkanRHI: Fix loc key for new Vulkan messagebox.

#rnx

diff --git a/Engine/Source/Runtime/RHI/Private/Linux/LinuxDynamicRHI.cpp b/Engine/Source/Runtime/RHI/Private/Linux/LinuxDynamicRHI.cpp
--- a/Engine/Source/Runtime/RHI/Private/Linux/LinuxDynamicRHI.cpp
+++ b/Engine/Source/Runtime/RHI/Private/Linux/LinuxDynamicRHI.cpp
@@ -175,7 +175,7 @@ FDynamicRHI* PlatformCreateDynamicRHI()
 			{
 				if (bVulkanSMFailed)
 				{
-					FMessageDialog::Open(EAppMsgType::Ok, NSLOCTEXT("LinuxDynamicRHI", "RequiredVulkan", "Vulkan device could not be created at the project's specified feature levels (see log for details)."));
+					FMessageDialog::Open(EAppMsgType::Ok, NSLOCTEXT("LinuxDynamicRHI", "ForcedVulkanSM", "Forced Vulkan device could not be created at the project's supported feature levels (see log for details)."));
 				}
 				else
 				{
@@ -218,7 +218,7 @@ FDynamicRHI* PlatformCreateDynamicRHI()
 				{
 					if (bVulkanSMFailed)
 					{
-						FMessageDialog::Open(EAppMsgType::Ok, NSLOCTEXT("LinuxDynamicRHI", "RequiredVulkan", "Vulkan device could not be created at the project's specified feature levels (see log for details)."));
+						FMessageDialog::Open(EAppMsgType::Ok, NSLOCTEXT("LinuxDynamicRHI", "SupportedVulkanSM", "Vulkan device could not be created at the project's supported feature levels (see log for details)."));
 					}
 					else
 					{
