Fix for "CurrentEndOfPipeFrame + 1 < NewEndOfPipeFrame" check in StatsData.cpp when the new GPU profiler is enabled.
- This was caused by duplicate frame indices being submitted to the stats system via EStatOperation::AdvanceFrameEventEndOfPipe.
- Replaced the StatsFrameRT global with a TOptional, which is consumed by the render thread when ending a frame (via ERHISubmitFlags::EndFrame). This prevents duplicates from happening.

diff --git a/Engine/Source/Runtime/Core/Private/Stats/StatsSystem.cpp b/Engine/Source/Runtime/Core/Private/Stats/StatsSystem.cpp
--- a/Engine/Source/Runtime/Core/Private/Stats/StatsSystem.cpp
+++ b/Engine/Source/Runtime/Core/Private/Stats/StatsSystem.cpp
@@ -22,6 +22,13 @@ DEFINE_STAT(STAT_SecondsPerCycle);
 
 namespace UE::Stats
 {
+#if STATS
+	// The engine stats system maintains its own frame counter. We need to propagate this to
+	// the end-of-pipe thread so that stat data it submits is attributed to the correct frame.
+	// This is done inside AdvanceRenderingThreadStats.
+	CORE_API TOptional<int64> FStats::StatsFrameRT;
+#endif
+
 	TAtomic<int32> FStats::GameThreadStatsFrame(1);
 
 	namespace Private
diff --git a/Engine/Source/Runtime/Core/Public/Stats/StatsSystem.h b/Engine/Source/Runtime/Core/Public/Stats/StatsSystem.h
--- a/Engine/Source/Runtime/Core/Public/Stats/StatsSystem.h
+++ b/Engine/Source/Runtime/Core/Public/Stats/StatsSystem.h
@@ -42,6 +42,13 @@ namespace UE::Stats
 
 		/** Current game thread stats frame. */
 		static CORE_API TAtomic<int32> GameThreadStatsFrame;
+
+	#if STATS
+		// The engine stats system maintains its own frame counter. We need to propagate this to
+		// the end-of-pipe thread so that stat data it submits is attributed to the correct frame.
+		// This is done inside AdvanceRenderingThreadStats.
+		static CORE_API TOptional<int64> StatsFrameRT;
+	#endif
 	};
 } // namespace UE::Stats
 
diff --git a/Engine/Source/Runtime/RHI/Private/GPUProfiler.cpp b/Engine/Source/Runtime/RHI/Private/GPUProfiler.cpp
--- a/Engine/Source/Runtime/RHI/Private/GPUProfiler.cpp
+++ b/Engine/Source/Runtime/RHI/Private/GPUProfiler.cpp
@@ -2137,10 +2137,6 @@ namespace UE::RHI::GPUProfiler
 		TMap<uint32, FFrameState> Frames;
 		TMap<FQueue, TUniquePtr<FQueueState>> QueueStates;
 
-	#if STATS
-		TOptional<int64> LastStatsFrame;
-	#endif
-
 		// Attempts to retrieve the CPU and GPU timestamps of when a fence wait is resolved by a signal on another queue.
 		TOptional<FResolvedWait> ResolveWait(FQueueState& LocalQueue, FEvent::FWaitFence const& WaitFenceEvent)
 		{
@@ -2290,8 +2286,9 @@ namespace UE::RHI::GPUProfiler
 						FFrameState& FrameState = Frames.FindOrAdd(FrameBoundary.FrameNumber);
 
 					#if STATS
-						check(!FrameState.StatsFrame.IsSet() || FrameState.StatsFrame.GetValue() == FrameBoundary.StatsFrame);
-						FrameState.StatsFrame = FrameBoundary.StatsFrame;
+						FrameState.StatsFrame = FrameBoundary.bStatsFrameSet
+							? FrameBoundary.StatsFrame
+							: TOptional<int64>();
 					#endif
 
 					#if WITH_PROFILEGPU
@@ -2340,13 +2337,10 @@ namespace UE::RHI::GPUProfiler
 		void ProcessFrame(FFrameState& FrameState)
 		{
 		#if STATS
-			check(FrameState.StatsFrame.IsSet());
 			FEndOfPipeStats* Stats = FEndOfPipeStats::Get();
-
-			if (!LastStatsFrame.IsSet() || *LastStatsFrame != *FrameState.StatsFrame)
+			if (FrameState.StatsFrame.IsSet())
 			{
 				Stats->AddMessage(FStatConstants::AdvanceFrame.GetEncodedName(), EStatOperation::AdvanceFrameEventEndOfPipe, *FrameState.StatsFrame);
-				LastStatsFrame = *FrameState.StatsFrame;
 			}
 		#endif
 
diff --git a/Engine/Source/Runtime/RHI/Private/RHICommandList.cpp b/Engine/Source/Runtime/RHI/Private/RHICommandList.cpp
--- a/Engine/Source/Runtime/RHI/Private/RHICommandList.cpp
+++ b/Engine/Source/Runtime/RHI/Private/RHICommandList.cpp
@@ -14,6 +14,7 @@
 #include "RHIContext.h"
 #include "RHIFwd.h"
 #include "RHITransition.h"
+#include "Stats/StatsSystem.h"
 #include "Stats/StatsTrace.h"
 #include "Stats/ThreadIdleStats.h"
 #include "HAL/PlatformMisc.h"
@@ -1439,12 +1440,21 @@ RHI_API FGraphEventRef FRHICommandListExecutor::Submit(TConstArrayView<FRHIComma
 
 	if (EnumHasAnyFlags(SubmitFlags, ERHISubmitFlags::SubmitToGPU))
 	{
+	#if STATS
+		TOptional<int64> LocalStatsFrame;
+		if (EnumHasAnyFlags(SubmitFlags, ERHISubmitFlags::EndFrame))
+		{
+			LocalStatsFrame = UE::Stats::FStats::StatsFrameRT;
+			UE::Stats::FStats::StatsFrameRT.Reset();
+		}
+	#endif
+
 		EnqueueDispatchTask({},
 			[
 				  this
 				, State = SubmitState
 			#if STATS
-				, LocalStatsFrame = StatsFrame
+				, LocalStatsFrame
 			#endif
 			]() mutable
 			{
diff --git a/Engine/Source/Runtime/RHI/Public/DynamicRHI.h b/Engine/Source/Runtime/RHI/Public/DynamicRHI.h
--- a/Engine/Source/Runtime/RHI/Public/DynamicRHI.h
+++ b/Engine/Source/Runtime/RHI/Public/DynamicRHI.h
@@ -243,7 +243,7 @@ public:
 		const TRHIPipelineArray<FRHIBreadcrumbNode*>& GPUBreadcrumbs;
 	#endif
 	#if STATS
-		int64 StatsFrame;
+		TOptional<int64> StatsFrame;
 	#endif
 	};
 	virtual void RHIEndFrame(const FRHIEndFrameArgs& Args) = 0;
diff --git a/Engine/Source/Runtime/RHI/Public/GPUProfiler.h b/Engine/Source/Runtime/RHI/Public/GPUProfiler.h
--- a/Engine/Source/Runtime/RHI/Public/GPUProfiler.h
+++ b/Engine/Source/Runtime/RHI/Public/GPUProfiler.h
@@ -98,15 +98,17 @@ namespace UE::RHI::GPUProfiler
 			// Very first frame of the engine is frame 0 (from boot to first call to RHIEndFrame).
 			uint32 FrameNumber;
 
+		#if STATS
+			// Should be TOptional<int64> but it is not trivially destructible
+			bool bStatsFrameSet;
+			int64 StatsFrame;
+		#endif
+
 		#if WITH_RHI_BREADCRUMBS
 			// The RHI breadcrumb currently at the top of the stack at the frame boundary.
 			FRHIBreadcrumbNode* Breadcrumb;
 		#endif
 
-		#if STATS
-			int64 StatsFrame;
-		#endif
-
 			FFrameBoundary(
 				  uint64 CPUTimestamp
 				, uint32 FrameNumber
@@ -114,17 +116,18 @@ namespace UE::RHI::GPUProfiler
 				, FRHIBreadcrumbNode* Breadcrumb
 			#endif
 			#if STATS
-				, int64 StatsFrame
+				, TOptional<int64> StatsFrame
 			#endif
 				)
 				: CPUTimestamp(CPUTimestamp)
 				, FrameNumber(FrameNumber)
+			#if STATS
+				, bStatsFrameSet(StatsFrame.IsSet())
+				, StatsFrame(StatsFrame.IsSet() ? *StatsFrame : 0)
+			#endif
 			#if WITH_RHI_BREADCRUMBS
 				, Breadcrumb(Breadcrumb)
 			#endif
-			#if STATS
-				, StatsFrame(StatsFrame)
-			#endif
 			{}
 		};
 
diff --git a/Engine/Source/Runtime/RHI/Public/RHICommandList.h b/Engine/Source/Runtime/RHI/Public/RHICommandList.h
--- a/Engine/Source/Runtime/RHI/Public/RHICommandList.h
+++ b/Engine/Source/Runtime/RHI/Public/RHICommandList.h
@@ -5458,7 +5458,7 @@ private:
 			TRHIPipelineArray<FRHIBreadcrumbNode*> GPUBreadcrumbs;
 		#endif
 		#if STATS
-			int64 StatsFrame;
+			TOptional<int64> StatsFrame;
 		#endif
 		};
 		void Submit(const FSubmitArgs& Args);
@@ -5497,19 +5497,6 @@ private:
 	// Counts the number of calls to RHIEndFrame, and is used in GPU profiler frame boundary events.
 	uint32 FrameNumber = 0;
 
-#if STATS
-	// The engine stats system maintains its own frame counter. We need to propagate this to
-	// the end-of-pipe thread so that stat data it submits is attributed to the correct frame.
-	// This is done inside AdvanceRenderingThreadStats.
-	int64 StatsFrame = 0;
-public:
-	void SetStatsFrame(int64 InStatsFrame)
-	{
-		StatsFrame = InStatsFrame;
-	}
-private:
-#endif
-
 	bool AllowParallel() const;
 };
 
diff --git a/Engine/Source/Runtime/RenderCore/Private/RenderingThread.cpp b/Engine/Source/Runtime/RenderCore/Private/RenderingThread.cpp
--- a/Engine/Source/Runtime/RenderCore/Private/RenderingThread.cpp
+++ b/Engine/Source/Runtime/RenderCore/Private/RenderingThread.cpp
@@ -21,6 +21,7 @@
 #include "RenderCommandFence.h"
 #include "RenderDeferredCleanup.h"
 #include "TickableObjectRenderThread.h"
+#include "Stats/StatsSystem.h"
 #include "Stats/StatsData.h"
 #include "HAL/ThreadHeartBeat.h"
 #include "RenderResource.h"
@@ -347,7 +348,7 @@ static void AdvanceRenderingThreadStats(int64 StatsFrame, int32 DisableChangeTag
 
 #if RHI_NEW_GPU_PROFILER
 	// Propagate the stats frame value down to the end-of-pipe thread.
-	GRHICommandList.SetStatsFrame(Frame);
+	UE::Stats::FStats::StatsFrameRT = Frame;
 #else
 	FThreadStats::AddMessage(FStatConstants::AdvanceFrame.GetEncodedName(), EStatOperation::AdvanceFrameEventEndOfPipe, Frame);
 #endif
