Removing TextureProfiler

diff --git a/Engine/Source/Runtime/D3D12RHI/Private/D3D12Texture.cpp b/Engine/Source/Runtime/D3D12RHI/Private/D3D12Texture.cpp
--- a/Engine/Source/Runtime/D3D12RHI/Private/D3D12Texture.cpp
+++ b/Engine/Source/Runtime/D3D12RHI/Private/D3D12Texture.cpp
@@ -6,7 +6,6 @@
 
 #include "D3D12Texture.h"
 #include "D3D12RHIPrivate.h"
-#include "TextureProfiler.h"
 #include "ProfilingDebugging/MemoryTrace.h"
 #include "RHIUtilities.h"
 #include "HAL/LowLevelMemStats.h"
@@ -73,17 +72,6 @@ static TStatId GetD3D12StatEnum(const FD3D12ResourceDesc& ResourceDesc)
 
 void FD3D12TextureStats::UpdateD3D12TextureStats(FD3D12Texture& Texture, const FD3D12ResourceDesc& ResourceDesc, const FRHITextureDesc& TextureDesc, uint64 TextureSize, bool bNewTexture, bool bAllocating)
 {
-#if TEXTURE_PROFILER_ENABLED
-	if (!bNewTexture
-		&& !Texture.ResourceLocation.IsTransient() 
-		&& !EnumHasAnyFlags(TextureDesc.Flags, ETextureCreateFlags::Virtual)
-		&& !Texture.ResourceLocation.IsAliased())
-	{
-		const uint64 SafeSize = bAllocating ? TextureSize : 0;
-		FTextureProfiler::Get()->UpdateTextureAllocation(&Texture, SafeSize, ResourceDesc.Alignment, 0);
-	}
-#endif
-
 	if (TextureSize == 0)
 	{
 		return;
@@ -165,15 +153,7 @@ void FD3D12TextureStats::D3D12TextureAllocated(FD3D12Texture& Texture)
 			const uint64 TextureSize = Texture.ResourceLocation.GetSize();
 			const bool bNewTexture = true;
 			UpdateD3D12TextureStats(Texture, ResourceDesc, TextureDesc, TextureSize, bNewTexture, true);
-
-#if TEXTURE_PROFILER_ENABLED
-			if (!Texture.ResourceLocation.IsAliased())
-			{
-				FTextureProfiler::Get()->AddTextureAllocation(&Texture, TextureSize, ResourceDesc.Alignment, 0);
-			}
-#endif
 		}
-
 	}
 }
 
@@ -193,13 +173,6 @@ void FD3D12TextureStats::D3D12TextureDeleted(FD3D12Texture& Texture)
 
 			const bool bNewTexture = false;
 			UpdateD3D12TextureStats(Texture, ResourceDesc, TextureDesc, TextureSize, bNewTexture, false);
-
-#if TEXTURE_PROFILER_ENABLED
-			if (!EnumHasAnyFlags(TextureDesc.Flags, ETextureCreateFlags::Virtual) && !Texture.ResourceLocation.IsAliased())
-			{
-				FTextureProfiler::Get()->RemoveTextureAllocation(&Texture);
-			}
-#endif
 		}
 	}
 }
@@ -2803,19 +2776,6 @@ void FD3D12DynamicRHI::RHIBindDebugLabelName(FRHICommandListBase& RHICmdList, FR
 
 	// Also set on RHI object
 	TextureRHI->SetName(Name);
-
-#if TEXTURE_PROFILER_ENABLED
-
-	FD3D12Texture* D3D12Texture = GetD3D12TextureFromRHITexture(TextureRHI);
-	
-	if (!EnumHasAnyFlags(TextureRHI->GetFlags(), TexCreate_Virtual)
-		&& !D3D12Texture->ResourceLocation.IsTransient()
-		&& !D3D12Texture->ResourceLocation.IsAliased())
-	{
-		FTextureProfiler::Get()->UpdateTextureName(TextureRHI);
-	}
-	
-#endif
 }
 
 FD3D12Texture* FD3D12DynamicRHI::CreateTextureFromResource(bool bTextureArray, bool bCubeTexture, EPixelFormat Format, ETextureCreateFlags TexCreateFlags, const FClearValueBinding& ClearValueBinding, ID3D12Resource* Resource)
diff --git a/Engine/Source/Runtime/RHI/Private/DynamicRHI.cpp b/Engine/Source/Runtime/RHI/Private/DynamicRHI.cpp
--- a/Engine/Source/Runtime/RHI/Private/DynamicRHI.cpp
+++ b/Engine/Source/Runtime/RHI/Private/DynamicRHI.cpp
@@ -17,7 +17,6 @@
 #include "PipelineStateCache.h"
 #include "RHI.h"
 #include "RHIFwd.h"
-#include "TextureProfiler.h"
 #include "DataDrivenShaderPlatformInfo.h"
 #include "RHICommandList.h"
 #include "RHIImmutableSamplerState.h"
@@ -347,10 +346,6 @@ void RHIInit(bool bHasEditorToken)
 				FGenericCrashContext::SetEngineData(TEXT("RHI.FeatureLevel"), FeatureLevelString);
 				FGenericCrashContext::SetEngineData(TEXT("RHI.GPUVendor"), RHIVendorIdToString());
 				FGenericCrashContext::SetEngineData(TEXT("RHI.DeviceId"), FString::Printf(TEXT("%04X"), GRHIDeviceId));
-
-#if TEXTURE_PROFILER_ENABLED
-				FTextureProfiler::Get()->Init();
-#endif
 			}
 #if PLATFORM_ALLOW_NULL_RHI
 			else
diff --git a/Engine/Source/Runtime/RHI/Private/RHIResources.cpp b/Engine/Source/Runtime/RHI/Private/RHIResources.cpp
--- a/Engine/Source/Runtime/RHI/Private/RHIResources.cpp
+++ b/Engine/Source/Runtime/RHI/Private/RHIResources.cpp
@@ -7,7 +7,6 @@
 #include "RHICommandList.h"
 #include "RHIUniformBufferLayoutInitializer.h"
 #include "Stats/Stats.h"
-#include "TextureProfiler.h"
 #include "RHIGlobals.h"
 
 UE::TConsumeAllMpmcQueue<FRHIResource*> PendingDeletes;
@@ -131,10 +130,6 @@ FRHITexture::FRHITexture(const FRHITextureCreateDesc& InDesc)
 void FRHITexture::SetName(FName InName)
 {
 	Name = InName;
-
-#if TEXTURE_PROFILER_ENABLED
-	FTextureProfiler::Get()->UpdateTextureName(this);
-#endif
 }
 
 const TCHAR* FRHIViewDesc::GetBufferTypeString(FRHIViewDesc::EBufferType BufferType)
diff --git a/Engine/Source/Runtime/RHI/Private/TextureProfiler.cpp b/Engine/Source/Runtime/RHI/Private/TextureProfiler.cpp
deleted file mode 100644
diff --git a/Engine/Source/Runtime/RHI/Public/TextureProfiler.h b/Engine/Source/Runtime/RHI/Public/TextureProfiler.h
--- a/Engine/Source/Runtime/RHI/Public/TextureProfiler.h
+++ b/Engine/Source/Runtime/RHI/Public/TextureProfiler.h
@@ -2,91 +2,16 @@
 
 #pragma once
 
+#include "Misc/CoreMiscDefines.h"
+
+UE_DEPRECATED_HEADER(5.7, "TextureProfiler.h is no longer used.")
+
+#if UE_ENABLE_INCLUDE_ORDER_DEPRECATED_IN_5_7
 #include "Misc/Build.h"
 #include "Templates/Function.h"
 #include "Misc/CoreMisc.h"
 #include "Containers/Map.h"
 #include "UObject/NameTypes.h"
+#endif
 
-#define TEXTURE_PROFILER_ENABLED (WITH_ENGINE && !UE_BUILD_SHIPPING)
-
-#if TEXTURE_PROFILER_ENABLED
-
-class FRHITexture;
-
-/**
-* FTextureProfiler class. This manages recording and reporting texture allocations in the RHI
-*/
-class FTextureProfiler
-{
-	static FTextureProfiler* Instance;
-public:
-	// Singleton interface
-	RHI_API static FTextureProfiler* Get();
-
-	RHI_API void Init();
-
-	RHI_API void DumpTextures(bool RenderTargets, bool CombineTextureNames, bool AsCSV, FOutputDevice& OutputDevice);
-
-	RHI_API void AddTextureAllocation(FRHITexture* UniqueTexturePtr, size_t Size, uint32 Alignment, size_t AllocationWaste);
-	RHI_API void UpdateTextureAllocation(FRHITexture* UniqueTexturePtr, size_t Size, uint32 Alignment, size_t AllocationWaste);
-	RHI_API void RemoveTextureAllocation(FRHITexture* UniqueTexturePtr);
-	RHI_API void UpdateTextureName(FRHITexture* UniqueTexturePtr);
-
-private:
-
-	FTextureProfiler() = default;
-	FTextureProfiler(const FTextureProfiler&) = delete;
-	FTextureProfiler(FTextureProfiler&&) = delete;
-
-	RHI_API void Update();
-
-	// Should only be used by FTextureDetails
-	RHI_API const char* GetTextureNameString(FName TextureName);
-
-	FCriticalSection TextureMapCS;
-
-	class FTextureDetails
-	{
-	public:
-		FTextureDetails() = default;
-		RHI_API FTextureDetails(FRHITexture* Texture, size_t InSize, uint32 InAlign, size_t InAllocationWaste);
-		RHI_API ~FTextureDetails();
-		
-		RHI_API void SetName(FName InTextureName);
-
-		RHI_API void ResetPeakSize();
-		RHI_API void SetValues(const FTextureDetails& Values);
-		FName GetTextureName() const { return TextureName; }
-		const char* GetTextureNameString() const { check(TextureNameString != nullptr);  return TextureNameString; }
-
-		RHI_API FTextureDetails& operator+=(const FTextureDetails& Other);
-		RHI_API FTextureDetails& operator-=(const FTextureDetails& Other);
-		FTextureDetails& operator=(const FTextureDetails& Other) = default;
-
-		size_t Size = 0;
-		size_t PeakSize = 0;
-		uint32 Align = 0;
-		size_t AllocationWaste = 0;
-		int Count = 0;
-		bool IsRenderTarget = false;
-
-	private:
-		FName TextureName;
-
-		// This memory is not owned, do not delete
-		const char* TextureNameString = nullptr;
-	};
-
-	// Allocated so that a reference to the string can be stored in FTextureDetails
-	TMap<FName, const char*> TextureNameStrings;
-	TMap<void*, FTextureDetails> TexturesMap;
-
-	// Keep track of the totals separately to reduce the cost of rounding error for sizes
-	FTextureDetails TotalTextureSize;
-	FTextureDetails TotalRenderTargetSize;
-	TMap<FName, FTextureDetails> CombinedTextureSizes;
-	TMap<FName, FTextureDetails> CombinedRenderTargetSizes;
-	friend class FTextureDetails;
-};
-#endif //TEXTURE_PROFILER_ENABLED
+#define TEXTURE_PROFILER_ENABLED 0
