RHI Texture and Buffer cleanup
- FRHIBuffer constructors should be protected to keep them internal to RHI.
- FRHIViewableResource constructor should set its own names.
- RHI Validation buffer constructors should take the CreateDesc as an argument and init barrier tracking off that instead of it being initialized after the object is created.

diff --git a/Engine/Source/Runtime/NullDrv/Public/NullRHI.h b/Engine/Source/Runtime/NullDrv/Public/NullRHI.h
--- a/Engine/Source/Runtime/NullDrv/Public/NullRHI.h
+++ b/Engine/Source/Runtime/NullDrv/Public/NullRHI.h
@@ -124,13 +124,22 @@ public:
 
 	}
 
+	class FNullBuffer : public FRHIBuffer
+	{
+	public:
+		FNullBuffer(const FRHIBufferCreateDesc& CreateDesc)
+			: FRHIBuffer(CreateDesc)
+		{
+		}
+	};
+
 	struct FNullBufferInitializer : public FRHIBufferInitializer
 	{
-		FNullBufferInitializer(FRHICommandListBase& RHICmdList, FRHIBuffer* Buffer, void* InWritableData, uint64 InWritableDataSize)
+		FNullBufferInitializer(FRHICommandListBase& RHICmdList, FNullBuffer* Buffer, void* InWritableData, uint64 InWritableDataSize)
 			: FRHIBufferInitializer(RHICmdList, Buffer, InWritableData, InWritableDataSize,
-				[Buffer = TRefCountPtr<FRHIBuffer>(Buffer)](FRHICommandListBase&) mutable
+				[Buffer = TRefCountPtr<FNullBuffer>(Buffer)](FRHICommandListBase&) mutable
 				{
-					return MoveTemp(Buffer);
+					return TRefCountPtr<FRHIBuffer>(MoveTemp(Buffer));
 				}
 			)
 		{
@@ -154,7 +163,7 @@ public:
 			WritableDataSize = CreateDesc.Size;
 		}
 
-		FRHIBuffer* Buffer = new FRHIBuffer(CreateDesc);
+		FNullBuffer* Buffer = new FNullBuffer(CreateDesc);
 
 		return FNullBufferInitializer(RHICmdList, Buffer, WritableData, WritableDataSize);
 	}
diff --git a/Engine/Source/Runtime/RHI/Private/RHIResources.cpp b/Engine/Source/Runtime/RHI/Private/RHIResources.cpp
--- a/Engine/Source/Runtime/RHI/Private/RHIResources.cpp
+++ b/Engine/Source/Runtime/RHI/Private/RHIResources.cpp
@@ -116,15 +116,28 @@ void FRHIResource::GatherResourcesToDelete(TArray<FRHIResource*>& OutResources,
 	});
 }
 
-FRHITexture::FRHITexture(const FRHITextureCreateDesc& InDesc)
-	: FRHIViewableResource(RRT_Texture, InDesc.InitialState)
+FRHIBuffer::FRHIBuffer(const FRHIBufferCreateDesc& CreateDesc)
+	: FRHIViewableResource(RRT_Buffer, CreateDesc.InitialState, CreateDesc.DebugName, CreateDesc.OwnerName)
 #if ENABLE_RHI_VALIDATION
-	, RHIValidation::FTextureResource(InDesc)
+	, RHIValidation::FBufferResource(CreateDesc)
 #endif
-	, TextureDesc(InDesc)
+	, Desc(CreateDesc)
 {
-	SetName(InDesc.DebugName);
-	SetOwnerName(InDesc.OwnerName);
+}
+
+FRHITexture::FRHITexture(const FRHITextureCreateDesc& CreateDesc)
+	: FRHIViewableResource(RRT_Texture, CreateDesc.InitialState, CreateDesc.DebugName, CreateDesc.OwnerName)
+#if ENABLE_RHI_VALIDATION
+	, RHIValidation::FTextureResource(CreateDesc)
+#endif
+	, TextureDesc(CreateDesc)
+{
+}
+
+FRHITexture::FRHITexture(ERHIResourceType InResourceType)
+	: FRHIViewableResource(InResourceType, ERHIAccess::Unknown, nullptr, NAME_None)
+{
+	check(InResourceType == RRT_TextureReference);
 }
 
 void FRHITexture::SetName(FName InName)
diff --git a/Engine/Source/Runtime/RHI/Private/RHIValidation.cpp b/Engine/Source/Runtime/RHI/Private/RHIValidation.cpp
--- a/Engine/Source/Runtime/RHI/Private/RHIValidation.cpp
+++ b/Engine/Source/Runtime/RHI/Private/RHIValidation.cpp
@@ -85,6 +85,16 @@ namespace RHIValidation
 		return Init.Strings;
 	}
 
+	FBufferResource::FBufferResource(const FRHIBufferCreateDesc& CreateDesc)
+	{
+		InitBarrierTracking(CreateDesc);
+	}
+
+	void FBufferResource::InitBarrierTracking(const FRHIBufferCreateDesc& CreateDesc)
+	{
+		FResource::InitBarrierTracking(1, 1, 1, CreateDesc.InitialState, CreateDesc.DebugName);
+	}
+
 	FTextureResource::FTextureResource(FRHITextureCreateDesc const& CreateDesc)
 		: FTextureResource()
 	{
@@ -3530,7 +3540,11 @@ FRHITransientBuffer* FValidationTransientResourceAllocator::CreateBuffer(const F
 
 	if (!RHIBuffer->IsBarrierTrackingInitialized())
 	{
-		RHIBuffer->InitBarrierTracking(ERHIAccess::Discard, InDebugName);
+		const FRHIBufferCreateDesc CreateDesc =
+			FRHIBufferCreateDesc::Create(InDebugName, InCreateInfo)
+			.SetInitialState(ERHIAccess::Discard);
+
+		RHIBuffer->InitBarrierTracking(CreateDesc);
 	}
 	else
 	{
diff --git a/Engine/Source/Runtime/RHI/Public/RHIResources.h b/Engine/Source/Runtime/RHI/Public/RHIResources.h
--- a/Engine/Source/Runtime/RHI/Public/RHIResources.h
+++ b/Engine/Source/Runtime/RHI/Public/RHIResources.h
@@ -1307,10 +1307,17 @@ public:
 #endif
 
 protected:
-	FRHIViewableResource(ERHIResourceType InResourceType, ERHIAccess InAccess)
+	FRHIViewableResource(ERHIResourceType InResourceType, ERHIAccess InAccess, const TCHAR* InName, FName InOwnerName)
 		: FRHIResource(InResourceType)
 		, TrackedAccess(InAccess)
 	{
+#if RHI_ENABLE_RESOURCE_INFO
+		if (InName)
+		{
+			Name = FName(InName);
+		}
+		SetOwnerName(InOwnerName);
+#endif
 	}
 
 	virtual void SetTrackedAccessFromContext(FRHITrackedAccess InTrackedAccess)
@@ -1549,8 +1556,8 @@ struct FRHIBufferCreateDesc : public FRHIBufferDesc
 	FRHIBufferCreateDesc& SetInitialState(ERHIAccess InInitialState)       { InitialState = InInitialState; return *this; }
 	FRHIBufferCreateDesc& DetermineInitialState()                          { if (InitialState == ERHIAccess::Unknown) InitialState = RHIGetDefaultResourceState(Usage, false); return *this; }
 	FRHIBufferCreateDesc& SetInitAction(ERHIBufferInitAction InInitAction) { InitAction = InInitAction;     return *this; }
-	FRHIBufferCreateDesc& SetClassName(const FName& InClassName)           { ClassName = InClassName;       return *this; }
-	FRHIBufferCreateDesc& SetOwnerName(const FName& InOwnerName)           { OwnerName = InOwnerName;       return *this; }
+	FRHIBufferCreateDesc& SetClassName(FName InClassName)                  { ClassName = InClassName;       return *this; }
+	FRHIBufferCreateDesc& SetOwnerName(FName InOwnerName)                  { OwnerName = InOwnerName;       return *this; }
 
 	FRHIBufferCreateDesc& SetInitActionNone()
 	{
@@ -1600,22 +1607,13 @@ class FRHIBuffer : public FRHIViewableResource
 	, public RHIValidation::FBufferResource
 #endif
 {
-public:
-	/** Initialization constructor. */
-	FRHIBuffer(const FRHIBufferCreateDesc& InCreateDesc)
-		: FRHIViewableResource(RRT_Buffer, InCreateDesc.InitialState)
-		, Desc(InCreateDesc)
-	{
-#if RHI_ENABLE_RESOURCE_INFO
-		if (InCreateDesc.DebugName)
-		{
-			Name = InCreateDesc.DebugName;
-		}
+protected:
+	FRHIBuffer() = delete;
 
-		SetOwnerName(InCreateDesc.OwnerName);
-#endif
-	}
+	/** Initialization constructor. Should only be called by platform RHI implementations. */
+	RHI_API FRHIBuffer(const FRHIBufferCreateDesc& CreateDesc);
 
+public:
 	const FRHIBufferDesc& GetDesc() const
 	{
 		return Desc;
@@ -1639,7 +1637,7 @@ public:
 		return Desc.Usage;
 	}
 
-	void SetName(const FName& InName)
+	void SetName(FName InName)
 	{
 		Name = InName;
 	}
@@ -2097,8 +2095,8 @@ struct FRHITextureCreateDesc : public FRHITextureDesc
 	FRHITextureCreateDesc& SetInitialState(ERHIAccess InInitialState)          { InitialState = InInitialState;            return *this; }
 	FRHITextureCreateDesc& SetInitAction(ERHITextureInitAction InInitAction)   { InitAction = InInitAction;                return *this; }
 	FRHITextureCreateDesc& SetGPUMask(FRHIGPUMask InGPUMask)                   { GPUMask = InGPUMask;                      return *this; }
-	FRHITextureCreateDesc& SetClassName(const FName& InClassName)			   { ClassName = InClassName;				   return *this; }
-	FRHITextureCreateDesc& SetOwnerName(const FName& InOwnerName)			   { OwnerName = InOwnerName;                  return *this; }
+	FRHITextureCreateDesc& SetClassName(FName InClassName)			           { ClassName = InClassName;				   return *this; }
+	FRHITextureCreateDesc& SetOwnerName(FName InOwnerName)                     { OwnerName = InOwnerName;                  return *this; }
 
 	FRHITextureCreateDesc& SetFastVRAMPercentage(float InFastVRAMPercentage)
 	{
@@ -2170,8 +2168,13 @@ class FRHITexture : public FRHIViewableResource
 #endif
 {
 protected:
+	FRHITexture() = delete;
+
 	/** Initialization constructor. Should only be called by platform RHI implementations. */
-	RHI_API FRHITexture(const FRHITextureCreateDesc& InDesc);
+	RHI_API FRHITexture(const FRHITextureCreateDesc& CreateDesc);
+
+	/** Constructor for texture references */
+	RHI_API FRHITexture(ERHIResourceType InResourceType);
 
 public:
 	/**
@@ -2375,14 +2378,7 @@ public:
 #endif
 
 private:
-
 	friend class FRHITextureReference;
-	/** Constructor for texture references */
-	FRHITexture(ERHIResourceType InResourceType)
-		: FRHIViewableResource(InResourceType, ERHIAccess::Unknown)
-	{
-		check(InResourceType == RRT_TextureReference);
-	}
 
 	FRHITextureDesc TextureDesc;
 
@@ -2450,7 +2446,7 @@ public:
 	 */
 	virtual void Wait(FRHICommandListImmediate& RHICmdList, FRHIGPUMask GPUMask) const = 0;
 
-	const FName& GetFName() const { return FenceName; }
+	FName GetFName() const { return FenceName; }
 
 	FThreadSafeCounter NumPendingWriteCommands;
 
diff --git a/Engine/Source/Runtime/RHI/Public/RHIValidation.h b/Engine/Source/Runtime/RHI/Public/RHIValidation.h
--- a/Engine/Source/Runtime/RHI/Public/RHIValidation.h
+++ b/Engine/Source/Runtime/RHI/Public/RHIValidation.h
@@ -434,26 +434,9 @@ public:
 		});
 	}
 
-	struct FRHIValidationBufferInitializer : public FRHIBufferInitializer
-	{
-		FRHIValidationBufferInitializer(const FRHIBufferCreateDesc& CreateDesc, FRHIBufferInitializer&& InOther)
-			: FRHIBufferInitializer(MoveTemp(InOther))
-		{
-			FFinalizeCallback OriginalFinalizeCallback = MoveTemp(FinalizeCallback);
-			FinalizeCallback = [OriginalFinalizeCallback = MoveTemp(OriginalFinalizeCallback), InitialState = CreateDesc.InitialState](FRHICommandListBase& RHICmdList)
-				{
-					FBufferRHIRef Buffer = OriginalFinalizeCallback(RHICmdList);
-
-					Buffer->InitBarrierTracking(InitialState, Buffer->GetName());
-
-					return Buffer;
-				};
-		}
-	};
-
 	[[nodiscard]] virtual FRHIBufferInitializer RHICreateBufferInitializer(FRHICommandListBase& RHICmdList, const FRHIBufferCreateDesc& CreateDesc) final override
 	{
-		return FRHIValidationBufferInitializer(CreateDesc, RHI->RHICreateBufferInitializer(RHICmdList, CreateDesc));
+		return RHI->RHICreateBufferInitializer(RHICmdList, CreateDesc);
 	}
 
 	// FlushType: Flush RHI Thread
@@ -1092,9 +1075,8 @@ public:
 
 	virtual FRayTracingSceneRHIRef RHICreateRayTracingScene(FRayTracingSceneInitializer Initializer) override final
 	{
-		FName DebugName = Initializer.DebugName;
 		FRayTracingSceneRHIRef Result = RHI->RHICreateRayTracingScene(MoveTemp(Initializer));
-		Result->InitBarrierTracking(ERHIAccess::BVHWrite, *DebugName.ToString()); // BVHs are always created in BVHWrite state
+		Result->InitBarrierTracking(ERHIAccess::BVHWrite, *Initializer.DebugName.ToString()); // BVHs are always created in BVHWrite state
 		return Result;
 	}
 
diff --git a/Engine/Source/Runtime/RHI/Public/RHIValidationCommon.h b/Engine/Source/Runtime/RHI/Public/RHIValidationCommon.h
--- a/Engine/Source/Runtime/RHI/Public/RHIValidationCommon.h
+++ b/Engine/Source/Runtime/RHI/Public/RHIValidationCommon.h
@@ -34,6 +34,7 @@ class FValidationComputeContext;
 class FValidationContext;
 class FValidationRHI;
 
+struct FRHIBufferCreateDesc;
 struct FRHITextureCreateDesc;
 struct FRHITransitionInfo;
 struct FRHIViewDesc;
@@ -419,16 +420,19 @@ namespace RHIValidation
 	class FBufferResource : public FResource
 	{
 	public:
-		inline void InitBarrierTracking(ERHIAccess InResourceState, FName InDebugName)
-		{
-			const FString LocalDebugName = InDebugName.ToString();
-			FResource::InitBarrierTracking(1, 1, 1, InResourceState, *LocalDebugName);
-		}
+		FBufferResource() = default;
+		RHI_API FBufferResource(const FRHIBufferCreateDesc& CreateDesc);
+
+		RHI_API void InitBarrierTracking(const FRHIBufferCreateDesc& CreateDesc);
 	};
 
 	class FAccelerationStructureResource : public FBufferResource
 	{
 	public:
+		void InitBarrierTracking(ERHIAccess InResourceState, const TCHAR* InDebugName)
+		{
+			FResource::InitBarrierTracking(1, 1, 1, InResourceState, InDebugName);
+		}
 	};
 	   
 	class FTextureResource
diff --git a/Engine/Source/Runtime/RHICore/Internal/RHICoreBufferInitializer.h b/Engine/Source/Runtime/RHICore/Internal/RHICoreBufferInitializer.h
--- a/Engine/Source/Runtime/RHICore/Internal/RHICoreBufferInitializer.h
+++ b/Engine/Source/Runtime/RHICore/Internal/RHICoreBufferInitializer.h
@@ -56,11 +56,20 @@ namespace UE::RHICore
 		}
 	};
 
+	class FInvalidBuffer : public FRHIBuffer
+	{
+	public:
+		FInvalidBuffer(const FRHIBufferCreateDesc& CreateDesc)
+			: FRHIBuffer(CreateDesc)
+		{
+		}
+	};
+
 	static FRHIBufferInitializer HandleUnknownBufferInitializerInitAction(FRHICommandListBase& RHICmdList, const FRHIBufferCreateDesc& CreateDesc)
 	{
 		UE_LOG(LogRHICore, Fatal, TEXT("Unknown or unhandled ERHIBufferInitAction: %d"), static_cast<uint32>(CreateDesc.InitAction));
 
-		FRHIBuffer* Buffer = new FRHIBuffer(CreateDesc);
+		FInvalidBuffer* Buffer = new FInvalidBuffer(CreateDesc);
 		return UE::RHICore::FDefaultBufferInitializer(RHICmdList, Buffer);
 	}
 
