More verbose checkf messages in GPU profiler ResolveSignal function

diff --git a/Engine/Source/Runtime/RHI/Private/GPUProfiler.cpp b/Engine/Source/Runtime/RHI/Private/GPUProfiler.cpp
--- a/Engine/Source/Runtime/RHI/Private/GPUProfiler.cpp
+++ b/Engine/Source/Runtime/RHI/Private/GPUProfiler.cpp
@@ -1870,10 +1870,10 @@ namespace UE::RHI::GPUProfiler
 				// Requiring that fences are always signalled in sequential order solves this.
 				// If the awaited value is less than the first Signal, the fence has already been signalled before the frame boundary.
 				//
-				check(Result.Value == MaxSignal.Value + 1);
+				checkf(Result.Value == MaxSignal.Value + 1, TEXT("Fence signals must be sequential. Result.Value: %llu, MaxSignal.Value + 1: %llu"), Result.Value, (MaxSignal.Value + 1));
 
 				// Signals should always advance in time
-				check(Result.GPUTimestampBOP >= MaxSignal.GPUTimestampBOP);
+				checkf(Result.GPUTimestampBOP >= MaxSignal.GPUTimestampBOP, TEXT("Signals should always advance in time. Result.GPUTimestampBOP: %llu, MaxSignal.GPUTimestampBOP: %llu"), Result.GPUTimestampBOP, MaxSignal.GPUTimestampBOP);
 
 				MaxSignal = Result;
 			}
