Fix for RHIGetParallelCommandContext in validation causing crash on some platforms due to invalid context in RecursiveHazardous

diff --git a/Engine/Source/Runtime/RHI/Private/RHIValidation.cpp b/Engine/Source/Runtime/RHI/Private/RHIValidation.cpp
--- a/Engine/Source/Runtime/RHI/Private/RHIValidation.cpp
+++ b/Engine/Source/Runtime/RHI/Private/RHIValidation.cpp
@@ -698,21 +698,34 @@ void FValidationRHI::RHIFinalizeContext(FRHIFinalizeContextArgs&& Args, TRHIPipe
 
 IRHIComputeContext* FValidationRHI::RHIGetParallelCommandContext(FRHIParallelRenderPassInfo const& ParallelRenderPass, FRHIGPUMask GPUMask)
 {
-	IRHIComputeContext* InnerContext = RHI->RHIGetParallelCommandContext(ParallelRenderPass, GPUMask);
-	check(InnerContext);
+	// If a platform has a ChildWait or ParentWait it is expected that they will override RHIGetParallelCommandContext
+	// otherwise, we need to manually call RHIGetCommandContext and RHIBeginRenderPass seperately, because the default implementation
+	// calls one after another potentially crasing in GetHightestLevelContext (FRHICommandList_RecursiveHazardous)
+	// TODO: Remove this after implementing RHISetupParallelPass
+	if (GRHIParallelRHIExecuteChildWait || GRHIParallelRHIExecuteParentWait)
+	{
+		IRHIComputeContext* InnerContext = RHI->RHIGetParallelCommandContext(ParallelRenderPass, GPUMask);
+		check(InnerContext);
 
-	FValidationContext* OuterContext = new FValidationContext(FValidationContext::EType::Parallel);
-	OuterContext->LinkToContext(static_cast<IRHICommandContext*>(InnerContext));
+		FValidationContext* OuterContext = new FValidationContext(FValidationContext::EType::Parallel);
+		OuterContext->LinkToContext(static_cast<IRHICommandContext*>(InnerContext));
 
-	// Parallel contexts are always inside a renderpass
-	OuterContext->State.bInsideBeginRenderPass = true;
-	OuterContext->State.RenderPassInfo = ParallelRenderPass;
-	if (ParallelRenderPass.PassName)
+		// Parallel contexts are always inside a renderpass
+		OuterContext->State.bInsideBeginRenderPass = true;
+		OuterContext->State.RenderPassInfo = ParallelRenderPass;
+		if (ParallelRenderPass.PassName)
+		{
+			OuterContext->State.RenderPassName = ParallelRenderPass.PassName;
+		}
+
+		return OuterContext;
+	}
+	else
 	{
-		OuterContext->State.RenderPassName = ParallelRenderPass.PassName;
+		IRHICommandContext* Context = static_cast<IRHICommandContext*>(RHIGetCommandContext(ERHIPipeline::Graphics, GPUMask));
+		Context->RHIBeginRenderPass(ParallelRenderPass, ParallelRenderPass.PassName);
+		return Context;
 	}
-
-	return OuterContext;
 }
 
 void FValidationRHI::RHICloseTranslateChain(FRHIFinalizeContextArgs&& Args, TRHIPipelineArray<IRHIPlatformCommandList*>& Output, bool bShouldFinalize)
