For PreviewShaderPlatform do not try to inherit from ParentShaderPlatform if that Section does not exist

diff --git a/Engine/Source/Runtime/RHI/Private/DataDrivenShaderPlatformInfo.cpp b/Engine/Source/Runtime/RHI/Private/DataDrivenShaderPlatformInfo.cpp
--- a/Engine/Source/Runtime/RHI/Private/DataDrivenShaderPlatformInfo.cpp
+++ b/Engine/Source/Runtime/RHI/Private/DataDrivenShaderPlatformInfo.cpp
@@ -431,18 +431,22 @@ void FGenericDataDrivenShaderPlatformInfo::Initialize()
 				const FString& SectionName = Section.Key;
 				const FConfigSection& SectionSettings = Section.Value;
 
-				const FString ParentShaderPlatformName = GetSectionString(SectionSettings, "ParentShaderPlatform");
-				const EShaderPlatform ParentShaderPlatform = ParseShaderPlatform(*ParentShaderPlatformName);
-
 				// get enum value for the string name
 				const EShaderPlatform ShaderPlatform = EShaderPlatform(CustomShaderPlatform++);
-
 				FGenericDataDrivenShaderPlatformInfo& Info = Infos[ShaderPlatform];
-				if (IsValid(ParentShaderPlatform))
+
+				const FString ParentShaderPlatformName = GetSectionString(SectionSettings, "ParentShaderPlatform");
+				if (!ParentShaderPlatformName.IsEmpty())
 				{
-					const FGenericDataDrivenShaderPlatformInfo& ParentInfo = Infos[ParentShaderPlatform];
-					Info = ParentInfo;
+					const EShaderPlatform ParentShaderPlatform = ParseShaderPlatform(*ParentShaderPlatformName);
+					
+					if (IsValid(ParentShaderPlatform))
+					{
+						const FGenericDataDrivenShaderPlatformInfo& ParentInfo = Infos[ParentShaderPlatform];
+						Info = ParentInfo;
+					}
 				}
+
 				Info.Name = *SectionName.Mid(22);
 				
 				ParseDataDrivenShaderInfo(SectionSettings, ShaderPlatform);
