Allow performance mode to be used with an RHI that is not D3D11

diff --git a/Engine/Source/Runtime/RHI/Private/Windows/WindowsDynamicRHI.cpp b/Engine/Source/Runtime/RHI/Private/Windows/WindowsDynamicRHI.cpp
--- a/Engine/Source/Runtime/RHI/Private/Windows/WindowsDynamicRHI.cpp
+++ b/Engine/Source/Runtime/RHI/Private/Windows/WindowsDynamicRHI.cpp
@@ -408,14 +408,14 @@ FParsedWindowsDynamicRHIConfig ParseWindowsDynamicRHIConfig()
 	return Config;
 }
 
-// Whether to default to Performance Mode on low-end machines (D3D11 RHI, ES31 feature level).
+// Whether to default to Performance Mode on low-end machines (ES31 feature level).
 // Looks at a few configuration options that control thresholds to enable this as the default.
 static bool ShouldDefaultToPerformanceMode()
 {
-	static TOptional<bool> ForceD3D11_ES31;
-	if (ForceD3D11_ES31.IsSet())
+	static TOptional<bool> ForceES31;
+	if (ForceES31.IsSet())
 	{
-		return ForceD3D11_ES31.GetValue();
+		return ForceES31.GetValue();
 	}
 
 	// Force Performance mode for laptops. We don't have an explicit way to check for laptops
@@ -424,7 +424,7 @@ static bool ShouldDefaultToPerformanceMode()
 	if (GConfig->GetBoolOrDefault(TEXT("PerformanceMode"), TEXT("bDefaultOnLaptops"), false, GEngineIni) && bIsLaptop)
 	{
 		UE_LOG(LogRHI, Log, TEXT("Running on laptop: default feature level to %s"), GetLogName(ERHIFeatureLevel::ES3_1));
-		ForceD3D11_ES31 = true;
+		ForceES31 = true;
 		return true;
 	}
 
@@ -433,7 +433,7 @@ static bool ShouldDefaultToPerformanceMode()
 	if (GConfig->GetInt(TEXT("PerformanceMode"), TEXT("MinCoreCount"), MinCoreCount, GEngineIni) && FPlatformMisc::NumberOfCoresIncludingHyperthreads() < MinCoreCount)
 	{
 		UE_LOG(LogRHI, Log, TEXT("MinCoreCount found (%d) and user core count (%d): default feature level to %s"), MinCoreCount, FPlatformMisc::NumberOfCoresIncludingHyperthreads(), GetLogName(ERHIFeatureLevel::ES3_1));
-		ForceD3D11_ES31 = true;
+		ForceES31 = true;
 		return true;
 	}
 
@@ -453,7 +453,7 @@ static bool ShouldDefaultToPerformanceMode()
 				{
 					UE_LOG(LogRHI, Log, TEXT("MinMemorySizeBucket found (%s) and user memory (%d): default feature level to %s"), *MinMemorySizeBucketString, int32(FPlatformMemory::GetMemorySizeBucket()), GetLogName(ERHIFeatureLevel::ES3_1));
 
-					ForceD3D11_ES31 = true;
+					ForceES31 = true;
 					return true;
 				}
 			}
@@ -466,7 +466,7 @@ static bool ShouldDefaultToPerformanceMode()
 				{
 					UE_LOG(LogRHI, Log, TEXT("MinMemorySizeBucket found (%s) and user memory (%d): default feature level to %s"), *MinMemorySizeBucketString, int32(FPlatformMemory::GetMemorySizeBucket()), GetLogName(ERHIFeatureLevel::ES3_1));
 
-					ForceD3D11_ES31 = true;
+					ForceES31 = true;
 					return true;
 				}
 			}
@@ -493,7 +493,7 @@ static bool ShouldDefaultToPerformanceMode()
 		{
 			UE_LOG(LogRHI, Log, TEXT("Found RHIName %s with DeviceName %s: default feature level to %s"), *RHIName, *DeviceName, GetLogName(ERHIFeatureLevel::ES3_1));
 
-			ForceD3D11_ES31 = true;
+			ForceES31 = true;
 			return true;
 		}
 
@@ -511,12 +511,12 @@ static bool ShouldDefaultToPerformanceMode()
 		{
 			UE_LOG(LogRHI, Log, TEXT("Found RHIName %s with DeviceName %s, VendorId %s, and DeviceId %s: default feature level to %s"), *RHIName, *DeviceName, *VendorId, *DeviceId, GetLogName(ERHIFeatureLevel::ES3_1));
 
-			ForceD3D11_ES31 = true;
+			ForceES31 = true;
 			return true;
 		}
 	}
 
-	ForceD3D11_ES31 = false;
+	ForceES31 = false;
 	return false;
 }
 
@@ -675,21 +675,11 @@ static TOptional<ERHIFeatureLevel::Type> ChooseDefaultFeatureLevel(EWindowsRHI I
 }
 
 // Choose the default RHI using this search order:
-//  * Performance mode (D3D11) if the machine is low-end (see ShouldDefaultToPerformanceMode) and if we are not running the editor;
 //  * DefaultGraphicsRHI if specified;
 //  * The first supported RHI from TargetedRHIs, enumerated according to GRHISearchOrder;
 //  * D3D11, if all else fails.
 static EWindowsRHI ChooseDefaultRHI(const FParsedWindowsDynamicRHIConfig& Config)
 {
-	if (!GIsEditor && !CVarIgnorePerformanceModeCheck.GetValueOnAnyThread())
-	{
-		bool bDefaultD3D11PerfMode = ShouldDefaultToPerformanceMode();
-		if (bDefaultD3D11PerfMode)
-		{
-			return EWindowsRHI::D3D11;
-		}
-	}
-
 	// Default graphics RHI is the main project setting that governs the choice, so it takes the priority
 	if (TOptional<EWindowsRHI> ConfigDefault = Config.DefaultRHI)
 	{
