Remove FRHIPipelineBinaryLibrary and associated overloads of RHICreateGraphicsPipelineState
- The RHI resource type is entirely unused in the engine

diff --git a/Engine/Source/Runtime/RHI/Private/RHIStrings.cpp b/Engine/Source/Runtime/RHI/Private/RHIStrings.cpp
--- a/Engine/Source/Runtime/RHI/Private/RHIStrings.cpp
+++ b/Engine/Source/Runtime/RHI/Private/RHIStrings.cpp
@@ -434,7 +434,6 @@ static const FRHIResourceTypeName GRHIResourceTypeNames[] =
 	RHI_RESOURCE_TYPE_DEF(StagingBuffer),
 	RHI_RESOURCE_TYPE_DEF(CustomPresent),
 	RHI_RESOURCE_TYPE_DEF(ShaderLibrary),
-	RHI_RESOURCE_TYPE_DEF(PipelineBinaryLibrary),
 	RHI_RESOURCE_TYPE_DEF(StreamSourceSlot)
 };
 
diff --git a/Engine/Source/Runtime/RHI/Public/DynamicRHI.h b/Engine/Source/Runtime/RHI/Public/DynamicRHI.h
--- a/Engine/Source/Runtime/RHI/Public/DynamicRHI.h
+++ b/Engine/Source/Runtime/RHI/Public/DynamicRHI.h
@@ -417,16 +417,6 @@ public:
 	// FlushType: Thread safe
 	virtual FComputePipelineStateRHIRef RHICreateComputePipelineState(const FComputePipelineStateInitializer& Initializer) = 0;
 
-	virtual FGraphicsPipelineStateRHIRef RHICreateGraphicsPipelineState(const FGraphicsPipelineStateInitializer& Initializer, FRHIPipelineBinaryLibrary* PipelineBinary)
-	{
-		return RHICreateGraphicsPipelineState(Initializer);
-	}
-
-	virtual FComputePipelineStateRHIRef RHICreateComputePipelineState(const FComputePipelineStateInitializer& Initializer, FRHIPipelineBinaryLibrary* PipelineBinary)
-	{
-		return RHICreateComputePipelineState(Initializer);
-	}
-
 	virtual FWorkGraphPipelineStateRHIRef RHICreateWorkGraphPipelineState(const FWorkGraphPipelineStateInitializer& Initializer)
 	{
 		checkNoEntry();
diff --git a/Engine/Source/Runtime/RHI/Public/RHIDefinitions.h b/Engine/Source/Runtime/RHI/Public/RHIDefinitions.h
--- a/Engine/Source/Runtime/RHI/Public/RHIDefinitions.h
+++ b/Engine/Source/Runtime/RHI/Public/RHIDefinitions.h
@@ -1049,7 +1049,6 @@ enum ERHIResourceType : uint8
 	RRT_StagingBuffer,
 	RRT_CustomPresent,
 	RRT_ShaderLibrary,
-	RRT_PipelineBinaryLibrary,
 	RRT_ShaderBundle,
 	RRT_WorkGraphShader,
 	RRT_WorkGraphPipelineState,
diff --git a/Engine/Source/Runtime/RHI/Public/RHIFwd.h b/Engine/Source/Runtime/RHI/Public/RHIFwd.h
--- a/Engine/Source/Runtime/RHI/Public/RHIFwd.h
+++ b/Engine/Source/Runtime/RHI/Public/RHIFwd.h
@@ -62,7 +62,6 @@ class FRHIGeometryShader;
 class FRHIGPUFence;
 class FRHIGraphicsPipelineState;
 class FRHIMeshShader;
-class FRHIPipelineBinaryLibrary;
 class FRHIPixelShader;
 class FRHIRasterizerState;
 class FRHIRayTracingGeometry;
@@ -120,7 +119,6 @@ using FShaderBindingTableRHIRef        = TRefCountPtr<FRHIShaderBindingTable>;
 using FRenderQueryPoolRHIRef           = TRefCountPtr<FRHIRenderQueryPool>;
 using FRenderQueryRHIRef               = TRefCountPtr<FRHIRenderQuery>;
 using FRHIResourceCollectionRef        = TRefCountPtr<FRHIResourceCollection>;
-using FRHIPipelineBinaryLibraryRef     = TRefCountPtr<FRHIPipelineBinaryLibrary>;
 using FRHIShaderLibraryRef             = TRefCountPtr<FRHIShaderLibrary>;
 using FRHIShaderResourceViewRef        = TRefCountPtr<FRHIShaderResourceView>;
 using FSamplerStateRHIRef              = TRefCountPtr<FRHISamplerState>;
diff --git a/Engine/Source/Runtime/RHI/Public/RHIShaderLibrary.h b/Engine/Source/Runtime/RHI/Public/RHIShaderLibrary.h
--- a/Engine/Source/Runtime/RHI/Public/RHIShaderLibrary.h
+++ b/Engine/Source/Runtime/RHI/Public/RHIShaderLibrary.h
@@ -49,15 +49,3 @@ protected:
 	FString LibraryName;
 	uint32 LibraryId;
 };
-
-class FRHIPipelineBinaryLibrary : public FRHIResource
-{
-public:
-	FRHIPipelineBinaryLibrary(EShaderPlatform InPlatform, FString const& FilePath) : FRHIResource(RRT_PipelineBinaryLibrary), Platform(InPlatform) {}
-	virtual ~FRHIPipelineBinaryLibrary() {}
-
-	FORCEINLINE EShaderPlatform GetPlatform(void) const { return Platform; }
-
-protected:
-	EShaderPlatform Platform;
-};
diff --git a/Engine/Source/Runtime/RHI/Public/RHIValidation.h b/Engine/Source/Runtime/RHI/Public/RHIValidation.h
--- a/Engine/Source/Runtime/RHI/Public/RHIValidation.h
+++ b/Engine/Source/Runtime/RHI/Public/RHIValidation.h
@@ -383,22 +383,6 @@ public:
 		return RHI->RHICreateComputePipelineState(Initializer);
 	}
 
-	virtual FGraphicsPipelineStateRHIRef RHICreateGraphicsPipelineState(const FGraphicsPipelineStateInitializer& Initializer, FRHIPipelineBinaryLibrary* PipelineBinary) override final
-	{
-		ValidatePipeline(Initializer);
-		FGraphicsPipelineStateRHIRef PSO = RHI->RHICreateGraphicsPipelineState(Initializer);
-		if (PSO.IsValid())
-		{
-			PSO->DSMode = Initializer.DepthStencilState->ActualDSMode;
-		}
-		return PSO;
-	}
-
-	virtual TRefCountPtr<FRHIComputePipelineState> RHICreateComputePipelineState(const FComputePipelineStateInitializer& Initializer, FRHIPipelineBinaryLibrary* PipelineBinary) override final
-	{
-		return RHI->RHICreateComputePipelineState(Initializer, PipelineBinary);
-	}
-
 	/**
 	* Creates a uniform buffer.  The contents of the uniform buffer are provided in a parameter, and are immutable.
 	* CAUTION: Even though this is marked as threadsafe, it is only valid to call from the render thread or the RHI thread. Thus is need not be threadsafe on platforms that do not support or aren't using an RHIThread
