Remove RRT_Texture2D, RRT_Texture2DArray, RRT_Texture3D and RRT_TextureCube enum values
- These have been unused since UE 5.1

diff --git a/Engine/Source/Runtime/RHI/Private/RHI.cpp b/Engine/Source/Runtime/RHI/Private/RHI.cpp
--- a/Engine/Source/Runtime/RHI/Private/RHI.cpp
+++ b/Engine/Source/Runtime/RHI/Private/RHI.cpp
@@ -734,13 +734,7 @@ namespace RHIInternal
 			{
 				if (TypeFilter == RRT_Texture)
 				{
-					if (ResourceInfo.Type != RRT_Texture2D &&
-						ResourceInfo.Type != RRT_Texture2DArray &&
-						ResourceInfo.Type != RRT_Texture3D &&
-						ResourceInfo.Type != RRT_TextureCube)
-					{
-						return false;
-					}
+					return false;
 				}
 				else if (TypeFilter != ResourceInfo.Type)
 				{
@@ -806,12 +800,7 @@ namespace RHIInternal
 		Flags.bMarkedForDelete = !Resource.ResourceInfo.bValid;
 		Flags.bTransient = Resource.ResourceInfo.IsTransient;
 
-		bool bIsTexture = Resource.ResourceInfo.Type == RRT_Texture ||
-			Resource.ResourceInfo.Type == RRT_Texture2D ||
-			Resource.ResourceInfo.Type == RRT_Texture2DArray ||
-			Resource.ResourceInfo.Type == RRT_Texture3D ||
-			Resource.ResourceInfo.Type == RRT_TextureCube;
-		if (bIsTexture)
+		if (Resource.ResourceInfo.Type == RRT_Texture)
 		{
 			FRHITexture* Texture = (FRHITexture*)Resource.Resource;
 			Flags.bRT = EnumHasAnyFlags(Texture->GetFlags(), TexCreate_RenderTargetable);
diff --git a/Engine/Source/Runtime/RHI/Private/RHIStrings.cpp b/Engine/Source/Runtime/RHI/Private/RHIStrings.cpp
--- a/Engine/Source/Runtime/RHI/Private/RHIStrings.cpp
+++ b/Engine/Source/Runtime/RHI/Private/RHIStrings.cpp
@@ -418,10 +418,12 @@ static const FRHIResourceTypeName GRHIResourceTypeNames[] =
 	RHI_RESOURCE_TYPE_DEF(UniformBuffer),
 	RHI_RESOURCE_TYPE_DEF(Buffer),
 	RHI_RESOURCE_TYPE_DEF(Texture),
+PRAGMA_DISABLE_DEPRECATION_WARNINGS
 	RHI_RESOURCE_TYPE_DEF(Texture2D),
 	RHI_RESOURCE_TYPE_DEF(Texture2DArray),
 	RHI_RESOURCE_TYPE_DEF(Texture3D),
 	RHI_RESOURCE_TYPE_DEF(TextureCube),
+PRAGMA_ENABLE_DEPRECATION_WARNINGS
 	RHI_RESOURCE_TYPE_DEF(TextureReference),
 	RHI_RESOURCE_TYPE_DEF(TimestampCalibrationQuery),
 	RHI_RESOURCE_TYPE_DEF(GPUFence),
diff --git a/Engine/Source/Runtime/RHI/Public/RHIDefinitions.h b/Engine/Source/Runtime/RHI/Public/RHIDefinitions.h
--- a/Engine/Source/Runtime/RHI/Public/RHIDefinitions.h
+++ b/Engine/Source/Runtime/RHI/Public/RHIDefinitions.h
@@ -1030,12 +1030,10 @@ enum ERHIResourceType : uint8
 	RRT_UniformBuffer,
 	RRT_Buffer,
 	RRT_Texture,
-	// @todo: texture type unification - remove these
-	RRT_Texture2D,
-	RRT_Texture2DArray,
-	RRT_Texture3D,
-	RRT_TextureCube,
-	// @todo: texture type unification - remove these
+	RRT_Texture2D      UE_DEPRECATED(5.7,      "RRT_Texture2D is unused. Remove any remaining uses of it. There is a single RHI texture type which always uses RRT_Texture. RHI texture references use RRT_TextureReference."),
+	RRT_Texture2DArray UE_DEPRECATED(5.7, "RRT_Texture2DArray is unused. Remove any remaining uses of it. There is a single RHI texture type which always uses RRT_Texture. RHI texture references use RRT_TextureReference."),
+	RRT_Texture3D      UE_DEPRECATED(5.7,      "RRT_Texture3D is unused. Remove any remaining uses of it. There is a single RHI texture type which always uses RRT_Texture. RHI texture references use RRT_TextureReference."),
+	RRT_TextureCube    UE_DEPRECATED(5.7,    "RRT_TextureCube is unused. Remove any remaining uses of it. There is a single RHI texture type which always uses RRT_Texture. RHI texture references use RRT_TextureReference."),
 	RRT_TextureReference,
 	RRT_TimestampCalibrationQuery,
 	RRT_GPUFence,
@@ -1646,21 +1644,10 @@ inline bool IsWorkGraphShaderFrequency(EShaderFrequency Frequency)
 	}
 }
 
-inline ERHIResourceType GetRHIResourceType(ETextureDimension Dimension)
+UE_DEPRECATED(5.7, "GetRHIResourceType(ETextureDimension Dimension) is deprecated. There is a single RHI texture type which always uses RRT_Texture now, which makes this function pointless. Remove any remaining uses of this function.")
+inline ERHIResourceType GetRHIResourceType(ETextureDimension /*Dimension*/)
 {
-	switch (Dimension)
-	{
-	case ETextureDimension::Texture2D:
-		return ERHIResourceType::RRT_Texture2D;
-	case ETextureDimension::Texture2DArray:
-		return ERHIResourceType::RRT_Texture2DArray;
-	case ETextureDimension::Texture3D:
-		return ERHIResourceType::RRT_Texture3D;
-	case ETextureDimension::TextureCube:
-	case ETextureDimension::TextureCubeArray:
-		return ERHIResourceType::RRT_TextureCube;
-	}
-	return ERHIResourceType::RRT_None;
+	return ERHIResourceType::RRT_Texture;
 }
 
 #if PLATFORM_SUPPORTS_GEOMETRY_SHADERS
