Add bPresent argument to RHIAdvanceFrameForGetViewportBackBuffer
- Required by some RHIs to avoid cycling through backbuffer indices when rendered frames aren't presented, which can lead to deadlocks.

diff --git a/Engine/Source/Runtime/Apple/MetalRHI/Private/MetalDynamicRHI.h b/Engine/Source/Runtime/Apple/MetalRHI/Private/MetalDynamicRHI.h
--- a/Engine/Source/Runtime/Apple/MetalRHI/Private/MetalDynamicRHI.h
+++ b/Engine/Source/Runtime/Apple/MetalRHI/Private/MetalDynamicRHI.h
@@ -166,7 +166,7 @@ public:
 	virtual FRenderQueryRHIRef RHICreateRenderQuery(ERenderQueryType QueryType) final override;
 	virtual bool RHIGetRenderQueryResult(FRHIRenderQuery* RenderQuery, uint64& OutResult, bool bWait, uint32 GPUIndex = INDEX_NONE) final override;
 	virtual FTextureRHIRef RHIGetViewportBackBuffer(FRHIViewport* Viewport) final override;
-	virtual void RHIAdvanceFrameForGetViewportBackBuffer(FRHIViewport* Viewport) final override;
+	virtual void RHIAdvanceFrameForGetViewportBackBuffer(FRHIViewport* Viewport, bool bPresent) final override;
 	virtual void RHIFlushResources() final override;
 	virtual FViewportRHIRef RHICreateViewport(void* WindowHandle, uint32 SizeX, uint32 SizeY, bool bIsFullscreen, EPixelFormat PreferredPixelFormat) final override;
 	virtual void RHIResizeViewport(FRHIViewport* Viewport, uint32 SizeX, uint32 SizeY, bool bIsFullscreen) final override;
diff --git a/Engine/Source/Runtime/Apple/MetalRHI/Private/MetalViewport.cpp b/Engine/Source/Runtime/Apple/MetalRHI/Private/MetalViewport.cpp
--- a/Engine/Source/Runtime/Apple/MetalRHI/Private/MetalViewport.cpp
+++ b/Engine/Source/Runtime/Apple/MetalRHI/Private/MetalViewport.cpp
@@ -895,7 +895,7 @@ FTextureRHIRef FMetalDynamicRHI::RHIGetViewportBackBuffer(FRHIViewport* Viewport
 	return FTextureRHIRef(Viewport->GetBackBuffer(EMetalViewportAccessRenderer).GetReference());
 }
 
-void FMetalDynamicRHI::RHIAdvanceFrameForGetViewportBackBuffer(FRHIViewport* ViewportRHI)
+void FMetalDynamicRHI::RHIAdvanceFrameForGetViewportBackBuffer(FRHIViewport* ViewportRHI, bool bPresent)
 {
 	if (GMetalSeparatePresentThread && (GRHISupportsRHIThread && IsRunningRHIInSeparateThread()))
 	{
diff --git a/Engine/Source/Runtime/D3D12RHI/Private/D3D12RHIPrivate.h b/Engine/Source/Runtime/D3D12RHI/Private/D3D12RHIPrivate.h
--- a/Engine/Source/Runtime/D3D12RHI/Private/D3D12RHIPrivate.h
+++ b/Engine/Source/Runtime/D3D12RHI/Private/D3D12RHIPrivate.h
@@ -394,7 +394,7 @@ public:
 	virtual uint64 RHIComputeStatePrecachePSOHash(const FGraphicsPipelineStateInitializer& Initializer) final override;
 	virtual uint64 RHIComputePrecachePSOHash(const FGraphicsPipelineStateInitializer& Initializer) final override;
 	virtual bool RHIMatchPrecachePSOInitializers(const FGraphicsPipelineStateInitializer& LHS, const FGraphicsPipelineStateInitializer& RHS) final override;
-	virtual void RHIAdvanceFrameForGetViewportBackBuffer(FRHIViewport* Viewport) final override;
+	virtual void RHIAdvanceFrameForGetViewportBackBuffer(FRHIViewport* Viewport, bool bPresent) final override;
 	virtual void RHIFlushResources() final override;
 	virtual FViewportRHIRef RHICreateViewport(void* WindowHandle, uint32 SizeX, uint32 SizeY, bool bIsFullscreen, EPixelFormat PreferredPixelFormat) final override;
 	virtual void RHIResizeViewport(FRHIViewport* Viewport, uint32 SizeX, uint32 SizeY, bool bIsFullscreen) final override;
diff --git a/Engine/Source/Runtime/D3D12RHI/Private/D3D12Viewport.cpp b/Engine/Source/Runtime/D3D12RHI/Private/D3D12Viewport.cpp
--- a/Engine/Source/Runtime/D3D12RHI/Private/D3D12Viewport.cpp
+++ b/Engine/Source/Runtime/D3D12RHI/Private/D3D12Viewport.cpp
@@ -822,7 +822,7 @@ void FD3D12CommandContextBase::RHIEndDrawingViewport(FRHIViewport* ViewportRHI,
 	}
 }
 
-void FD3D12DynamicRHI::RHIAdvanceFrameForGetViewportBackBuffer(FRHIViewport* ViewportRHI)
+void FD3D12DynamicRHI::RHIAdvanceFrameForGetViewportBackBuffer(FRHIViewport* ViewportRHI, bool bPresent)
 {
 	check(IsInRenderingThread());
 
diff --git a/Engine/Source/Runtime/NullDrv/Public/NullRHI.h b/Engine/Source/Runtime/NullDrv/Public/NullRHI.h
--- a/Engine/Source/Runtime/NullDrv/Public/NullRHI.h
+++ b/Engine/Source/Runtime/NullDrv/Public/NullRHI.h
@@ -444,7 +444,7 @@ public:
 	{
 
 	}
-	virtual void RHIAdvanceFrameForGetViewportBackBuffer(FRHIViewport* Viewport) final override
+	virtual void RHIAdvanceFrameForGetViewportBackBuffer(FRHIViewport* Viewport, bool bPresent) final override
 	{
 
 	}
diff --git a/Engine/Source/Runtime/OpenGLDrv/Private/OpenGLDrv.h b/Engine/Source/Runtime/OpenGLDrv/Private/OpenGLDrv.h
--- a/Engine/Source/Runtime/OpenGLDrv/Private/OpenGLDrv.h
+++ b/Engine/Source/Runtime/OpenGLDrv/Private/OpenGLDrv.h
@@ -371,7 +371,7 @@ public:
 	virtual FTextureRHIRef RHIGetViewportBackBuffer(FRHIViewport* Viewport) final override;
 	virtual void RHIAliasTextureResources(FTextureRHIRef& DestTextureRHI, FTextureRHIRef& SrcTextureRHI) final override;
 	virtual FTextureRHIRef RHICreateAliasedTexture(FTextureRHIRef& SourceTexture) final override;
-	virtual void RHIAdvanceFrameForGetViewportBackBuffer(FRHIViewport* Viewport) final override;
+	virtual void RHIAdvanceFrameForGetViewportBackBuffer(FRHIViewport* Viewport, bool bPresent) final override;
 	virtual void RHIAcquireThreadOwnership() final override;
 	virtual void RHIReleaseThreadOwnership() final override;
 	virtual void RHIFlushResources() final override;
diff --git a/Engine/Source/Runtime/OpenGLDrv/Private/OpenGLViewport.cpp b/Engine/Source/Runtime/OpenGLDrv/Private/OpenGLViewport.cpp
--- a/Engine/Source/Runtime/OpenGLDrv/Private/OpenGLViewport.cpp
+++ b/Engine/Source/Runtime/OpenGLDrv/Private/OpenGLViewport.cpp
@@ -211,7 +211,7 @@ FTextureRHIRef FOpenGLDynamicRHI::RHIGetViewportBackBuffer(FRHIViewport* Viewpor
 	return Viewport->GetBackBuffer();
 }
 
-void FOpenGLDynamicRHI::RHIAdvanceFrameForGetViewportBackBuffer(FRHIViewport* Viewport)
+void FOpenGLDynamicRHI::RHIAdvanceFrameForGetViewportBackBuffer(FRHIViewport* Viewport, bool bPresent)
 {
 }
 
diff --git a/Engine/Source/Runtime/RHI/Private/RHICommandList.cpp b/Engine/Source/Runtime/RHI/Private/RHICommandList.cpp
--- a/Engine/Source/Runtime/RHI/Private/RHICommandList.cpp
+++ b/Engine/Source/Runtime/RHI/Private/RHICommandList.cpp
@@ -1941,7 +1941,7 @@ void FRHICommandListImmediate::EndDrawingViewport(FRHIViewport* Viewport, bool b
 		}
 	}
 
-	RHIAdvanceFrameForGetViewportBackBuffer(Viewport);
+	GDynamicRHI->RHIAdvanceFrameForGetViewportBackBuffer(Viewport, bPresent);
 }
 
 void FRHICommandListImmediate::EndFrame()
diff --git a/Engine/Source/Runtime/RHI/Public/DynamicRHI.h b/Engine/Source/Runtime/RHI/Public/DynamicRHI.h
--- a/Engine/Source/Runtime/RHI/Public/DynamicRHI.h
+++ b/Engine/Source/Runtime/RHI/Public/DynamicRHI.h
@@ -698,7 +698,7 @@ public:
 
 	// Only relevant with an RHI thread, this advances the backbuffer for the purpose of GetViewportBackBuffer
 	// FlushType: Thread safe
-	virtual void RHIAdvanceFrameForGetViewportBackBuffer(FRHIViewport* Viewport) = 0;
+	virtual void RHIAdvanceFrameForGetViewportBackBuffer(FRHIViewport* Viewport, bool bPresent) = 0;
 
 	//
 	// Acquires or releases ownership of the platform-specific rendering context for the calling thread.
@@ -1271,9 +1271,10 @@ FORCEINLINE uint32 RHIGetHTilePlatformConfig(uint32 DepthWidth, uint32 DepthHeig
 	return GDynamicRHI->RHIGetHTilePlatformConfig(DepthWidth, DepthHeight);
 }
 
+UE_DEPRECATED(5.7, "RHIAdvanceFrameForGetViewportBackBuffer is deprecated. This happens automatically when RHIEndDrawingViewport is called. Remove calls to this function.")
 FORCEINLINE void RHIAdvanceFrameForGetViewportBackBuffer(FRHIViewport* Viewport)
 {
-	return GDynamicRHI->RHIAdvanceFrameForGetViewportBackBuffer(Viewport);
+	return GDynamicRHI->RHIAdvanceFrameForGetViewportBackBuffer(Viewport, true);
 }
 
 RHI_API uint32 RHIGetGPUFrameCycles(uint32 GPUIndex = 0);
diff --git a/Engine/Source/Runtime/RHI/Public/RHIValidation.h b/Engine/Source/Runtime/RHI/Public/RHIValidation.h
--- a/Engine/Source/Runtime/RHI/Public/RHIValidation.h
+++ b/Engine/Source/Runtime/RHI/Public/RHIValidation.h
@@ -764,9 +764,9 @@ public:
 
 	// Only relevant with an RHI thread, this advances the backbuffer for the purpose of GetViewportBackBuffer
 	// FlushType: Thread safe
-	virtual void RHIAdvanceFrameForGetViewportBackBuffer(FRHIViewport* Viewport) override final
+	virtual void RHIAdvanceFrameForGetViewportBackBuffer(FRHIViewport* Viewport, bool bPresent) override final
 	{
-		RHI->RHIAdvanceFrameForGetViewportBackBuffer(Viewport);
+		RHI->RHIAdvanceFrameForGetViewportBackBuffer(Viewport, bPresent);
 	}
 
 	virtual void RHIAcquireThreadOwnership() override final
diff --git a/Engine/Source/Runtime/VulkanRHI/Private/VulkanDynamicRHI.h b/Engine/Source/Runtime/VulkanRHI/Private/VulkanDynamicRHI.h
--- a/Engine/Source/Runtime/VulkanRHI/Private/VulkanDynamicRHI.h
+++ b/Engine/Source/Runtime/VulkanRHI/Private/VulkanDynamicRHI.h
@@ -217,7 +217,7 @@ public:
 	virtual FTextureRHIRef RHIGetViewportBackBuffer(FRHIViewport* Viewport) final override;
 	virtual void RHIAliasTextureResources(FTextureRHIRef& DestTexture, FTextureRHIRef& SrcTexture) final override;
 	virtual FTextureRHIRef RHICreateAliasedTexture(FTextureRHIRef& SourceTexture) final override;
-	virtual void RHIAdvanceFrameForGetViewportBackBuffer(FRHIViewport* Viewport) final override;
+	virtual void RHIAdvanceFrameForGetViewportBackBuffer(FRHIViewport* Viewport, bool bPresent) final override;
 	virtual void RHIFlushResources() final override;
 	virtual FViewportRHIRef RHICreateViewport(void* WindowHandle, uint32 SizeX, uint32 SizeY, bool bIsFullscreen, EPixelFormat PreferredPixelFormat) final override;
 	virtual void RHIResizeViewport(FRHIViewport* Viewport, uint32 SizeX, uint32 SizeY, bool bIsFullscreen) final override;
diff --git a/Engine/Source/Runtime/VulkanRHI/Private/VulkanViewport.cpp b/Engine/Source/Runtime/VulkanRHI/Private/VulkanViewport.cpp
--- a/Engine/Source/Runtime/VulkanRHI/Private/VulkanViewport.cpp
+++ b/Engine/Source/Runtime/VulkanRHI/Private/VulkanViewport.cpp
@@ -1289,7 +1289,7 @@ FTextureRHIRef FVulkanDynamicRHI::RHIGetViewportBackBuffer(FRHIViewport* Viewpor
 	return Viewport->GetBackBuffer(RHICmdList);
 }
 
-void FVulkanDynamicRHI::RHIAdvanceFrameForGetViewportBackBuffer(FRHIViewport* ViewportRHI)
+void FVulkanDynamicRHI::RHIAdvanceFrameForGetViewportBackBuffer(FRHIViewport* ViewportRHI, bool bPresent)
 {
 	FRHICommandListImmediate& RHICmdList = FRHICommandListImmediate::Get();
 
diff --git a/Engine/Source/Runtime/Windows/D3D11RHI/Private/D3D11RHIPrivate.h b/Engine/Source/Runtime/Windows/D3D11RHI/Private/D3D11RHIPrivate.h
--- a/Engine/Source/Runtime/Windows/D3D11RHI/Private/D3D11RHIPrivate.h
+++ b/Engine/Source/Runtime/Windows/D3D11RHI/Private/D3D11RHIPrivate.h
@@ -496,7 +496,7 @@ public:
 	virtual FTextureRHIRef RHIGetViewportBackBuffer(FRHIViewport* Viewport) final override;
 	virtual void RHIAliasTextureResources(FTextureRHIRef& DestTexture, FTextureRHIRef& SrcTexture) final override;
 	virtual FTextureRHIRef RHICreateAliasedTexture(FTextureRHIRef& SourceTexture) final override;
-	virtual void RHIAdvanceFrameForGetViewportBackBuffer(FRHIViewport* Viewport) final override;
+	virtual void RHIAdvanceFrameForGetViewportBackBuffer(FRHIViewport* Viewport, bool bPresent) final override;
 	virtual void RHIFlushResources() final override;
 	virtual FViewportRHIRef RHICreateViewport(void* WindowHandle, uint32 SizeX, uint32 SizeY, bool bIsFullscreen, EPixelFormat PreferredPixelFormat) override;
 	virtual void RHIResizeViewport(FRHIViewport* Viewport, uint32 SizeX, uint32 SizeY, bool bIsFullscreen) final override;
diff --git a/Engine/Source/Runtime/Windows/D3D11RHI/Private/D3D11Viewport.cpp b/Engine/Source/Runtime/Windows/D3D11RHI/Private/D3D11Viewport.cpp
--- a/Engine/Source/Runtime/Windows/D3D11RHI/Private/D3D11Viewport.cpp
+++ b/Engine/Source/Runtime/Windows/D3D11RHI/Private/D3D11Viewport.cpp
@@ -852,7 +852,7 @@ void FD3D11DynamicRHI::RHIEndDrawingViewport(FRHIViewport* ViewportRHI,bool bPre
 	}
 }
 
-void FD3D11DynamicRHI::RHIAdvanceFrameForGetViewportBackBuffer(FRHIViewport* Viewport)
+void FD3D11DynamicRHI::RHIAdvanceFrameForGetViewportBackBuffer(FRHIViewport* Viewport, bool bPresent)
 {
 }
 
