Wrap remaining references to RHICalibrateTimers in #if (RHI_NEW_GPU_PROFILER == 0)
- This was only used by the legacy GPU profiler.

diff --git a/Engine/Source/Runtime/Apple/MetalRHI/Private/MetalRHIContext.cpp b/Engine/Source/Runtime/Apple/MetalRHI/Private/MetalRHIContext.cpp
--- a/Engine/Source/Runtime/Apple/MetalRHI/Private/MetalRHIContext.cpp
+++ b/Engine/Source/Runtime/Apple/MetalRHI/Private/MetalRHIContext.cpp
@@ -464,6 +464,7 @@ void FMetalRHICommandContext::RHINextSubpass()
 #endif
 }
 
+#if (RHI_NEW_GPU_PROFILER == 0)
 void FMetalRHICommandContext::RHICalibrateTimers(FRHITimestampCalibrationQuery* CalibrationQuery)
 {
     MTL::Device* MTLDevice = Device.GetDevice();
@@ -474,6 +475,7 @@ void FMetalRHICommandContext::RHICalibrateTimers(FRHITimestampCalibrationQuery*
     CalibrationQuery->CPUMicroseconds[0] = uint64(CPUTimeStamp / 1000.0);
     CalibrationQuery->GPUMicroseconds[0] = uint64(GPUTimestamp / 1000.0);
 }
+#endif
 
 void FMetalRHICommandContext::FillBuffer(MTL::Buffer* Buffer, NS::Range Range, uint8 Value)
 {
diff --git a/Engine/Source/Runtime/Apple/MetalRHI/Public/MetalRHIContext.h b/Engine/Source/Runtime/Apple/MetalRHI/Public/MetalRHIContext.h
--- a/Engine/Source/Runtime/Apple/MetalRHI/Public/MetalRHIContext.h
+++ b/Engine/Source/Runtime/Apple/MetalRHI/Public/MetalRHIContext.h
@@ -83,7 +83,9 @@ public:
 	virtual void RHICopyTexture(FRHITexture* SourceTextureRHI, FRHITexture* DestTextureRHI, const FRHICopyTextureInfo& CopyInfo) final override;
 	virtual void RHICopyBufferRegion(FRHIBuffer* DstBufferRHI, uint64 DstOffset, FRHIBuffer* SrcBufferRHI, uint64 SrcOffset, uint64 NumBytes) final override;
 
+#if (RHI_NEW_GPU_PROFILER == 0)
     virtual void RHICalibrateTimers(FRHITimestampCalibrationQuery* CalibrationQuery) final override;
+#endif
     
 	virtual void RHIBeginRenderQuery(FRHIRenderQuery* RenderQuery) final override;
 	virtual void RHIEndRenderQuery(FRHIRenderQuery* RenderQuery) final override;
diff --git a/Engine/Source/Runtime/RHI/Public/RHICommandList.h b/Engine/Source/Runtime/RHI/Public/RHICommandList.h
--- a/Engine/Source/Runtime/RHI/Public/RHICommandList.h
+++ b/Engine/Source/Runtime/RHI/Public/RHICommandList.h
@@ -2488,6 +2488,7 @@ FRHICOMMAND_MACRO(FRHICommandEndRenderQuery)
 	RHI_EXECUTE_API void Execute(FRHICommandListBase& CmdList);
 };
 
+#if (RHI_NEW_GPU_PROFILER == 0)
 FRHICOMMAND_MACRO(FRHICommandCalibrateTimers)
 {
 	FRHITimestampCalibrationQuery* CalibrationQuery;
@@ -2498,6 +2499,7 @@ FRHICOMMAND_MACRO(FRHICommandCalibrateTimers)
 	}
 	RHI_EXECUTE_API void Execute(FRHICommandListBase & CmdList);
 };
+#endif
 
 FRHICOMMAND_MACRO(FRHICommandPostExternalCommandsReset)
 {
@@ -4229,6 +4231,7 @@ public:
 		GDynamicRHI->RHIEndRenderQuery_TopOfPipe(*this, RenderQuery);
 	}
 
+#if (RHI_NEW_GPU_PROFILER == 0)
 	FORCEINLINE_DEBUGGABLE void CalibrateTimers(FRHITimestampCalibrationQuery* CalibrationQuery)
 	{
 		if (Bypass())
@@ -4238,6 +4241,7 @@ public:
 		}
 		ALLOC_COMMAND(FRHICommandCalibrateTimers)(CalibrationQuery);
 	}
+#endif
 
 	FORCEINLINE_DEBUGGABLE void BeginRenderPass(const FRHIRenderPassInfo& InInfo, const TCHAR* Name)
 	{
diff --git a/Engine/Source/Runtime/RHI/Public/RHICommandListCommandExecutes.inl b/Engine/Source/Runtime/RHI/Public/RHICommandListCommandExecutes.inl
--- a/Engine/Source/Runtime/RHI/Public/RHICommandListCommandExecutes.inl
+++ b/Engine/Source/Runtime/RHI/Public/RHICommandListCommandExecutes.inl
@@ -478,11 +478,13 @@ void FRHICommandEndRenderQuery::Execute(FRHICommandListBase& CmdList)
 	INTERNAL_DECORATOR(RHIEndRenderQuery)(RenderQuery);
 }
 
+#if (RHI_NEW_GPU_PROFILER == 0)
 void FRHICommandCalibrateTimers::Execute(FRHICommandListBase& CmdList)
 {
 	RHISTAT(CalibrateTimers);
 	INTERNAL_DECORATOR(RHICalibrateTimers)(CalibrationQuery);
 }
+#endif
 
 void FRHICommandPostExternalCommandsReset::Execute(FRHICommandListBase& CmdList)
 {
diff --git a/Engine/Source/Runtime/RHI/Public/RHIContext.h b/Engine/Source/Runtime/RHI/Public/RHIContext.h
--- a/Engine/Source/Runtime/RHI/Public/RHIContext.h
+++ b/Engine/Source/Runtime/RHI/Public/RHIContext.h
@@ -671,10 +671,12 @@ public:
 
 	virtual void RHIEndRenderQuery(FRHIRenderQuery* RenderQuery) = 0;
 
+#if (RHI_NEW_GPU_PROFILER == 0)
 	virtual void RHICalibrateTimers(FRHITimestampCalibrationQuery* CalibrationQuery)
 	{
 		/* empty default implementation */
 	}
+#endif
 
 	// Not all RHIs need this (Mobile specific)
 	virtual void RHIDiscardRenderTargets(bool Depth, bool Stencil, uint32 ColorBitMask) {}
diff --git a/Engine/Source/Runtime/RHI/Public/RHIFwd.h b/Engine/Source/Runtime/RHI/Public/RHIFwd.h
--- a/Engine/Source/Runtime/RHI/Public/RHIFwd.h
+++ b/Engine/Source/Runtime/RHI/Public/RHIFwd.h
@@ -82,7 +82,9 @@ class FRHIShaderBundle;
 class FRHIStagingBuffer;
 class FRHITexture;
 class FRHITextureReference;
+#if (RHI_NEW_GPU_PROFILER == 0)
 class FRHITimestampCalibrationQuery;
+#endif
 class FRHIUniformBuffer;
 class FRHIUnorderedAccessView;
 class FRHIVertexDeclaration;
@@ -127,7 +129,9 @@ using FShaderBundleRHIRef              = TRefCountPtr<FRHIShaderBundle>;
 using FStagingBufferRHIRef             = TRefCountPtr<FRHIStagingBuffer>;
 using FTextureReferenceRHIRef          = TRefCountPtr<FRHITextureReference>;
 using FTextureRHIRef                   = TRefCountPtr<FRHITexture>;
+#if (RHI_NEW_GPU_PROFILER == 0)
 using FTimestampCalibrationQueryRHIRef = TRefCountPtr<FRHITimestampCalibrationQuery>;
+#endif
 using FUniformBufferLayoutRHIRef       = TRefCountPtr<const FRHIUniformBufferLayout>;
 using FUniformBufferRHIRef             = TRefCountPtr<FRHIUniformBuffer>;
 using FUnorderedAccessViewRHIRef       = TRefCountPtr<FRHIUnorderedAccessView>;
diff --git a/Engine/Source/Runtime/RHI/Public/RHIResources.h b/Engine/Source/Runtime/RHI/Public/RHIResources.h
--- a/Engine/Source/Runtime/RHI/Public/RHIResources.h
+++ b/Engine/Source/Runtime/RHI/Public/RHIResources.h
@@ -2389,6 +2389,7 @@ private:
 // Misc
 //
 
+#if (RHI_NEW_GPU_PROFILER == 0)
 class FRHITimestampCalibrationQuery : public FRHIResource
 {
 public:
@@ -2396,6 +2397,7 @@ public:
 	uint64 GPUMicroseconds[MAX_NUM_GPUS] = {};
 	uint64 CPUMicroseconds[MAX_NUM_GPUS] = {};
 };
+#endif
 
 class FRHIGPUFence : public FRHIResource
 {
diff --git a/Engine/Source/Runtime/RHI/Public/RHIValidationContext.h b/Engine/Source/Runtime/RHI/Public/RHIValidationContext.h
--- a/Engine/Source/Runtime/RHI/Public/RHIValidationContext.h
+++ b/Engine/Source/Runtime/RHI/Public/RHIValidationContext.h
@@ -745,10 +745,12 @@ public:
 		RHIContext->RHIEndRenderQuery(RenderQuery);
 	}
 
+#if (RHI_NEW_GPU_PROFILER == 0)
 	virtual void RHICalibrateTimers(FRHITimestampCalibrationQuery* CalibrationQuery) override final
 	{
 		RHIContext->RHICalibrateTimers(CalibrationQuery);
 	}
+#endif
 
 	// Not all RHIs need this (Mobile specific)
 	virtual void RHIDiscardRenderTargets(bool bDepth, bool bStencil, uint32 ColorBitMask) override final
