Separate ES31 device profile by RHI e.g. WindowsClient_ES31_D3D12RHI

The old ES31 device profile e.g. WindowsClient_ES31 is still usable for all RHI.

diff --git a/Engine/Plugins/Runtime/WindowsDeviceProfileSelector/Source/WindowsDeviceProfileSelector/Private/WindowsDeviceProfileSelectorModule.cpp b/Engine/Plugins/Runtime/WindowsDeviceProfileSelector/Source/WindowsDeviceProfileSelector/Private/WindowsDeviceProfileSelectorModule.cpp
--- a/Engine/Plugins/Runtime/WindowsDeviceProfileSelector/Source/WindowsDeviceProfileSelector/Private/WindowsDeviceProfileSelectorModule.cpp
+++ b/Engine/Plugins/Runtime/WindowsDeviceProfileSelector/Source/WindowsDeviceProfileSelector/Private/WindowsDeviceProfileSelectorModule.cpp
@@ -44,6 +44,15 @@ const FString FWindowsDeviceProfileSelectorModule::GetRuntimeDeviceProfileName()
 			// Use RHI specific device profile if exist
 			ProfileName = MoveTemp(TmpProfileName);
 		}
+		else if (TmpProfileName.Contains(TEXT("_ES31")))
+		{
+			TmpProfileName = ProfileName + TEXT("_ES31");
+			if (AvailableProfiles.Contains(FString::Printf(TEXT("%s DeviceProfile"), *TmpProfileName)))
+			{
+				// General ES31 device profile
+				ProfileName = MoveTemp(TmpProfileName);
+			}
+		}
 	}
 
 	return ProfileName;
diff --git a/Engine/Source/Runtime/RHI/Private/Windows/WindowsDynamicRHI.cpp b/Engine/Source/Runtime/RHI/Private/Windows/WindowsDynamicRHI.cpp
--- a/Engine/Source/Runtime/RHI/Private/Windows/WindowsDynamicRHI.cpp
+++ b/Engine/Source/Runtime/RHI/Private/Windows/WindowsDynamicRHI.cpp
@@ -174,6 +174,44 @@ static FString GetRHINameFromWindowsRHI(EWindowsRHI InWindowsRHI, ERHIFeatureLev
 	}
 }
 
+static const TCHAR* GetES31ModuleName(ERHIInterfaceType InterfaceType)
+{
+	switch (InterfaceType)
+	{
+	default: check(false);
+	case ERHIInterfaceType::D3D11:  return TEXT("ES31_D3D11RHI");
+	case ERHIInterfaceType::D3D12:  return TEXT("ES31_D3D12RHI");
+	case ERHIInterfaceType::Vulkan: return TEXT("ES31_VulkanRHI");
+	case ERHIInterfaceType::OpenGL: return TEXT("ES31_OpenGLDrv");
+	}
+}
+
+static const TCHAR* GetES31ModuleNameByString(const TCHAR* RHIModuleName)
+{
+	if (FCString::Strcmp(RHIModuleName, ModuleNameFromWindowsRHI(EWindowsRHI::D3D11)) == 0)
+	{
+		return GetES31ModuleName(ERHIInterfaceType::D3D11);
+	}
+
+	if (FCString::Strcmp(RHIModuleName, ModuleNameFromWindowsRHI(EWindowsRHI::D3D12)) == 0)
+	{
+		return GetES31ModuleName(ERHIInterfaceType::D3D12);
+	}
+
+	if (FCString::Strcmp(RHIModuleName, ModuleNameFromWindowsRHI(EWindowsRHI::Vulkan)) == 0)
+	{
+		return GetES31ModuleName(ERHIInterfaceType::Vulkan);
+	}
+
+	if (FCString::Strcmp(RHIModuleName, ModuleNameFromWindowsRHI(EWindowsRHI::OpenGL)) == 0)
+	{
+		return GetES31ModuleName(ERHIInterfaceType::OpenGL);
+	}
+
+	check(false);
+	return RHIModuleName;
+}
+
 static ERHIFeatureLevel::Type GetDefaultFeatureLevelForRHI(EWindowsRHI InWindowsRHI)
 {
 	switch (InWindowsRHI)
@@ -1108,7 +1146,7 @@ const TCHAR* GetSelectedDynamicRHIModuleName(bool bCleanup)
 	if (GDynamicRHI)
 	{
 		check(!!GLoadedRHIModuleName);
-		return GMaxRHIFeatureLevel == ERHIFeatureLevel::ES3_1 ? TEXT("ES31") : GLoadedRHIModuleName;
+		return GMaxRHIFeatureLevel == ERHIFeatureLevel::ES3_1 ? GetES31ModuleName(GDynamicRHI->GetInterfaceType()) : GLoadedRHIModuleName;
 	}
 	else
 	{
@@ -1122,7 +1160,7 @@ const TCHAR* GetSelectedDynamicRHIModuleName(bool bCleanup)
 			FModuleManager::Get().UnloadModule(RHIModuleName);
 		}
 
-		return DesiredFeatureLevel == ERHIFeatureLevel::ES3_1 ? TEXT("ES31") : RHIModuleName;
+		return DesiredFeatureLevel == ERHIFeatureLevel::ES3_1 ? GetES31ModuleNameByString(RHIModuleName) : RHIModuleName;
 	}
 }
 
