Restore legacy GPU profiler UI in the editor.
- Shows only work submitted to the graphics queue.
- Removed OnProfileGPU delegate that was used to show the output log window instead.

Additional cvars to control which columns and details are shown in the profileGPU log output.
- r.ProfileGPU.ShowHeader
- r.ProfileGPU.ShowEmptyQueues
- r.ProfileGPU.ShowStats
- r.ProfileGPU.ShowPercentColumn
- r.ProfileGPU.ShowInclusive
- r.ProfileGPU.ShowExclusive

Minor fixes to table text output when certain columns are hidden.

diff --git a/Engine/Source/Editor/UnrealEd/Private/UnrealEdMisc.cpp b/Engine/Source/Editor/UnrealEd/Private/UnrealEdMisc.cpp
--- a/Engine/Source/Editor/UnrealEd/Private/UnrealEdMisc.cpp
+++ b/Engine/Source/Editor/UnrealEd/Private/UnrealEdMisc.cpp
@@ -602,14 +602,6 @@ void FUnrealEdMisc::OnInit()
 	FEditorDelegates::OnEditorInitialized.Broadcast(TotalEditorStartupTime);
 	
 	GShaderCompilingManager->PrintStats();
-
-#if RHI_NEW_GPU_PROFILER
-	OnProfileGPUHandle = UE::RHI::GPUProfiler::OnProfileGPU.AddLambda([this]()
-	{
-		FOutputLogModule& OutputLogModule = FModuleManager::Get().LoadModuleChecked<FOutputLogModule>("OutputLog");
-		OutputLogModule.FocusOutputLogAndScrollToEnd();
-	});
-#endif
 }
 
 FString FUnrealEdMisc::FindMapFileFromPartialName(const FString& PartialMapName)
@@ -1065,13 +1057,6 @@ void FUnrealEdMisc::OnExit()
 		NavigationBuildingNotificationHandler = NULL;
 	}
 
-#if RHI_NEW_GPU_PROFILER
-	if (OnProfileGPUHandle.IsValid())
-	{
-		UE::RHI::GPUProfiler::OnProfileGPU.Remove(OnProfileGPUHandle);
-	}
-#endif
-
 	// Report session maximum window and tab counts to engine analytics, if available
 	if (FEngineAnalytics::IsAvailable())
 	{
diff --git a/Engine/Source/Editor/UnrealEd/Public/UnrealEdMisc.h b/Engine/Source/Editor/UnrealEd/Public/UnrealEdMisc.h
--- a/Engine/Source/Editor/UnrealEd/Public/UnrealEdMisc.h
+++ b/Engine/Source/Editor/UnrealEd/Public/UnrealEdMisc.h
@@ -356,7 +356,4 @@ private:
 	bool bActorMultiSelect = false;
 	TArray<AActor*> ActorsSelected;
 	TArray<FGuid> UnloadedActorsSelected;
-
-	/** Handle used to register the "OnProfileGPU" delegate in the RHI. */
-	FDelegateHandle OnProfileGPUHandle;
 };
diff --git a/Engine/Source/Runtime/RHI/Private/GPUProfiler.cpp b/Engine/Source/Runtime/RHI/Private/GPUProfiler.cpp
--- a/Engine/Source/Runtime/RHI/Private/GPUProfiler.cpp
+++ b/Engine/Source/Runtime/RHI/Private/GPUProfiler.cpp
@@ -65,6 +65,49 @@ static TAutoConsoleVariable<bool> GCVarProfileGPU_ShowLeafEvents(
 	TEXT("Allows profileGPU to display event-only leaf nodes with no draws associated."),
 	ECVF_Default);
 
+static TAutoConsoleVariable<bool> GCVarProfileGPU_ShowHeader(
+	TEXT("r.ProfileGPU.ShowHeader"),
+	true,
+	TEXT("When true, prints a summary of the profileGPU settings before the report table in the log."),
+	ECVF_Default);
+
+static TAutoConsoleVariable<bool> GCVarProfileGPU_ShowEmptyQueues(
+	TEXT("r.ProfileGPU.ShowEmptyQueues"),
+	true,
+	TEXT("When true, GPU queues without any registered work are still displayed in the report tables."),
+	ECVF_Default);
+
+static TAutoConsoleVariable<bool> GCVarProfileGPU_ShowStats(
+	TEXT("r.ProfileGPU.ShowStats"),
+	true,
+	TEXT("When true, additional stat columns are shown in the report (numbers of draws, dispatches, vertices and primitives)."),
+	ECVF_Default);
+
+static TAutoConsoleVariable<bool> GCVarProfileGPU_ShowPercentColumn(
+	TEXT("r.ProfileGPU.ShowPercentColumn"),
+	true,
+	TEXT("When true, a column showing the relative portion of time each stat takes as a percentage is displayed, including a visual unicode bar when unicode output is enabled."),
+	ECVF_Default);
+
+static TAutoConsoleVariable<bool> GCVarProfileGPU_ShowInclusive(
+	TEXT("r.ProfileGPU.ShowInclusive"),
+	true,
+	TEXT("When true, inclusive GPU times are shown."),
+	ECVF_Default);
+
+static TAutoConsoleVariable<bool> GCVarProfileGPU_ShowExclusive(
+	TEXT("r.ProfileGPU.ShowExclusive"),
+	true,
+	TEXT("When true, exclusive GPU times are shown."),
+	ECVF_Default);
+
+static TAutoConsoleVariable<bool> GCVarProfileGPU_ShowUI(
+	TEXT("r.ProfileGPU.ShowUI"),
+	true,
+	TEXT("Whether the user interface profiler should be displayed after profiling the GPU.\n")
+	TEXT("The results will always go to the log/console."),
+	ECVF_Default);
+
 static TAutoConsoleVariable<int> CVarGPUCsvStatsEnabled(
 	TEXT("r.GPUCsvStatsEnabled"),
 	0,
@@ -115,14 +158,6 @@ static TAutoConsoleVariable<int32> GSaveScreenshotAfterProfilingGPUCVar(
 	TEXT("Whether a screenshot should be taken when profiling the GPU. 0:off, 1:on (default)"),
 	ECVF_RenderThreadSafe);
 
-static TAutoConsoleVariable<int32> GShowProfilerAfterProfilingGPUCVar(
-	TEXT("r.ProfileGPU.ShowUI"),
-	1,
-	TEXT("Whether the user interface profiler should be displayed after profiling the GPU.\n")
-	TEXT("The results will always go to the log/console\n")
-	TEXT("0:off, 1:on (default)"),
-	ECVF_RenderThreadSafe);
-
 static TAutoConsoleVariable<float> GGPUHitchThresholdCVar(
 	TEXT("RHI.GPUHitchThreshold"),
 	100.0f,
@@ -150,7 +185,7 @@ namespace RHIConfig
 
 	bool ShouldShowProfilerAfterProfilingGPU()
 	{
-		return GShowProfilerAfterProfilingGPUCVar.GetValueOnAnyThread() != 0;
+		return GCVarProfileGPU_ShowUI.GetValueOnAnyThread() != 0;
 	}
 
 	float GetGPUHitchThreshold()
@@ -722,8 +757,6 @@ bool FGPUTiming::GAreGlobalsInitialized = false;
 
 namespace UE::RHI::GPUProfiler
 {
-	RHI_API FRHIOnProfileGPU OnProfileGPU;
-
 	TLockFreePointerListUnordered<void, PLATFORM_CACHE_LINE_SIZE> FEventStream::FChunk::MemoryPool;
 
 	static TArray<FEventSink*>& GetSinks()
@@ -892,9 +925,21 @@ namespace UE::RHI::GPUProfiler
 	struct FTable
 	{
 		bool const bUnicodeOutput;
+		bool const bShowHeader;
+		bool const bShowEmptyQueues;
+		bool const bShowStats;
+		bool const bShowPercent;
+		bool const bShowInclusive;
+		bool const bShowExclusive;
 
 		FTable()
-			: bUnicodeOutput(GCVarProfileGPU_UnicodeOutput.GetValueOnAnyThread())
+			: bUnicodeOutput  (GCVarProfileGPU_UnicodeOutput    .GetValueOnAnyThread())
+			, bShowHeader     (GCVarProfileGPU_ShowHeader       .GetValueOnAnyThread())
+			, bShowEmptyQueues(GCVarProfileGPU_ShowEmptyQueues  .GetValueOnAnyThread())
+			, bShowStats      (GCVarProfileGPU_ShowStats        .GetValueOnAnyThread())
+			, bShowPercent    (GCVarProfileGPU_ShowPercentColumn.GetValueOnAnyThread())
+			, bShowInclusive  (GCVarProfileGPU_ShowInclusive    .GetValueOnAnyThread())
+			, bShowExclusive  (GCVarProfileGPU_ShowExclusive    .GetValueOnAnyThread())
 		{}
 
 		enum class EColumn : uint32
@@ -1023,20 +1068,42 @@ namespace UE::RHI::GPUProfiler
 
 			static constexpr TCHAR const BarSeparator[] = TEXT(" ┊ ");
 
-			Col(EColumn::Exclusive_NumDraws     ) = FString::Printf(TEXT("%d"), Exclusive.NumDraws);
-			Col(EColumn::Exclusive_NumDispatches) = FString::Printf(TEXT("%d"), Exclusive.NumDispatches);
-			Col(EColumn::Exclusive_NumPrimitives) = FString::Printf(TEXT("%d"), Exclusive.NumPrimitives);
-			Col(EColumn::Exclusive_NumVertices  ) = FString::Printf(TEXT("%d"), Exclusive.NumVertices);
-			Col(EColumn::Exclusive_Percent      ) = FString::Printf(TEXT("%.1f%%%s%s"), ExclusivePercent * 100.0, bUnicodeOutput ? BarSeparator : TEXT(""), bUnicodeOutput ? ExclusiveBar.Text : TEXT(""));
-			Col(EColumn::Exclusive_Time         ) = FString::Printf(TEXT("%.3f ms"), FPlatformTime::ToMilliseconds64(Exclusive.BusyCycles));
-
-			Col(EColumn::Inclusive_NumDraws     ) = FString::Printf(TEXT("%d"), Inclusive.NumDraws);
-			Col(EColumn::Inclusive_NumDispatches) = FString::Printf(TEXT("%d"), Inclusive.NumDispatches);
-			Col(EColumn::Inclusive_NumPrimitives) = FString::Printf(TEXT("%d"), Inclusive.NumPrimitives);
-			Col(EColumn::Inclusive_NumVertices  ) = FString::Printf(TEXT("%d"), Inclusive.NumVertices);
-			Col(EColumn::Inclusive_Percent      ) = FString::Printf(TEXT("%.1f%%%s%s"), InclusivePercent * 100.0, bUnicodeOutput ? BarSeparator : TEXT(""), bUnicodeOutput ? InclusiveBar.Text : TEXT(""));
-			Col(EColumn::Inclusive_Time         ) = FString::Printf(TEXT("%.3f ms"), FPlatformTime::ToMilliseconds64(Inclusive.BusyCycles));
-				
+			if (bShowExclusive)
+			{
+				if (bShowStats)
+				{
+					Col(EColumn::Exclusive_NumDraws     ) = FString::Printf(TEXT("%d"), Exclusive.NumDraws);
+					Col(EColumn::Exclusive_NumDispatches) = FString::Printf(TEXT("%d"), Exclusive.NumDispatches);
+					Col(EColumn::Exclusive_NumPrimitives) = FString::Printf(TEXT("%d"), Exclusive.NumPrimitives);
+					Col(EColumn::Exclusive_NumVertices  ) = FString::Printf(TEXT("%d"), Exclusive.NumVertices);
+				}
+
+				if (bShowPercent)
+				{
+					Col(EColumn::Exclusive_Percent) = FString::Printf(TEXT("%.1f%%%s%s"), ExclusivePercent * 100.0, bUnicodeOutput ? BarSeparator : TEXT(""), bUnicodeOutput ? ExclusiveBar.Text : TEXT(""));
+				}
+
+				Col(EColumn::Exclusive_Time) = FString::Printf(TEXT("%.3f ms"), FPlatformTime::ToMilliseconds64(Exclusive.BusyCycles));
+			}
+
+			if (bShowInclusive)
+			{
+				if (bShowStats)
+				{
+					Col(EColumn::Inclusive_NumDraws     ) = FString::Printf(TEXT("%d"), Inclusive.NumDraws);
+					Col(EColumn::Inclusive_NumDispatches) = FString::Printf(TEXT("%d"), Inclusive.NumDispatches);
+					Col(EColumn::Inclusive_NumPrimitives) = FString::Printf(TEXT("%d"), Inclusive.NumPrimitives);
+					Col(EColumn::Inclusive_NumVertices  ) = FString::Printf(TEXT("%d"), Inclusive.NumVertices);
+				}
+
+				if (bShowPercent)
+				{
+					Col(EColumn::Inclusive_Percent) = FString::Printf(TEXT("%.1f%%%s%s"), InclusivePercent * 100.0, bUnicodeOutput ? BarSeparator : TEXT(""), bUnicodeOutput ? InclusiveBar.Text : TEXT(""));
+				}
+
+				Col(EColumn::Inclusive_Time) = FString::Printf(TEXT("%.3f ms"), FPlatformTime::ToMilliseconds64(Inclusive.BusyCycles));
+			}
+
 			static constexpr uint32 SpacesPerIndent = 3;
 			Col(EColumn::Events) = FString::Printf(TEXT("%*s"), Name.Len() + (Level * SpacesPerIndent), *Name);
 
@@ -1059,6 +1126,7 @@ namespace UE::RHI::GPUProfiler
 		{
 			TCHAR const* LineMajor;
 			TCHAR const* LineMinor;
+			TCHAR const* Indent;
 
 			FChars const TopRow;
 			FChars const GroupNameRow;
@@ -1076,14 +1144,15 @@ namespace UE::RHI::GPUProfiler
 				{
 					.LineMajor  = TEXT("━"),
 					.LineMinor  = TEXT("─"),
-
-					//                     Left     GrpSep     LastGrp     Right     CellSep
-					.TopRow        { TEXT("    ┏"), TEXT("┳"), TEXT("┳"), TEXT("┓"), TEXT(" ") },
-					.GroupNameRow  { TEXT("    ┃"), TEXT("┃"), TEXT("┃"), TEXT("┃"), TEXT(" ") },
-					.GroupBorderRow{ TEXT("    ┠"), TEXT("╂"), TEXT("┨"), TEXT("┃"), TEXT("┬") },
-					.ValueRow      { TEXT("    ┃"), TEXT("┃"), TEXT("┃"), TEXT("┃"), TEXT("│") },
-					.DividorRow    { TEXT("    ┠"), TEXT("╂"), TEXT("╂"), TEXT("┨"), TEXT("┼") },
-					.BottomRow     { TEXT("    ┗"), TEXT("┻"), TEXT("┻"), TEXT("┛"), TEXT("┷") },
+					.Indent     = TEXT("    "),
+
+					//                 Left     GrpSep     LastGrp     Right     CellSep
+					.TopRow        { TEXT("┏"), TEXT("┳"), TEXT("┳"), TEXT("┓"), TEXT(" ") },
+					.GroupNameRow  { TEXT("┃"), TEXT("┃"), TEXT("┃"), TEXT("┃"), TEXT(" ") },
+					.GroupBorderRow{ TEXT("┠"), TEXT("╂"), TEXT("┨"), TEXT("┃"), TEXT("┬") },
+					.ValueRow      { TEXT("┃"), TEXT("┃"), TEXT("┃"), TEXT("┃"), TEXT("│") },
+					.DividorRow    { TEXT("┠"), TEXT("╂"), TEXT("╂"), TEXT("┨"), TEXT("┼") },
+					.BottomRow     { TEXT("┗"), TEXT("┻"), TEXT("┻"), TEXT("┛"), TEXT("┷") },
 				};
 
 				return ToStringInner(Unicode);
@@ -1094,14 +1163,15 @@ namespace UE::RHI::GPUProfiler
 				{
 					.LineMajor  = TEXT("-"),
 					.LineMinor  = TEXT("-"),
-
-					//                     Left     GrpSep     LastGrp     Right     CellSep
-					.TopRow        { TEXT("    +"), TEXT("+"), TEXT("+"), TEXT("+"), TEXT(" ") },
-					.GroupNameRow  { TEXT("    |"), TEXT("|"), TEXT("|"), TEXT("|"), TEXT(" ") },
-					.GroupBorderRow{ TEXT("    +"), TEXT("+"), TEXT("+"), TEXT("|"), TEXT("+") },
-					.ValueRow      { TEXT("    |"), TEXT("|"), TEXT("|"), TEXT("|"), TEXT("|") },
-					.DividorRow    { TEXT("    +"), TEXT("+"), TEXT("+"), TEXT("+"), TEXT("+") },
-					.BottomRow     { TEXT("    +"), TEXT("+"), TEXT("+"), TEXT("+"), TEXT("+") },
+					.Indent = TEXT("    "),
+
+					//                 Left     GrpSep     LastGrp     Right     CellSep
+					.TopRow        { TEXT("+"), TEXT("+"), TEXT("+"), TEXT("+"), TEXT(" ") },
+					.GroupNameRow  { TEXT("|"), TEXT("|"), TEXT("|"), TEXT("|"), TEXT(" ") },
+					.GroupBorderRow{ TEXT("+"), TEXT("+"), TEXT("+"), TEXT("|"), TEXT("+") },
+					.ValueRow      { TEXT("|"), TEXT("|"), TEXT("|"), TEXT("|"), TEXT("|") },
+					.DividorRow    { TEXT("+"), TEXT("+"), TEXT("+"), TEXT("+"), TEXT("+") },
+					.BottomRow     { TEXT("+"), TEXT("+"), TEXT("+"), TEXT("+"), TEXT("+") },
 				};
 
 				return ToStringInner(Ascii);
@@ -1143,6 +1213,7 @@ namespace UE::RHI::GPUProfiler
 			{
 				uint32 const CellSeparatorLength = FCString::Strlen(Chars.CellSeparator);
 
+				Result += Format.Indent;
 				Result += Chars.Left;
 
 				uint32 GroupWidth = 0;
@@ -1183,11 +1254,30 @@ namespace UE::RHI::GPUProfiler
 				Result += TEXT("\n");
 			};
 
-			auto EmitValueRow = [&](FChars const& Chars, TUniqueFunction<void(FColumn)> CellCallback)
+			auto EmitValueRow = [&](bool bGroupBorderRow, FChars const& Chars, TUniqueFunction<void(FColumn)> CellCallback)
 			{
-				Result += Chars.Left;
+				Result += Format.Indent;
 
-				for (uint32 ColumnIndex = 0; ColumnIndex < uint32(EColumn::Num); ++ColumnIndex)
+				uint32 FirstColumnIndex = 0;
+
+				if (bGroupBorderRow)
+				{
+					// Find first visible column
+					while (FirstColumnIndex < uint32(EColumn::Num) && Columns[FirstColumnIndex].Num() == 0)
+					{
+						++FirstColumnIndex;
+					}
+
+					Result += FirstColumnIndex < uint32(EColumn::Events)
+						? Chars.Left
+						: Chars.Right;
+				}
+				else
+				{
+					Result += Chars.Left;
+				}
+
+				for (uint32 ColumnIndex = FirstColumnIndex; ColumnIndex < uint32(EColumn::Num); ++ColumnIndex)
 				{
 					if (Columns[ColumnIndex].Num() == 0)
 						continue;
@@ -1196,11 +1286,20 @@ namespace UE::RHI::GPUProfiler
 
 					if (ColumnIndex < uint32(EColumn::Num) - 1)
 					{
-						uint32 GroupIndex = GetColumnGroup(EColumn(ColumnIndex));
-						if (GroupIndex != GetColumnGroup(EColumn(ColumnIndex + 1)))
+						// Find next visible column
+						uint32 NextColumnIndex = ColumnIndex + 1;
+						while (NextColumnIndex < uint32(EColumn::Num) && Columns[NextColumnIndex].Num() == 0)
+						{
+							++NextColumnIndex;
+						}
+
+						uint32 CurrentGroupIndex = GetColumnGroup(EColumn(ColumnIndex));
+						uint32 NextGroupIndex    = GetColumnGroup(EColumn(NextColumnIndex));
+
+						if (CurrentGroupIndex != NextGroupIndex)
 						{
 							// Group change, add the group separator
-							Result += GroupIndex < NumGroups - 2
+							Result += NextGroupIndex < NumGroups - 1
 								? Chars.GroupSeparator
 								: Chars.LastGroupSeparator;
 						}
@@ -1246,7 +1345,7 @@ namespace UE::RHI::GPUProfiler
 			});
 
 			// Events Group Row
-			EmitValueRow(Format.GroupBorderRow, [&](FColumn Column)
+			EmitValueRow(true, Format.GroupBorderRow, [&](FColumn Column)
 			{
 				if (Column.Index == uint32(EColumn::Events))
 				{
@@ -1262,13 +1361,13 @@ namespace UE::RHI::GPUProfiler
 			});
 
 			// Header Row
-			EmitValueRow(Format.ValueRow, [&](FColumn Column)
+			EmitValueRow(false, Format.ValueRow, [&](FColumn Column)
 			{
 				AlignCenter(GetColumnHeader(EColumn(Column.Index)), Column.Width);
 			});
 
 			// Header Border Row
-			EmitValueRow(Format.DividorRow, [&](FColumn Column)
+			EmitValueRow(false, Format.DividorRow, [&](FColumn Column)
 			{
 				while (Column.Width--)
 				{
@@ -1282,7 +1381,7 @@ namespace UE::RHI::GPUProfiler
 				if (RowIndex > 0 && RowBreaks[RowIndex])
 				{
 					// Add a horizontal rule
-					EmitValueRow(Format.DividorRow, [&](FColumn Column)
+					EmitValueRow(false, Format.DividorRow, [&](FColumn Column)
 					{
 						while (Column.Width--)
 						{
@@ -1291,7 +1390,7 @@ namespace UE::RHI::GPUProfiler
 					});
 				}
 
-				EmitValueRow(Format.ValueRow, [&](FColumn Column)
+				EmitValueRow(false, Format.ValueRow, [&](FColumn Column)
 				{
 					int32 Width = Column.Width - (CellPadding * 2);
 					if (EColumn(Column.Index) == EColumn::Events)
@@ -1308,7 +1407,7 @@ namespace UE::RHI::GPUProfiler
 			}
 
 			// Bottom Border
-			EmitValueRow(Format.BottomRow, [&](FColumn Column)
+			EmitValueRow(false, Format.BottomRow, [&](FColumn Column)
 			{
 				while (Column.Width--)
 				{
@@ -1790,38 +1889,112 @@ namespace UE::RHI::GPUProfiler
 						Recurse(Recurse, nullptr, First, false, 0);
 					}
 
+					FString Header;
+					if (Table.bShowHeader)
+					{
+						Header = FString::Printf(
+							TEXT("    - %-30s: %.2fms\n")
+							TEXT("    - %-30s: \"%s\"\n")
+							TEXT("    - %-30s: %.2f%%\n")
+							TEXT("    - %-30s: %s\n")
+							TEXT("\n")
+							, TEXT("Frame Time")
+							, First ? First->Inclusive.GetBusyMilliseconds() : 0.0
+							, *IConsoleManager::Get().FindConsoleObjectName(GCVarProfileGPU_Root.AsVariable())
+							, *RootWildcard
+							, *IConsoleManager::Get().FindConsoleObjectName(GCVarProfileGPU_ThresholdPercent.AsVariable())
+							, PercentThreshold
+							, *IConsoleManager::Get().FindConsoleObjectName(GCVarProfileGPU_ShowLeafEvents.AsVariable())
+							, bShowEmptyNodes ? TEXT("true") : TEXT("false")
+						);
+					}
+
+					FString GPUString = FString::Printf(TEXT("Frame %d - %s %d - GPU %d")
+						, FrameNumber
+						, QueueState.Queue.GetTypeString()
+						, QueueState.Queue.Index
+						, QueueState.Queue.GPU
+					);
+
 					FString Final = FString::Printf(
 						TEXT("\n")
-						TEXT("GPU Profile for Frame %d - GPU %d - %s %d\n")
-						TEXT("\n")
-						TEXT("    - %-30s: %.2fms\n")
-						TEXT("    - %-30s: \"%s\"\n")
-						TEXT("    - %-30s: %.2f%%\n")
-						TEXT("    - %-30s: %s\n")
+						TEXT("GPU Profile for %s\n")
 						TEXT("\n")
 						TEXT("%s")
-						, FrameNumber
-						, QueueState.Queue.GPU
-						, QueueState.Queue.GetTypeString()
-						, QueueState.Queue.Index
-						, TEXT("Frame Time")
-						, First ? First->Inclusive.GetBusyMilliseconds() : 0.0
-						, *IConsoleManager::Get().FindConsoleObjectName(GCVarProfileGPU_Root.AsVariable())
-						, *RootWildcard
-						, *IConsoleManager::Get().FindConsoleObjectName(GCVarProfileGPU_ThresholdPercent.AsVariable())
-						, PercentThreshold
-						, *IConsoleManager::Get().FindConsoleObjectName(GCVarProfileGPU_ShowLeafEvents.AsVariable())
-						, bShowEmptyNodes ? TEXT("true") : TEXT("false")
+						TEXT("%s")
+						, *GPUString
+						, *Header
 						, Table.HasRows() ? *Table.ToString() : TEXT("    No recorded work for this queue.\n")
 					);
 
-					TArray<FString> Lines;
-					Final.ParseIntoArrayLines(Lines, false);
+					if (Table.HasRows() || Table.bShowEmptyQueues)
+					{
+						TArray<FString> Lines;
+						Final.ParseIntoArrayLines(Lines, false);
 
-					for (FString const& Line : Lines)
+						for (FString const& Line : Lines)
+						{
+							UE_LOG(LogRHI, Display, TEXT("%s"), *Line);
+						}
+					}
+
+				#if !UE_BUILD_SHIPPING
+					// Create and display profile visualizer data for the graphics queue
+					if (GCVarProfileGPU_ShowUI.GetValueOnAnyThread() && QueueState.Queue.Type == FQueue::EType::Graphics && Table.HasRows())
 					{
-						UE_LOG(LogRHI, Display, TEXT("%s"), *Line);
+						// Count the total number of exclusive cycles in the frame. Needed to draw the horizontal bars in the visualizer.
+						uint64 TotalBusyCycles = 0;
+						{
+							auto CountCycles = [&TotalBusyCycles](auto& CountCycles, FNode* CurrentNode) -> void
+							{
+								TotalBusyCycles += CurrentNode->Exclusive.BusyCycles;
+								for (FNode* Child : CurrentNode->Children)
+								{
+									CountCycles(CountCycles, Child);
+								}
+							};
+							CountCycles(CountCycles, First);
+						}
+
+						// Recursive function to build the visualizer data structs from our nodes.
+						uint64 StartCycles = 0;
+						auto BuildData = [&](auto& BuildData, FNode* CurrentNode, TSharedPtr<FVisualizerEvent> const& Parent) -> TSharedPtr<FVisualizerEvent>
+						{
+							TSharedPtr<FVisualizerEvent> VisualizerEvent = MakeShared<FVisualizerEvent>(
+								double(StartCycles) / TotalBusyCycles,
+								double(CurrentNode->Inclusive.BusyCycles) / TotalBusyCycles,
+								CurrentNode->Inclusive.GetBusyMilliseconds(),
+								0,
+								CurrentNode->Name
+							);
+							VisualizerEvent->ParentEvent = Parent;
+
+							StartCycles += CurrentNode->Exclusive.BusyCycles;
+
+							for (FNode* ChildNode : CurrentNode->Children)
+							{
+								VisualizerEvent->Children.Add(BuildData(BuildData, ChildNode, VisualizerEvent));
+							}
+
+							return VisualizerEvent;
+						};
+
+						// Launch the visualizer on the game thread
+						FFunctionGraphTask::CreateAndDispatchWhenReady([
+							Title = GPUString,
+							VisualizerData = BuildData(BuildData, First, nullptr)
+						]
+						{
+							static FName ProfileVisualizerModule(TEXT("ProfileVisualizer"));
+							if (FModuleManager::Get().IsModuleLoaded(ProfileVisualizerModule))
+							{
+								IProfileVisualizerModule& ProfileVisualizer = FModuleManager::GetModuleChecked<IProfileVisualizerModule>(ProfileVisualizerModule);
+								// Display a warning if this is a GPU profile and the GPU was profiled with v-sync enabled (otherwise InVsyncEnabledWarningText is empty)
+								ProfileVisualizer.DisplayProfileVisualizer(VisualizerData, *Title);
+							}
+						}, TStatId(), nullptr, ENamedThreads::GameThread);
 					}
+				#endif
 				}
 			} Profile;
 		#endif
@@ -2473,11 +2646,6 @@ namespace UE::RHI::GPUProfiler
 		FConsoleCommandWithArgsDelegate::CreateLambda([](const TArray<FString>& Args)
 		{
 			GGPUProfilerSink_StatSystem.bTriggerProfile = true;
-
-			if (OnProfileGPU.IsBound())
-			{
-				OnProfileGPU.Broadcast();
-			}
 		}));
 #endif // WITH_PROFILEGPU
 }
diff --git a/Engine/Source/Runtime/RHI/Public/GPUProfiler.h b/Engine/Source/Runtime/RHI/Public/GPUProfiler.h
--- a/Engine/Source/Runtime/RHI/Public/GPUProfiler.h
+++ b/Engine/Source/Runtime/RHI/Public/GPUProfiler.h
@@ -21,9 +21,6 @@
 
 namespace UE::RHI::GPUProfiler
 {
-	DECLARE_MULTICAST_DELEGATE(FRHIOnProfileGPU);
-	extern RHI_API FRHIOnProfileGPU OnProfileGPU;
-
 	struct FQueue
 	{
 		enum class EType : uint8
