Linux should force a SM6 FL when any bindless command line argument is passed in.

diff --git a/Engine/Source/Runtime/RHI/Private/Linux/LinuxDynamicRHI.cpp b/Engine/Source/Runtime/RHI/Private/Linux/LinuxDynamicRHI.cpp
--- a/Engine/Source/Runtime/RHI/Private/Linux/LinuxDynamicRHI.cpp
+++ b/Engine/Source/Runtime/RHI/Private/Linux/LinuxDynamicRHI.cpp
@@ -30,6 +30,13 @@ static TOptional<ERHIFeatureLevel::Type> GetForcedFeatureLevel()
 		ForcedFeatureLevel = ERHIFeatureLevel::SM6;
 	}
 
+#if WITH_EDITOR
+	if (TOptional<ERHIBindlessConfiguration> ForcedConfig = RHIGetForcedBindlessConfiguration(); ForcedConfig && ForcedConfig != ERHIBindlessConfiguration::Disabled)
+	{
+		ForcedFeatureLevel = ERHIFeatureLevel::SM6;
+	}
+#endif
+
 	return ForcedFeatureLevel;
 }
 
