Cleanup used precached PSOs, when they are kept in memory to workaround driver cache inefficiencies.

Once a PSO is needed for rendering, the PSO precache system sets a flag that marks the precached PSO (if existing) as used, and enqueues it for cleanup if it is currently being held in memory. This is enabled with r.PSOPrecache.KeepInMemoryUntilUsed=2.

Also remove some unnecessary duplicated code now that we have compute initializers, and clearly separate the implementation for the compile task between graphics and compute. This clean up work should be a "no fuctional change" cleanup but comes with opportunities to simplify the code further down the line.

#tests tested locally

diff --git a/Engine/Source/Runtime/RHI/Private/PipelineStateCache.cpp b/Engine/Source/Runtime/RHI/Private/PipelineStateCache.cpp
--- a/Engine/Source/Runtime/RHI/Private/PipelineStateCache.cpp
+++ b/Engine/Source/Runtime/RHI/Private/PipelineStateCache.cpp
@@ -334,9 +334,13 @@ int32 GPSOPrecacheKeepInMemoryUntilUsed = 0;
 static FAutoConsoleVariableRef CVarPSOPrecacheKeepInMemoryUntilUsed(
 	TEXT("r.PSOPrecache.KeepInMemoryUntilUsed"),
 	GPSOPrecacheKeepInMemoryUntilUsed,
-	TEXT("If enabled and if the underlying GPU vendor is NVIDIA, precached PSOs will be kept in memory instead of being deleted immediately after creation, and will only be deleted once they are actually used for rendering.\n")
+	TEXT("If enabled and if the underlying GPU vendor is NVIDIA, precached PSOs will be kept in memory instead of being deleted immediately after creation, and will only be deleted once they are actually used for rendering or when the maximum limit of in-memory PSOs is reached.\n")
 	TEXT("This can speed up the re-creation of precached PSOs for NVIDIA drivers and avoid small hitches, at the cost of memory.\n")
-	TEXT("It's recommended to set r.PSOPrecache.KeepInMemoryGraphicsMaxNum and r.PSOPrecache.KeepInMemoryComputeMaxNum to a non-zero value to ensure the number of in-memory PSOs is bounded."),
+	TEXT("It's recommended to set r.PSOPrecache.KeepInMemoryGraphicsMaxNum and r.PSOPrecache.KeepInMemoryComputeMaxNum to a non-zero value to ensure the number of in-memory PSOs is bounded.\n")
+	TEXT("Valid options:\n")
+	TEXT("0 = off (default)\n")
+	TEXT("1 = PSOs are kept in memory after precaching but not deleted immediately when used by the renderer. They are instead only evicted when the limit of in-memory PSOs is reached\n")
+	TEXT("2 = PSOs are kept in memory after precaching and deleted when used by the renderer"),
 	ECVF_ReadOnly);
 
 int32 GPSOPrecacheKeepInMemoryGraphicsMaxNum = 2000;
@@ -360,6 +364,11 @@ bool ShouldKeepPrecachedPSOsInMemory()
 	return GPSOPrecacheKeepInMemoryUntilUsed && IsRHIDeviceNVIDIA();
 }
 
+bool ShouldTrackUsedPrecachedPSOs()
+{
+	return ShouldKeepPrecachedPSOsInMemory() && GPSOPrecacheKeepInMemoryUntilUsed == 2;
+}
+
 const TCHAR* LexToString(EPSOPrecacheResult Result)
 {
 	switch (Result)
@@ -536,18 +545,18 @@ static void HandlePipelineCreationFailure(const FGraphicsPipelineStateInitialize
 
 // Prints out information about a failed compute pipeline compilation.
 // This is fatal unless the compilation request is coming from the precaching system.
-static void HandlePipelineCreationFailure(const FRHIComputeShader* ComputeShader, bool bPrecache)
+static void HandlePipelineCreationFailure(const FComputePipelineStateInitializer& Init)
 {
 	// Dump the shader hash so it can be looked up in the SHK data. See the previous function for details.
-	UE_LOG(LogRHI, Error, TEXT("Failed to create compute pipeline with hash %s."), *ComputeShader->GetHash().ToString());
+	UE_LOG(LogRHI, Error, TEXT("Failed to create compute pipeline with hash %s."), *Init.ComputeShader->GetHash().ToString());
 
 #if !(UE_BUILD_SHIPPING || UE_BUILD_TEST)
-	UE_LOG(LogRHI, Error, TEXT("Shader: %s"), ComputeShader->GetShaderName());
+	UE_LOG(LogRHI, Error, TEXT("Shader: %s"), Init.ComputeShader->GetShaderName());
 #endif
 
-	if (!bPrecache)
+	if (!Init.bPSOPrecache && !Init.bFromPSOFileCache)
 	{
-		// Same as above, precache failures are not fatal.
+		// Precache requests are allowed to fail, but if the PSO is needed by a draw/dispatch command, we cannot continue.
 		UE_LOG(LogRHI, Fatal, TEXT("Shader compilation failures are Fatal."));
 	}
 }
@@ -895,8 +904,8 @@ public:
 class FComputePipelineState : public FPipelineStateAsync
 {
 public:
-	FComputePipelineState(FRHIComputeShader* InComputeShader, const TCHAR* InName)
-		: ComputeShader(InComputeShader), Name(InName)
+	FComputePipelineState(FRHIComputeShader* InComputeShader)
+		: ComputeShader(InComputeShader)
 	{
 		ComputeShader->AddRef();
 	}
@@ -935,7 +944,6 @@ public:
 	}
 
 	FRHIComputeShader* ComputeShader;
-	const TCHAR* Name;
 	TRefCountPtr<FRHIComputePipelineState> RHIPipeline;
 #if PIPELINESTATECACHE_VERIFYTHREADSAFE
 	FThreadSafeCounter InUseCount;
@@ -1196,7 +1204,8 @@ int32 FindRayTracingMissShaderIndex(FRayTracingPipelineState* Pipeline, FRHIRayT
 	return INDEX_NONE;
 }
 
-bool IsPrecachedPSO(const FGraphicsPipelineStateInitializer& Initializer)
+template <typename TPipelineInitializer>
+bool IsPrecachedPSO(const TPipelineInitializer& Initializer)
 {
 	return Initializer.bFromPSOFileCache || Initializer.bPSOPrecache;
 }
@@ -2065,6 +2074,7 @@ public:
 		{
 			InMemoryPSOIndices.Reserve(InMaxInMemoryPSOs);
 		}
+		PrecachedPSOsToCleanup.Reserve(100);
 	}
 
 	~TPrecachePipelineCacheBase()
@@ -2277,13 +2287,25 @@ public:
 			FRWScopeLock ReadLock(PrecachePSOsRWLock, SLT_ReadOnly);
 			InitializerHash = PrecachedPSOInitializers[RequestID.RequestID];
 		}
-		return GetPrecachingStateInternal(InitializerHash);
+		bool bMarkUsed = false;
+		bool bOutMarkedUsed = false;
+		return GetPrecachingStateInternal(InitializerHash, bMarkUsed, bOutMarkedUsed);
 	}
 
 	EPSOPrecacheResult GetPrecachingState(const TPrecachedPSOInitializer& Initializer)
 	{
 		uint64 InitializerHash = TPrecachePipelineCacheDerived::PipelineStateInitializerHash(Initializer);
-		return GetPrecachingStateInternal(InitializerHash);
+		bool bMarkUsed = false;
+		bool bOutMarkedUsed = false;
+		return GetPrecachingStateInternal(InitializerHash, bMarkUsed, bOutMarkedUsed);
+	}
+
+	// Retrieves the precaching state for the given PSO, and marks it as used if not already marked.
+	EPSOPrecacheResult GetPrecachingStateAndMarkUsed(const TPrecachedPSOInitializer& Initializer, bool& bOutMarkedUsed)
+	{
+		uint64 InitializerHash = TPrecachePipelineCacheDerived::PipelineStateInitializerHash(Initializer);
+		bool bMarkUsed = true;
+		return GetPrecachingStateInternal(InitializerHash, bMarkUsed, bOutMarkedUsed);
 	}
 
 	bool IsPrecaching()
@@ -2377,7 +2399,6 @@ public:
 		SET_DWORD_STAT_FName(TPrecachePipelineCacheDerived::GetActiveCompileStatName(), ActiveCompileCount);
 		SET_DWORD_STAT_FName(TPrecachePipelineCacheDerived::GetHighPriorityCompileStatName(), HighPriorityCompileCount);
 		SET_DWORD_STAT_FName(TPrecachePipelineCacheDerived::GetHighestPriorityCompileStatName(), HighestPriorityCompileCount);
-		
 		FRWScopeLock WriteLock(PrecachePSOsRWLock, SLT_Write);
 		for (int32 Index = 0; Index < PrecachedPSOsToCleanup.Num(); ++Index)
 		{
@@ -2413,28 +2434,15 @@ public:
 	// Does not do anything if PSOs are not kept in memory, or if a bound on the number of PSOs kept in memory is set.
 	void MarkPSOAsUsed(const TPrecachedPSOInitializer& Initializer)
 	{
-		if (!ShouldKeepPrecachedPSOsInMemory() || MaxInMemoryPSOs > 0)
+		if (!ShouldTrackUsedPrecachedPSOs())
 		{
 			return;
 		}
 
 		uint64 InitializerHash = TPrecachePipelineCacheDerived::PipelineStateInitializerHash(Initializer);
-		bool bMarkedAsUsed = false;
-		{
-			FRWScopeLock ReadOnlyLock(PrecachePSOsRWLock, SLT_ReadOnly);
-			FPrecacheTask* FindResult = PrecachedPSOInitializerData.Find(InitializerHash);
-			if (FindResult)
-			{
-				EPSOPrecacheStateMask PreviousState = FindResult->AddPSOPrecacheState(EPSOPrecacheStateMask::UsedForRendering);
-				bMarkedAsUsed = !EnumHasAnyFlags(PreviousState, EPSOPrecacheStateMask::UsedForRendering);
-			}
-		}
 
-		if (bMarkedAsUsed)
-		{
-			FRWScopeLock WriteLock(PrecachePSOsRWLock, SLT_Write);
-			PrecachedPSOsToCleanup.Add(InitializerHash);
-		}
+		FRWScopeLock WriteLock(PrecachePSOsRWLock, SLT_Write);
+		PrecachedPSOsToCleanup.Add(InitializerHash);
 	}
 
 protected:
@@ -2456,7 +2464,7 @@ protected:
 		return false;
 	}
 
-	EPSOPrecacheResult GetPrecachingStateInternal(uint64 InitializerHash)
+	EPSOPrecacheResult GetPrecachingStateInternal(uint64 InitializerHash, bool bTryMarkUsed, bool& bOutMarkedUsed)
 	{
 		EPSOPrecacheStateMask CompilationState = EPSOPrecacheStateMask::None;
 		{
@@ -2467,9 +2475,19 @@ protected:
 				return EPSOPrecacheResult::Missed;
 			}
 
-			CompilationState = FindResult->ReadPSOPrecacheState();
+			if (bTryMarkUsed)
+			{
+				CompilationState = FindResult->AddPSOPrecacheState(EPSOPrecacheStateMask::UsedForRendering);
+			}
+			else
+			{
+				CompilationState = FindResult->ReadPSOPrecacheState();
+			}
 		}
 
+		// We successfully marked it as used if nobody else did concurrently.
+		bOutMarkedUsed = bTryMarkUsed && !EnumHasAnyFlags(CompilationState, EPSOPrecacheStateMask::UsedForRendering);
+
 		if (!IsCompilationDone(CompilationState))
 		{
 			return EPSOPrecacheResult::Active;
@@ -2569,32 +2587,17 @@ protected:
 	uint32 MaxInMemoryPSOs = 0;
 };
 
-struct FPrecacheComputeInitializer
-{
-	FPrecacheComputeInitializer()
-	: RHIComputeShaderAsU64(0)
-	{}
-
-
-	FPrecacheComputeInitializer(const FRHIComputeShader* InRHIComputeShader, const TCHAR* InName)
-		: RHIComputeShaderAsU64((uint64)InRHIComputeShader), Name(InName)
-	{}
-
-	uint64 RHIComputeShaderAsU64; // Using a U64 here rather than FRHIComputeShader* because we keep FPrecacheComputeInitializer but the FRHIComputeShader might get deleted
-	const TCHAR* Name;
-};
-
-class FPrecacheComputePipelineCache : public TPrecachePipelineCacheBase<FPrecacheComputePipelineCache, FPrecacheComputeInitializer, FComputePipelineState>
+class FPrecacheComputePipelineCache : public TPrecachePipelineCacheBase<FPrecacheComputePipelineCache, FComputePipelineStateInitializer, FComputePipelineState>
 {
 public:
-	static FComputePipelineState* CreateNewPSO(const FPrecacheComputeInitializer& ComputeShaderInitializer)
+	static FComputePipelineState* CreateNewPSO(const FComputePipelineStateInitializer& ComputeShaderInitializer)
 	{
-		return new FComputePipelineState((FRHIComputeShader*)ComputeShaderInitializer.RHIComputeShaderAsU64, ComputeShaderInitializer.Name);
+		return new FComputePipelineState(ComputeShaderInitializer.ComputeShader);
 	}
 
 	FPrecacheComputePipelineCache(uint32 InMemoryPSOsMaxNum) : TPrecachePipelineCacheBase(FPSOPrecacheRequestID::EType::Compute, InMemoryPSOsMaxNum) {}
 	FPSOPrecacheRequestResult PrecacheComputePipelineState(FRHIComputeShader* ComputeShader, const TCHAR* Name, bool bForcePrecache);
-	static void OnNewPipelineStateCreated(const FPrecacheComputeInitializer& ComputeInitializer, FComputePipelineState* NewComputePipelineState, const FString& PSOCompilationEventName, bool bDoAsyncCompile);
+	static void OnNewPipelineStateCreated(const FComputePipelineStateInitializer& ComputeInitializer, FComputePipelineState* NewComputePipelineState, const FString& PSOCompilationEventName, bool bDoAsyncCompile);
  
 	static const FName GetActiveCompileStatName()
 	{
@@ -2609,9 +2612,9 @@ public:
 		return GET_STATFNAME(STAT_HighestPriorityComputePSOPrecacheRequests);
 	}
 
-	static FORCEINLINE uint64 PipelineStateInitializerHash(const FPrecacheComputeInitializer& Key)
+	static FORCEINLINE uint64 PipelineStateInitializerHash(const FComputePipelineStateInitializer& Key)
 	{
-		return Key.RHIComputeShaderAsU64;
+		return (uint64)Key.ComputeShader;
 	}
 };
 
@@ -2910,115 +2913,41 @@ int32 PipelineStateCache::GetNumActivePipelinePrecompileTasks()
 	return GPipelinePrecompileTasksInFlight.load();
 }
 
-class FCompilePipelineStateTask
+template <typename TCompilePipelineStateTaskDerived, typename TPipelineInitializer>
+class TCompilePipelineStateTaskBase
 {
-public:
+protected:
 	FPipelineStateAsync* Pipeline;
-	FGraphicsPipelineStateInitializer Initializer;
+	TPipelineInitializer Initializer;
 	EPSOPrecacheResult PSOPrecacheResult;
+	bool bPrecachedPSOMarkedUsed;
 	bool bInImmediateCmdList;
 	FPSOCompilationDebugData PSOCompilationDebugData;
 
-	// InInitializer is only used for non-compute tasks, a default can just be used otherwise
-	FCompilePipelineStateTask(
-		FPipelineStateAsync* InPipeline, 
-		const FGraphicsPipelineStateInitializer& InInitializer, 
-		EPSOPrecacheResult InPSOPrecacheResult, 
+
+public:
+	TCompilePipelineStateTaskBase(
+		FPipelineStateAsync* InPipeline,
+		const TPipelineInitializer& InInitializer,
+		EPSOPrecacheResult InPSOPrecacheResult,
+		bool InbPrecachedPSOMarkedUsed,
 		bool InbInImmediateCmdList,
 		const FPSOCompilationDebugData& InPSOCompilationDebugData)
 		: Pipeline(InPipeline)
 		, Initializer(InInitializer)
 		, PSOPrecacheResult(InPSOPrecacheResult)
+		, bPrecachedPSOMarkedUsed(InbPrecachedPSOMarkedUsed)
 		, bInImmediateCmdList(InbInImmediateCmdList)
 		, PSOCompilationDebugData(InPSOCompilationDebugData)
 	{
 		ensure(Pipeline->GetCompletionEvent() != nullptr);
-		if(Initializer.bFromPSOFileCache)
+		if (Initializer.bFromPSOFileCache)
 		{
 			GPipelinePrecompileTasksInFlight++;
 		}
-				
-#if PLATFORM_WINDOWS
-		auto MarkInUseByPSOCompilation = [](FRHIShader* Shader)
-		{
-			if (Shader)
-			{
-				Shader->SetInUseByPSOCompilation(true);
-			}
-		};
-#endif // PLATFORM_WINDOWS
-
-		if (!Pipeline->IsCompute())
-		{
-#if PLATFORM_WINDOWS
-			MarkInUseByPSOCompilation(Initializer.BoundShaderState.GetMeshShader());
-			MarkInUseByPSOCompilation(Initializer.BoundShaderState.GetAmplificationShader());
-			MarkInUseByPSOCompilation(Initializer.BoundShaderState.VertexShaderRHI);
-			MarkInUseByPSOCompilation(Initializer.BoundShaderState.PixelShaderRHI);
-			MarkInUseByPSOCompilation(Initializer.BoundShaderState.GetGeometryShader());
-			MarkInUseByPSOCompilation(Initializer.BoundShaderState.GetMeshShader());
-#endif // PLATFORM_WINDOWS
-
-			if (Initializer.BoundShaderState.GetMeshShader())
-			{
-				Initializer.BoundShaderState.GetMeshShader()->AddRef();
-			}
-			if (Initializer.BoundShaderState.GetAmplificationShader())
-			{
-				Initializer.BoundShaderState.GetAmplificationShader()->AddRef();
-			}
-			if (Initializer.BoundShaderState.VertexDeclarationRHI)
-			{
-				Initializer.BoundShaderState.VertexDeclarationRHI->AddRef();
-			}
-			if (Initializer.BoundShaderState.VertexShaderRHI)
-			{
-				Initializer.BoundShaderState.VertexShaderRHI->AddRef();
-			}
-			if (Initializer.BoundShaderState.PixelShaderRHI)
-			{
-				Initializer.BoundShaderState.PixelShaderRHI->AddRef();
-			}
-			if (Initializer.BoundShaderState.GetGeometryShader())
-			{
-				Initializer.BoundShaderState.GetGeometryShader()->AddRef();
-			}
-			if (Initializer.BlendState)
-			{
-				Initializer.BlendState->AddRef();
-			}
-			if (Initializer.RasterizerState)
-			{
-				Initializer.RasterizerState->AddRef();
-			}
-			if (Initializer.DepthStencilState)
-			{
-				Initializer.DepthStencilState->AddRef();
-			}
-
-			if (Initializer.BlendState)
-			{
-				Initializer.BlendState->AddRef();
-			}
-			if (Initializer.RasterizerState)
-			{
-				Initializer.RasterizerState->AddRef();
-			}
-			if (Initializer.DepthStencilState)
-			{
-				Initializer.DepthStencilState->AddRef();
-			}
-		}		
-#if PLATFORM_WINDOWS
-		else
-		{
-			FComputePipelineState* ComputePipeline = static_cast<FComputePipelineState*>(Pipeline);
-			MarkInUseByPSOCompilation(ComputePipeline->ComputeShader);
-		}
-#endif // PLATFORM_WINDOWS
 	}
 
-	~FCompilePipelineStateTask()
+	~TCompilePipelineStateTaskBase()
 	{
 		if (Initializer.bFromPSOFileCache)
 		{
@@ -3060,173 +2989,10 @@ public:
 			case EPSOPrecacheResult::Untracked:    PSOPrecacheResultScopeString = TEXT("PSOPrecache: Untracked"); break;
 		}
 
-#if PLATFORM_WINDOWS
-		auto MarkUnusedByPSOCompilation = [](FRHIShader* Shader)
-		{
-			if (Shader)
-			{
-				Shader->SetInUseByPSOCompilation(false);
-			}
-		};
-#endif // PLATFORM_WINDOWS
-
 		TRACE_CPUPROFILER_EVENT_SCOPE_TEXT(PSOPrecacheResultScopeString);
 		{
 			TRACE_CPUPROFILER_EVENT_SCOPE_TEXT_CONDITIONAL(*PSOCompilationDebugData.PSOCompilationEventName, !PSOCompilationDebugData.PSOCompilationEventName.IsEmpty())
-
-			if (Pipeline->IsCompute())
-			{
-				FComputePipelineState* ComputePipeline = static_cast<FComputePipelineState*>(Pipeline);
-
-				uint64 StartTime = FPlatformTime::Cycles64();
-				FComputePipelineStateInitializer ComputeInitializer = { ComputePipeline->ComputeShader, 0 };
-				ComputeInitializer.bFromPSOFileCache = Initializer.bFromPSOFileCache;
-				ComputeInitializer.bPSOPrecache = Initializer.bPSOPrecache;
-				ComputePipeline->RHIPipeline = RHICreateComputePipelineState(ComputeInitializer);
-				CheckAndUpdateHitchCountStat(FPSOPrecacheRequestID::EType::Compute, !IsPrecachedPSO(Initializer), PSOCompilationDebugData, StartTime, PSOPrecacheResult);
-
-				if (!ComputePipeline->RHIPipeline)
-				{
-					HandlePipelineCreationFailure(ComputePipeline->ComputeShader, Initializer.bFromPSOFileCache || Initializer.bPSOPrecache);
-				}
-
-				FPrecacheComputeInitializer PrecacheInitializer{ ComputePipeline->ComputeShader, ComputePipeline->Name };
-				if (Initializer.bPSOPrecache)
-				{
-					bool bCSValid = ComputePipeline->RHIPipeline != nullptr && ComputePipeline->RHIPipeline->IsValid();
-					GPrecacheComputePipelineCache->PrecacheFinished(PrecacheInitializer, bCSValid);
-				}
-				else
-				{
-					GPrecacheComputePipelineCache->MarkPSOAsUsed(PrecacheInitializer);
-				}
-				
-#if PLATFORM_WINDOWS
-				MarkUnusedByPSOCompilation(ComputePipeline->ComputeShader);
-#endif // PLATFORM_WINDOWS
-			}
-			else
-			{
-				bool bSkipCreation = false;
-				if (GRHISupportsMeshShadersTier0)
-				{
-					if (!Initializer.BoundShaderState.VertexShaderRHI && !Initializer.BoundShaderState.GetMeshShader())
-					{
-						UE_LOG(LogRHI, Error, TEXT("Tried to create a Gfx Pipeline State without Vertex or Mesh Shader"));
-						bSkipCreation = true;
-					}
-				}
-				else
-				{
-					if (Initializer.BoundShaderState.GetMeshShader())
-					{
-						UE_LOG(LogRHI, Error, TEXT("Tried to create a Gfx Pipeline State with Mesh Shader on hardware without mesh shader support."));
-						bSkipCreation = true;
-					}
-
-					if (!Initializer.BoundShaderState.VertexShaderRHI)
-					{
-						UE_LOG(LogRHI, Error, TEXT("Tried to create a Gfx Pipeline State without Vertex Shader"));
-						bSkipCreation = true;
-					}
-				}
-
-			    const bool bAbortPSOCompileDueToShutdown = IsEngineExitRequested() && Initializer.bPSOPrecache;
-			    if (bAbortPSOCompileDueToShutdown)
-			    {
-				    UE_LOG(LogRHI, Verbose, TEXT("Skipping a precache compile due to engine shutdown."));
-				    bSkipCreation = true;
-			    }
-
-			    if(OptionalPriorityOverride)
-			    {
-				    Initializer.PrecacheCompileType = (uint32)FMath::Clamp((uint32)*OptionalPriorityOverride, (uint32)FGraphicsPipelineStateInitializer::EPSOPrecacheCompileType::MinPri, (uint32)FGraphicsPipelineStateInitializer::EPSOPrecacheCompileType::MaxPri);
-			    }
-
-				FGraphicsPipelineState* GfxPipeline = static_cast<FGraphicsPipelineState*>(Pipeline);
-
-				uint64 StartTime = FPlatformTime::Cycles64();
-				GfxPipeline->RHIPipeline = bSkipCreation ? nullptr : RHICreateGraphicsPipelineState(Initializer);
-				CheckAndUpdateHitchCountStat(FPSOPrecacheRequestID::EType::Graphics, !IsPrecachedPSO(Initializer), PSOCompilationDebugData, StartTime, PSOPrecacheResult);
-
-				if (GfxPipeline->RHIPipeline)
-				{
-					GfxPipeline->SortKey = GfxPipeline->RHIPipeline->GetSortKey();
-				}
-				else if(!bAbortPSOCompileDueToShutdown)
-				{
-					HandlePipelineCreationFailure(Initializer);
-				}
-
-				// Mark as finished when it's a precaching job
-				if (Initializer.bPSOPrecache)
-				{
-					GPrecacheGraphicsPipelineCache->PrecacheFinished(Initializer, GfxPipeline->RHIPipeline != nullptr);
-				}
-				else
-				{
-					GPrecacheGraphicsPipelineCache->MarkPSOAsUsed(Initializer);
-				}
-				
-#if PLATFORM_WINDOWS
-				MarkUnusedByPSOCompilation(Initializer.BoundShaderState.GetMeshShader());
-				MarkUnusedByPSOCompilation(Initializer.BoundShaderState.GetAmplificationShader());
-				MarkUnusedByPSOCompilation(Initializer.BoundShaderState.VertexShaderRHI);
-				MarkUnusedByPSOCompilation(Initializer.BoundShaderState.PixelShaderRHI);
-				MarkUnusedByPSOCompilation(Initializer.BoundShaderState.GetGeometryShader());
-				MarkUnusedByPSOCompilation(Initializer.BoundShaderState.GetMeshShader());
-#endif // PLATFORM_WINDOWS
-
-				if (Initializer.BoundShaderState.GetMeshShader())
-				{
-					Initializer.BoundShaderState.GetMeshShader()->Release();
-				}
-				if (Initializer.BoundShaderState.GetAmplificationShader())
-				{
-					Initializer.BoundShaderState.GetAmplificationShader()->Release();
-				}
-				if (Initializer.BoundShaderState.VertexDeclarationRHI)
-				{
-					Initializer.BoundShaderState.VertexDeclarationRHI->Release();
-				}
-				if (Initializer.BoundShaderState.VertexShaderRHI)
-				{
-					Initializer.BoundShaderState.VertexShaderRHI->Release();
-				}
-				if (Initializer.BoundShaderState.PixelShaderRHI)
-				{
-					Initializer.BoundShaderState.PixelShaderRHI->Release();
-				}
-				if (Initializer.BoundShaderState.GetGeometryShader())
-				{
-					Initializer.BoundShaderState.GetGeometryShader()->Release();
-				}
-				if (Initializer.BlendState)
-				{
-					Initializer.BlendState->Release();
-				}
-				if (Initializer.RasterizerState)
-				{
-					Initializer.RasterizerState->Release();
-				}
-				if (Initializer.DepthStencilState)
-				{
-					Initializer.DepthStencilState->Release();
-				}
-
-				if (Initializer.BlendState)
-				{
-					Initializer.BlendState->Release();
-				}
-				if (Initializer.RasterizerState)
-				{
-					Initializer.RasterizerState->Release();
-				}
-				if (Initializer.DepthStencilState)
-				{
-					Initializer.DepthStencilState->Release();
-				}
-			}
+			static_cast<TCompilePipelineStateTaskDerived*>(this)->CompilePSOInternal(OptionalPriorityOverride);
 		}
 
 #if WITH_RHI_BREADCRUMBS
@@ -3265,6 +3031,266 @@ public:
 	}
 };
 
+class FCompileGraphicsPipelineStateTask : public TCompilePipelineStateTaskBase<FCompileGraphicsPipelineStateTask, FGraphicsPipelineStateInitializer>
+{
+public:
+	FCompileGraphicsPipelineStateTask(
+		FPipelineStateAsync* InPipeline,
+		const FGraphicsPipelineStateInitializer& InInitializer,
+		EPSOPrecacheResult InPSOPrecacheResult,
+		bool InbPrecachedPSOMarkedUsed,
+		bool InbInImmediateCmdList,
+		const FPSOCompilationDebugData& InPSOCompilationDebugData)
+		: TCompilePipelineStateTaskBase(InPipeline, InInitializer, InPSOPrecacheResult, InbPrecachedPSOMarkedUsed, InbInImmediateCmdList, InPSOCompilationDebugData)
+	{
+#if PLATFORM_WINDOWS
+		auto MarkInUseByPSOCompilation = [](FRHIShader* Shader)
+			{
+				if (Shader)
+				{
+					Shader->SetInUseByPSOCompilation(true);
+				}
+			};
+		MarkInUseByPSOCompilation(Initializer.BoundShaderState.GetMeshShader());
+		MarkInUseByPSOCompilation(Initializer.BoundShaderState.GetAmplificationShader());
+		MarkInUseByPSOCompilation(Initializer.BoundShaderState.VertexShaderRHI);
+		MarkInUseByPSOCompilation(Initializer.BoundShaderState.PixelShaderRHI);
+		MarkInUseByPSOCompilation(Initializer.BoundShaderState.GetGeometryShader());
+		MarkInUseByPSOCompilation(Initializer.BoundShaderState.GetMeshShader());
+#endif // PLATFORM_WINDOWS
+
+		if (Initializer.BoundShaderState.GetMeshShader())
+		{
+			Initializer.BoundShaderState.GetMeshShader()->AddRef();
+		}
+		if (Initializer.BoundShaderState.GetAmplificationShader())
+		{
+			Initializer.BoundShaderState.GetAmplificationShader()->AddRef();
+		}
+		if (Initializer.BoundShaderState.VertexDeclarationRHI)
+		{
+			Initializer.BoundShaderState.VertexDeclarationRHI->AddRef();
+		}
+		if (Initializer.BoundShaderState.VertexShaderRHI)
+		{
+			Initializer.BoundShaderState.VertexShaderRHI->AddRef();
+		}
+		if (Initializer.BoundShaderState.PixelShaderRHI)
+		{
+			Initializer.BoundShaderState.PixelShaderRHI->AddRef();
+		}
+		if (Initializer.BoundShaderState.GetGeometryShader())
+		{
+			Initializer.BoundShaderState.GetGeometryShader()->AddRef();
+		}
+		if (Initializer.BlendState)
+		{
+			Initializer.BlendState->AddRef();
+		}
+		if (Initializer.RasterizerState)
+		{
+			Initializer.RasterizerState->AddRef();
+		}
+		if (Initializer.DepthStencilState)
+		{
+			Initializer.DepthStencilState->AddRef();
+		}
+
+		if (Initializer.BlendState)
+		{
+			Initializer.BlendState->AddRef();
+		}
+		if (Initializer.RasterizerState)
+		{
+			Initializer.RasterizerState->AddRef();
+		}
+		if (Initializer.DepthStencilState)
+		{
+			Initializer.DepthStencilState->AddRef();
+		}
+	}
+
+	void CompilePSOInternal(const FGraphicsPipelineStateInitializer::EPSOPrecacheCompileType* OptionalPriorityOverride = nullptr)
+	{
+		bool bSkipCreation = false;
+		if (GRHISupportsMeshShadersTier0)
+		{
+			if (!Initializer.BoundShaderState.VertexShaderRHI && !Initializer.BoundShaderState.GetMeshShader())
+			{
+				UE_LOG(LogRHI, Error, TEXT("Tried to create a Gfx Pipeline State without Vertex or Mesh Shader"));
+				bSkipCreation = true;
+			}
+		}
+		else
+		{
+			if (Initializer.BoundShaderState.GetMeshShader())
+			{
+				UE_LOG(LogRHI, Error, TEXT("Tried to create a Gfx Pipeline State with Mesh Shader on hardware without mesh shader support."));
+				bSkipCreation = true;
+			}
+
+			if (!Initializer.BoundShaderState.VertexShaderRHI)
+			{
+				UE_LOG(LogRHI, Error, TEXT("Tried to create a Gfx Pipeline State without Vertex Shader"));
+				bSkipCreation = true;
+			}
+		}
+
+		const bool bAbortPSOCompileDueToShutdown = IsEngineExitRequested() && Initializer.bPSOPrecache;
+		if (bAbortPSOCompileDueToShutdown)
+		{
+			UE_LOG(LogRHI, Verbose, TEXT("Skipping a precache compile due to engine shutdown."));
+			bSkipCreation = true;
+		}
+
+		if (OptionalPriorityOverride)
+		{
+			Initializer.PrecacheCompileType = (uint32)FMath::Clamp((uint32)*OptionalPriorityOverride, (uint32)FGraphicsPipelineStateInitializer::EPSOPrecacheCompileType::MinPri, (uint32)FGraphicsPipelineStateInitializer::EPSOPrecacheCompileType::MaxPri);
+		}
+
+		FGraphicsPipelineState* GfxPipeline = static_cast<FGraphicsPipelineState*>(Pipeline);
+
+		uint64 StartTime = FPlatformTime::Cycles64();
+		GfxPipeline->RHIPipeline = bSkipCreation ? nullptr : RHICreateGraphicsPipelineState(Initializer);
+		CheckAndUpdateHitchCountStat(FPSOPrecacheRequestID::EType::Graphics, !IsPrecachedPSO(Initializer), PSOCompilationDebugData, StartTime, PSOPrecacheResult);
+
+		if (GfxPipeline->RHIPipeline)
+		{
+			GfxPipeline->SortKey = GfxPipeline->RHIPipeline->GetSortKey();
+		}
+		else if (!bAbortPSOCompileDueToShutdown)
+		{
+			HandlePipelineCreationFailure(Initializer);
+		}
+
+		// Mark as finished when it's a precaching job
+		if (Initializer.bPSOPrecache)
+		{
+			GPrecacheGraphicsPipelineCache->PrecacheFinished(Initializer, GfxPipeline->RHIPipeline != nullptr);
+		}
+		else if (bPrecachedPSOMarkedUsed)
+		{
+			if (Initializer.StatePrecachePSOHash == 0)
+			{
+				Initializer.StatePrecachePSOHash = RHIComputeStatePrecachePSOHash(Initializer);
+			}
+			GPrecacheGraphicsPipelineCache->MarkPSOAsUsed(Initializer);
+		}
+
+#if PLATFORM_WINDOWS
+		auto MarkUnusedByPSOCompilation = [](FRHIShader* Shader)
+			{
+				if (Shader)
+				{
+					Shader->SetInUseByPSOCompilation(false);
+				}
+			};
+		MarkUnusedByPSOCompilation(Initializer.BoundShaderState.GetMeshShader());
+		MarkUnusedByPSOCompilation(Initializer.BoundShaderState.GetAmplificationShader());
+		MarkUnusedByPSOCompilation(Initializer.BoundShaderState.VertexShaderRHI);
+		MarkUnusedByPSOCompilation(Initializer.BoundShaderState.PixelShaderRHI);
+		MarkUnusedByPSOCompilation(Initializer.BoundShaderState.GetGeometryShader());
+		MarkUnusedByPSOCompilation(Initializer.BoundShaderState.GetMeshShader());
+#endif // PLATFORM_WINDOWS
+
+		if (Initializer.BoundShaderState.GetMeshShader())
+		{
+			Initializer.BoundShaderState.GetMeshShader()->Release();
+		}
+		if (Initializer.BoundShaderState.GetAmplificationShader())
+		{
+			Initializer.BoundShaderState.GetAmplificationShader()->Release();
+		}
+		if (Initializer.BoundShaderState.VertexDeclarationRHI)
+		{
+			Initializer.BoundShaderState.VertexDeclarationRHI->Release();
+		}
+		if (Initializer.BoundShaderState.VertexShaderRHI)
+		{
+			Initializer.BoundShaderState.VertexShaderRHI->Release();
+		}
+		if (Initializer.BoundShaderState.PixelShaderRHI)
+		{
+			Initializer.BoundShaderState.PixelShaderRHI->Release();
+		}
+		if (Initializer.BoundShaderState.GetGeometryShader())
+		{
+			Initializer.BoundShaderState.GetGeometryShader()->Release();
+		}
+		if (Initializer.BlendState)
+		{
+			Initializer.BlendState->Release();
+		}
+		if (Initializer.RasterizerState)
+		{
+			Initializer.RasterizerState->Release();
+		}
+		if (Initializer.DepthStencilState)
+		{
+			Initializer.DepthStencilState->Release();
+		}
+
+		if (Initializer.BlendState)
+		{
+			Initializer.BlendState->Release();
+		}
+		if (Initializer.RasterizerState)
+		{
+			Initializer.RasterizerState->Release();
+		}
+		if (Initializer.DepthStencilState)
+		{
+			Initializer.DepthStencilState->Release();
+		}
+	}
+};
+
+class FCompileComputePipelineStateTask : public TCompilePipelineStateTaskBase<FCompileComputePipelineStateTask, FComputePipelineStateInitializer>
+{
+public:
+	FCompileComputePipelineStateTask(
+		FPipelineStateAsync* InPipeline,
+		const FComputePipelineStateInitializer& InInitializer,
+		EPSOPrecacheResult InPSOPrecacheResult,
+		bool InbPrecachedPSOMarkedUsed,
+		bool InbInImmediateCmdList,
+		const FPSOCompilationDebugData& InPSOCompilationDebugData)
+		: TCompilePipelineStateTaskBase(InPipeline, InInitializer, InPSOPrecacheResult, InbPrecachedPSOMarkedUsed, InbInImmediateCmdList, InPSOCompilationDebugData)
+	{
+#if PLATFORM_WINDOWS
+		InInitializer.ComputeShader->SetInUseByPSOCompilation(true);
+#endif // PLATFORM_WINDOWS
+	}
+
+	void CompilePSOInternal(const FGraphicsPipelineStateInitializer::EPSOPrecacheCompileType* OptionalPriorityOverride = nullptr)
+	{
+		FComputePipelineState* ComputePipeline = static_cast<FComputePipelineState*>(Pipeline);
+
+		uint64 StartTime = FPlatformTime::Cycles64();
+		ComputePipeline->RHIPipeline = RHICreateComputePipelineState(Initializer);
+		CheckAndUpdateHitchCountStat(FPSOPrecacheRequestID::EType::Compute, !IsPrecachedPSO(Initializer), PSOCompilationDebugData, StartTime, PSOPrecacheResult);
+
+		if (!ComputePipeline->RHIPipeline)
+		{
+			HandlePipelineCreationFailure(Initializer);
+		}
+
+		if (Initializer.bPSOPrecache)
+		{
+			bool bCSValid = ComputePipeline->RHIPipeline != nullptr && ComputePipeline->RHIPipeline->IsValid();
+			GPrecacheComputePipelineCache->PrecacheFinished(Initializer, bCSValid);
+		}
+		else if (bPrecachedPSOMarkedUsed)
+		{
+			GPrecacheComputePipelineCache->MarkPSOAsUsed(Initializer);
+		}
+
+#if PLATFORM_WINDOWS
+		Initializer.ComputeShader->SetInUseByPSOCompilation(false);
+#endif // PLATFORM_WINDOWS
+	}
+};
+
 void PipelineStateCache::ReportFrameHitchToCSV()
 {
 	ReportFrameHitchThisFrame = true;
@@ -3377,7 +3403,48 @@ uint64 PipelineStateCache::RetrieveGraphicsPipelineStateSortKey(const FGraphicsP
 	return GraphicsPipelineState != nullptr ? GraphicsPipelineState->SortKey : 0;
 }
 
-static void InternalCreateComputePipelineState(FRHIComputeShader* ComputeShader, bool bDoAsyncCompile, bool bFromPSOFileCache, EPSOPrecacheResult PSOPrecacheResult, FComputePipelineState* CachedState, const FPSOCompilationDebugData& PSOCompilationDebugData, bool bInImmediateCmdList)
+template <typename TCompilePipelineTask, typename TPipelineInitializer>
+static void InternalAsyncCompilePipelineState(const TPipelineInitializer& Initializer, EPSOPrecacheResult PSOPrecacheResult, bool bPrecachedPSOMarkedUsed, FPipelineStateAsync* CachedState, const FPSOCompilationDebugData& PSOCompilationDebugData, bool bInImmediateCmdList)
+{
+	if (!Initializer.bPSOPrecache || !FPSOPrecacheThreadPool::UsePool())
+	{
+		TGraphTask<TCompilePipelineTask>::CreateTask().ConstructAndDispatchWhenReady(CachedState, Initializer, PSOPrecacheResult, bPrecachedPSOMarkedUsed, bInImmediateCmdList, PSOCompilationDebugData);
+	}
+	else
+	{
+		// Here, PSO precompiles use a separate thread pool.
+		// Note that we do not add precompile tasks as cmdlist prerequisites.
+		TUniquePtr<TCompilePipelineTask> ThreadPoolTask = MakeUnique<TCompilePipelineTask>(CachedState, Initializer, PSOPrecacheResult, bPrecachedPSOMarkedUsed, bInImmediateCmdList, PSOCompilationDebugData);
+		CachedState->SetPrecompileTask(MakeUnique<FPSOPrecacheAsyncTask>(
+			[ThreadPoolTask = MoveTemp(ThreadPoolTask)](const FPSOPrecacheAsyncTask* ThisTask)
+			{
+				// Convert the task priority to PSO precompile priority.
+				// Update here as the task's priority may have changed since creation.
+				FGraphicsPipelineStateInitializer::EPSOPrecacheCompileType PriOverride;
+				switch (ThisTask->GetPriority())
+				{
+					case EQueuedWorkPriority::Blocking:
+					case EQueuedWorkPriority::Highest:
+						PriOverride = FGraphicsPipelineStateInitializer::EPSOPrecacheCompileType::MaxPri;
+						break;
+					case EQueuedWorkPriority::High:
+					case EQueuedWorkPriority::Normal:
+						PriOverride = FGraphicsPipelineStateInitializer::EPSOPrecacheCompileType::NormalPri;
+						break;
+					default: checkNoEntry(); [[fallthrough]];
+					case EQueuedWorkPriority::Low:
+					case EQueuedWorkPriority::Lowest:
+						PriOverride = FGraphicsPipelineStateInitializer::EPSOPrecacheCompileType::MinPri;
+						break;
+				}
+				ThreadPoolTask->CompilePSO(&PriOverride);
+			}
+		));
+		CachedState->GetPrecompileTask()->StartBackgroundTask(&GPSOPrecacheThreadPool.Get(), EQueuedWorkPriority::Normal);
+	}
+}
+
+static void InternalCreateComputePipelineState(const FComputePipelineStateInitializer& Initializer, bool bDoAsyncCompile, EPSOPrecacheResult PSOPrecacheResult, bool bPrecachedPSOMarkedUsed, FComputePipelineState* CachedState, const FPSOCompilationDebugData& PSOCompilationDebugData, bool bInImmediateCmdList)
 {
 	FGraphEventRef GraphEvent = CachedState->GetCompletionEvent();
 
@@ -3385,10 +3452,7 @@ static void InternalCreateComputePipelineState(FRHIComputeShader* ComputeShader,
 	if (bDoAsyncCompile)
 	{
 		check(GraphEvent != nullptr);
-
-		FGraphicsPipelineStateInitializer GraphicsPipelineStateInitializer;
-		GraphicsPipelineStateInitializer.bFromPSOFileCache = bFromPSOFileCache;
-		TGraphTask<FCompilePipelineStateTask>::CreateTask().ConstructAndDispatchWhenReady(CachedState, GraphicsPipelineStateInitializer, PSOPrecacheResult, bInImmediateCmdList, PSOCompilationDebugData);
+		InternalAsyncCompilePipelineState<FCompileComputePipelineStateTask, FComputePipelineStateInitializer>(Initializer, PSOPrecacheResult, bPrecachedPSOMarkedUsed, CachedState, PSOCompilationDebugData, bInImmediateCmdList);
 	}
 	else
 	{
@@ -3396,14 +3460,17 @@ static void InternalCreateComputePipelineState(FRHIComputeShader* ComputeShader,
 
 		check(GraphEvent == nullptr);
 		uint64 StartTime = FPlatformTime::Cycles64();
-		FComputePipelineStateInitializer ComputeInitializerRHI = { ComputeShader, 0 };
-		ComputeInitializerRHI.bFromPSOFileCache = bFromPSOFileCache;
-		CachedState->RHIPipeline = RHICreateComputePipelineState(ComputeInitializerRHI);
-		CheckAndUpdateHitchCountStat(FPSOPrecacheRequestID::EType::Compute, !bFromPSOFileCache, PSOCompilationDebugData, StartTime, PSOPrecacheResult);
+		CachedState->RHIPipeline = RHICreateComputePipelineState(Initializer);
+		CheckAndUpdateHitchCountStat(FPSOPrecacheRequestID::EType::Compute, !IsPrecachedPSO(Initializer), PSOCompilationDebugData, StartTime, PSOPrecacheResult);
+
+		if (Initializer.bPSOPrecache)
+		{
+			GPrecacheComputePipelineCache->PrecacheFinished(Initializer, CachedState->RHIPipeline != nullptr);
+		}
 
 		if (!CachedState->RHIPipeline)
 		{
-			HandlePipelineCreationFailure(ComputeShader, bFromPSOFileCache);
+			HandlePipelineCreationFailure(Initializer);
 		}
 	}
 }
@@ -3419,7 +3486,22 @@ FComputePipelineState* PipelineStateCache::GetAndOrCreateComputePipelineState(FR
 
 	if (bWasFound == false)
 	{		
-		EPSOPrecacheResult PSOPrecacheResult = CheckPipelineStateInCache(ComputeShader);
+		FComputePipelineStateInitializer ComputeInitializer = { ComputeShader, 0 };
+		ComputeInitializer.bFromPSOFileCache = bFromFileCache;
+
+		EPSOPrecacheResult PSOPrecacheResult = EPSOPrecacheResult::Unknown;
+		bool bPrecachedPSOMarkedUsed = false;
+		if (IsPSOPrecachingEnabled())
+		{
+			if (ShouldTrackUsedPrecachedPSOs())
+			{
+				PSOPrecacheResult = GPrecacheComputePipelineCache->GetPrecachingStateAndMarkUsed(ComputeInitializer, bPrecachedPSOMarkedUsed);
+			}
+			else
+			{
+				PSOPrecacheResult = GPrecacheComputePipelineCache->GetPrecachingState(ComputeInitializer);
+			}
+		}
 
 		bool bWasPSOPrecached = PSOPrecacheResult == EPSOPrecacheResult::Active || PSOPrecacheResult == EPSOPrecacheResult::Complete;
 		
@@ -3427,7 +3509,7 @@ FComputePipelineState* PipelineStateCache::GetAndOrCreateComputePipelineState(FR
 		FPipelineFileCacheManager::CacheComputePSO(GetTypeHash(ComputeShader), ComputeShader, bWasPSOPrecached, &Stats);
 
 		// create new compute state
-		OutCachedState = new FComputePipelineState(ComputeShader, nullptr);
+		OutCachedState = new FComputePipelineState(ComputeShader);
 		OutCachedState->Stats = Stats;
 		if (DoAsyncCompile)
 		{
@@ -3457,7 +3539,7 @@ FComputePipelineState* PipelineStateCache::GetAndOrCreateComputePipelineState(FR
 #endif // WITH_RHI_BREADCRUMBS
 		
 		FGraphEventRef GraphEvent = OutCachedState->GetCompletionEvent();
-		InternalCreateComputePipelineState(ComputeShader, DoAsyncCompile, bFromFileCache, PSOPrecacheResult, OutCachedState, PSOCompilationDebugData, RHICmdList.IsImmediate());
+		InternalCreateComputePipelineState(ComputeInitializer, DoAsyncCompile, PSOPrecacheResult, bPrecachedPSOMarkedUsed, OutCachedState, PSOCompilationDebugData, RHICmdList.IsImmediate());
 
 		// Don't add precached PSOs as a dispatch prerequisite. We don't need to wait for them to complete before the RHICmdList can be dispatched.
 		if (GraphEvent.IsValid() && !bFromFileCache)
@@ -3641,9 +3723,6 @@ private:
 
 	FRHIShaderBindingLayout ShaderBindingLayout;
 };
-#endif // RHI_RAYTRACING
-
-#if RHI_RAYTRACING
 
 static bool ValidateRayTracingPipelinePayloadMask(const FRayTracingPipelineStateInitializer& InInitializer)
 {
@@ -3704,7 +3783,6 @@ static bool ValidateRayTracingPipelinePayloadMask(const FRayTracingPipelineState
 
 #endif // RHI_RAYTRACING
 
-
 FRayTracingPipelineState* PipelineStateCache::GetAndOrCreateRayTracingPipelineState(
 	FRHICommandList& RHICmdList,
 	const FRayTracingPipelineStateInitializer& InInitializer,
@@ -3851,7 +3929,7 @@ inline void ValidateGraphicsPipelineStateInitializer(const FGraphicsPipelineStat
 	check(Initializer.DepthStencilState && Initializer.BlendState && Initializer.RasterizerState);
 }
 
-static void InternalCreateGraphicsPipelineState(const FGraphicsPipelineStateInitializer& Initializer, EPSOPrecacheResult PSOPrecacheResult, bool bDoAsyncCompile, bool bPSOPrecache, FGraphicsPipelineState* CachedState, const FPSOCompilationDebugData& PSOCompilationDebugData, bool bInImmediateCmdList)
+static void InternalCreateGraphicsPipelineState(const FGraphicsPipelineStateInitializer& Initializer, EPSOPrecacheResult PSOPrecacheResult, bool bPrecachedPSOMarkedUsed, bool bDoAsyncCompile, FGraphicsPipelineState* CachedState, const FPSOCompilationDebugData& PSOCompilationDebugData, bool bInImmediateCmdList)
 {
 	FGraphEventRef GraphEvent = CachedState->GetCompletionEvent();
 
@@ -3859,43 +3937,7 @@ static void InternalCreateGraphicsPipelineState(const FGraphicsPipelineStateInit
 	if (bDoAsyncCompile)
 	{
 		check(GraphEvent != nullptr);
-		// Use normal task graph for non-precompile jobs (or when thread pool is not enabled)
-		if (!bPSOPrecache || !FPSOPrecacheThreadPool::UsePool())
-		{
-			TGraphTask<FCompilePipelineStateTask>::CreateTask().ConstructAndDispatchWhenReady(CachedState, Initializer, PSOPrecacheResult, bInImmediateCmdList, PSOCompilationDebugData);
-		}
-		else
-		{
-			// Here, PSO precompiles use a separate thread pool.
-			// Note that we do not add precompile tasks as cmdlist prerequisites.
-			TUniquePtr<FCompilePipelineStateTask> ThreadPoolTask = MakeUnique<FCompilePipelineStateTask>(CachedState, Initializer, PSOPrecacheResult, bInImmediateCmdList, PSOCompilationDebugData);
-			CachedState->SetPrecompileTask(MakeUnique<FPSOPrecacheAsyncTask>(
-				[ThreadPoolTask = MoveTemp(ThreadPoolTask)](const FPSOPrecacheAsyncTask* ThisTask)
-				{
-					// Convert the task priority to PSO precompile priority.
-					// Update here as the task's priority may have changed since creation.
-					FGraphicsPipelineStateInitializer::EPSOPrecacheCompileType PriOverride;
-					switch (ThisTask->GetPriority())
-					{
-						case EQueuedWorkPriority::Blocking:
-						case EQueuedWorkPriority::Highest:
-							PriOverride = FGraphicsPipelineStateInitializer::EPSOPrecacheCompileType::MaxPri;
-						break;
-						case EQueuedWorkPriority::High:
-						case EQueuedWorkPriority::Normal:
-							PriOverride = FGraphicsPipelineStateInitializer::EPSOPrecacheCompileType::NormalPri;
-						break;
-						default: checkNoEntry(); [[fallthrough]];
-						case EQueuedWorkPriority::Low:
-						case EQueuedWorkPriority::Lowest:
-							PriOverride = FGraphicsPipelineStateInitializer::EPSOPrecacheCompileType::MinPri;
-						break;
-					}
-					ThreadPoolTask->CompilePSO(&PriOverride);
-				}
-			));
-			CachedState->GetPrecompileTask()->StartBackgroundTask(&GPSOPrecacheThreadPool.Get(), Initializer.bFromPSOFileCache ? EQueuedWorkPriority::Normal : EQueuedWorkPriority::Low);
-		}
+		InternalAsyncCompilePipelineState<FCompileGraphicsPipelineStateTask, FGraphicsPipelineStateInitializer>(Initializer, PSOPrecacheResult, bPrecachedPSOMarkedUsed, CachedState, PSOCompilationDebugData, bInImmediateCmdList);
 	}
 	else
 	{		
@@ -3972,11 +4014,23 @@ FGraphicsPipelineState* PipelineStateCache::GetAndOrCreateGraphicsPipelineState(
 	if (bWasFound == false)
 	{
 		EPSOPrecacheResult PSOPrecacheResult = EPSOPrecacheResult::Unknown;
+		bool bPrecachedPSOMarkedUsed = false;
 		if (IsPSOPrecachingEnabled())
 		{
 			FGraphicsPipelineStateInitializer PSOPrecacheInitializer = Initializer;
-			PSOPrecacheInitializer.StatePrecachePSOHash = RHIComputeStatePrecachePSOHash(PSOPrecacheInitializer);
-			PSOPrecacheResult = CheckPipelineStateInCache(PSOPrecacheInitializer);
+			if (PSOPrecacheInitializer.StatePrecachePSOHash == 0)
+			{
+				PSOPrecacheInitializer.StatePrecachePSOHash = RHIComputeStatePrecachePSOHash(PSOPrecacheInitializer);
+			}
+
+			if (ShouldTrackUsedPrecachedPSOs())
+			{
+				PSOPrecacheResult = GPrecacheGraphicsPipelineCache->GetPrecachingStateAndMarkUsed(PSOPrecacheInitializer, bPrecachedPSOMarkedUsed);
+			}
+			else
+			{
+				PSOPrecacheResult = GPrecacheGraphicsPipelineCache->GetPrecachingState(PSOPrecacheInitializer);
+			}
 		}
 
 		bool DoAsyncCompile = IsAsyncCompilationAllowed(RHICmdList, Initializer.bFromPSOFileCache);
@@ -4016,13 +4070,12 @@ FGraphicsPipelineState* PipelineStateCache::GetAndOrCreateGraphicsPipelineState(
 		}
 #endif // WITH_RHI_BREADCRUMBS
 
-		bool bPSOPrecache = Initializer.bFromPSOFileCache;
 		FGraphEventRef GraphEvent = OutCachedState->GetCompletionEvent();
-		InternalCreateGraphicsPipelineState(Initializer, PSOPrecacheResult, DoAsyncCompile, bPSOPrecache, OutCachedState, PSOCompilationDebugData, RHICmdList.IsImmediate());
+		InternalCreateGraphicsPipelineState(Initializer, PSOPrecacheResult, bPrecachedPSOMarkedUsed, DoAsyncCompile, OutCachedState, PSOCompilationDebugData, RHICmdList.IsImmediate());
 
 		// Add dispatch pre requisite for non precaching jobs only
 		//if (GraphEvent.IsValid() && (!bPSOPrecache || !FPSOPrecacheThreadPool::UsePool()))
-		if (GraphEvent.IsValid() && !bPSOPrecache)
+		if (GraphEvent.IsValid() && !Initializer.bFromPSOFileCache)
 		{
 			check(DoAsyncCompile);
 			RHICmdList.AddDispatchPrerequisite(GraphEvent);
@@ -4146,51 +4199,24 @@ FPSOPrecacheRequestResult FPrecacheComputePipelineCache::PrecacheComputePipeline
 		return Result;
 	}
 
-	FPrecacheComputeInitializer PrecacheComputeInitializer(ComputeShader, Name);
 	static bool bDoAsyncCompile = FApp::ShouldUseThreadingForPerformance();
+
+	// Mark as precache.
+	FComputePipelineStateInitializer PrecacheComputeInitializer = { ComputeShader, 0 };
+	PrecacheComputeInitializer.bPSOPrecache = true;
+	
 	return TryAddNewState(PrecacheComputeInitializer, Name, bDoAsyncCompile);
 }
 
-void FPrecacheComputePipelineCache::OnNewPipelineStateCreated(const FPrecacheComputeInitializer& ComputeInitializer, FComputePipelineState* CachedState, const FString& PSOCompilationEventName, bool bDoAsyncCompile)
+void FPrecacheComputePipelineCache::OnNewPipelineStateCreated(const FComputePipelineStateInitializer& ComputeInitializer, FComputePipelineState* CachedState, const FString& PSOCompilationEventName, bool bDoAsyncCompile)
 {
-	// create a compilation task, or just do it now...
-	if (bDoAsyncCompile)
-	{
-		check(CachedState->GetCompletionEvent() != nullptr);
-		FGraphicsPipelineStateInitializer GraphicsPipelineStateInitializer;
-		GraphicsPipelineStateInitializer.bPSOPrecache = true;
-		GraphicsPipelineStateInitializer.SetPSOPrecacheCompileType(FGraphicsPipelineStateInitializer::EPSOPrecacheCompileType::NormalPri);
+	check(ComputeInitializer.ComputeShader);
+	check(ComputeInitializer.bPSOPrecache);
 
-		FPSOCompilationDebugData PSOCompilationDebugData;
-		PSOCompilationDebugData.PSOCompilationEventName = PSOCompilationEventName;
+	FPSOCompilationDebugData PSOCompilationDebugData;
+	PSOCompilationDebugData.PSOCompilationEventName = PSOCompilationEventName;
 
-		// Thread pool disabled?
-		if (!FPSOPrecacheThreadPool::UsePool())
-		{
-			TGraphTask<FCompilePipelineStateTask>::CreateTask().ConstructAndDispatchWhenReady(CachedState, GraphicsPipelineStateInitializer, EPSOPrecacheResult::Active, false, PSOCompilationDebugData);
-		}
-		else
-		{
-			// Here, PSO precompiles use a separate thread pool.
-			// Note that we do not add precompile tasks as cmdlist prerequisites.
-			TUniquePtr<FCompilePipelineStateTask> ThreadPoolTask = MakeUnique<FCompilePipelineStateTask>(CachedState, GraphicsPipelineStateInitializer, EPSOPrecacheResult::Active, false, PSOCompilationDebugData);
-			CachedState->SetPrecompileTask(MakeUnique<FPSOPrecacheAsyncTask>(
-				[ThreadPoolTask = MoveTemp(ThreadPoolTask)](const FPSOPrecacheAsyncTask* ThisTask)
-				{
-					ThreadPoolTask->CompilePSO();
-				}
-			));
-			CachedState->GetPrecompileTask()->StartBackgroundTask(&GPSOPrecacheThreadPool.Get(), EQueuedWorkPriority::Normal);
-		}
-	}
-	else
-	{
-		check(CachedState->GetCompletionEvent() == nullptr);
-		FComputePipelineStateInitializer ComputeInitializerRHI = { CachedState->ComputeShader, 0 };
-		ComputeInitializerRHI.bPSOPrecache = true;
-		CachedState->RHIPipeline = RHICreateComputePipelineState(ComputeInitializerRHI);
-		GPrecacheComputePipelineCache->PrecacheFinished(ComputeInitializer, CachedState->RHIPipeline != nullptr);
-	}
+	InternalCreateComputePipelineState(ComputeInitializer, bDoAsyncCompile, EPSOPrecacheResult::Active, false, CachedState, PSOCompilationDebugData, false);
 }
 
 FPSOPrecacheRequestResult PipelineStateCache::PrecacheComputePipelineState(FRHIComputeShader* ComputeShader, const TCHAR* Name, bool bForcePrecache)
@@ -4231,7 +4257,7 @@ void FPrecacheGraphicsPipelineCache::OnNewPipelineStateCreated(const FGraphicsPi
 	PSOCompilationDebugData.PSOCompilationEventName = PSOCompilationEventName;
 
 	// Start the precache task	
-	InternalCreateGraphicsPipelineState(InitializerCopy, EPSOPrecacheResult::Active, bDoAsyncCompile, InitializerCopy.bPSOPrecache, NewGraphicsPipelineState, PSOCompilationDebugData, false);
+	InternalCreateGraphicsPipelineState(InitializerCopy, EPSOPrecacheResult::Active, false, bDoAsyncCompile, NewGraphicsPipelineState, PSOCompilationDebugData, false);
 }
 
 FPSOPrecacheRequestResult PipelineStateCache::PrecacheGraphicsPipelineState(const FGraphicsPipelineStateInitializer& Initializer)
@@ -4256,7 +4282,7 @@ EPSOPrecacheResult PipelineStateCache::CheckPipelineStateInCache(FRHIComputeShad
 		return EPSOPrecacheResult::Unknown;
 	}
 
-	return GPrecacheComputePipelineCache->GetPrecachingState(FPrecacheComputeInitializer(ComputeShader, nullptr));
+	return GPrecacheComputePipelineCache->GetPrecachingState(FComputePipelineStateInitializer(ComputeShader, 0));
 }
 
 bool PipelineStateCache::IsPrecaching(const FPSOPrecacheRequestID& PSOPrecacheRequestID)
@@ -4295,7 +4321,7 @@ bool PipelineStateCache::IsPrecaching(FRHIComputeShader* ComputeShader)
 		return false;
 	}
 
-	return GPrecacheComputePipelineCache->GetPrecachingState(FPrecacheComputeInitializer(ComputeShader, nullptr)) == EPSOPrecacheResult::Active;
+	return GPrecacheComputePipelineCache->GetPrecachingState(FComputePipelineStateInitializer(ComputeShader, 0)) == EPSOPrecacheResult::Active;
 }
 
 bool PipelineStateCache::IsPrecaching()
