Start graphics PSO precache compilation tasks at Low priority as they support priority bumping.

Restores previous behavior that was lost in refactor.

diff --git a/Engine/Source/Runtime/RHI/Private/PipelineStateCache.cpp b/Engine/Source/Runtime/RHI/Private/PipelineStateCache.cpp
--- a/Engine/Source/Runtime/RHI/Private/PipelineStateCache.cpp
+++ b/Engine/Source/Runtime/RHI/Private/PipelineStateCache.cpp
@@ -3440,7 +3440,17 @@ static void InternalAsyncCompilePipelineState(const TPipelineInitializer& Initia
 				ThreadPoolTask->CompilePSO(&PriOverride);
 			}
 		));
-		CachedState->GetPrecompileTask()->StartBackgroundTask(&GPSOPrecacheThreadPool.Get(), EQueuedWorkPriority::Normal);
+
+
+		// Graphics pipeline initializers support overriding the priority for PSO precaching, so start them at low.
+		if constexpr (std::is_same_v<TPipelineInitializer, FGraphicsPipelineStateInitializer>)
+		{
+			CachedState->GetPrecompileTask()->StartBackgroundTask(&GPSOPrecacheThreadPool.Get(), Initializer.bPSOPrecache ? EQueuedWorkPriority::Low : EQueuedWorkPriority::Normal);
+		}
+		else
+		{
+			CachedState->GetPrecompileTask()->StartBackgroundTask(&GPSOPrecacheThreadPool.Get(), EQueuedWorkPriority::Normal);
+		}
 	}
 }
 
