Add PSO precache CSV stats to monitor precache queue lengths by priority and some in-progress compile stats.

diff --git a/Engine/Source/Runtime/RHI/Private/PipelineStateCache.cpp b/Engine/Source/Runtime/RHI/Private/PipelineStateCache.cpp
--- a/Engine/Source/Runtime/RHI/Private/PipelineStateCache.cpp
+++ b/Engine/Source/Runtime/RHI/Private/PipelineStateCache.cpp
@@ -41,6 +41,7 @@ namespace PipelineStateCache
 #define PIPELINESTATECACHE_VERIFYTHREADSAFE (!UE_BUILD_SHIPPING && !UE_BUILD_TEST)
 
 CSV_DECLARE_CATEGORY_EXTERN(PSO);
+CSV_DEFINE_CATEGORY(PSOPrecacheCompiling, false);
 
 DEFINE_LOG_CATEGORY_STATIC(LogPSOHitching, Log, All);
 
@@ -252,6 +253,7 @@ static inline void CheckAndUpdateHitchCountStat(FPSOPrecacheRequestID::EType PSO
 #endif // WITH_RHI_BREADCRUMBS
 				INC_DWORD_STAT(STAT_RuntimeGraphicsPSOHitchCount);
 				CSV_CUSTOM_STAT(PSO, GraphicsPSOHitch, 1, ECsvCustomStatOp::Accumulate);
+				CSV_CUSTOM_STAT(PSO, GraphicsPSOHitchTime, PSOCreationTimeMs, ECsvCustomStatOp::Accumulate);
 				GraphicsPSOCreationHitchCount++;
 			}
 			else if (PSOType == FPSOPrecacheRequestID::EType::Compute)
@@ -263,6 +265,7 @@ static inline void CheckAndUpdateHitchCountStat(FPSOPrecacheRequestID::EType PSO
 #endif // WITH_RHI_BREADCRUMBS
 				INC_DWORD_STAT(STAT_RuntimeComputePSOHitchCount);
 				CSV_CUSTOM_STAT(PSO, ComputePSOHitch, 1, ECsvCustomStatOp::Accumulate);
+				CSV_CUSTOM_STAT(PSO, ComputePSOHitchTime, PSOCreationTimeMs, ECsvCustomStatOp::Accumulate);
 				ComputePSOCreationHitchCount++;
 			}
 
@@ -732,6 +735,202 @@ namespace PipelineStateCache
 	}
 }
 
+static FGraphicsPipelineStateInitializer::EPSOPrecacheCompileType GetPSOCompileTypeFromQueuePri(EQueuedWorkPriority QueuePri)
+{
+	switch (QueuePri)
+	{
+		case EQueuedWorkPriority::Blocking:
+		case EQueuedWorkPriority::Highest:
+			return FGraphicsPipelineStateInitializer::EPSOPrecacheCompileType::MaxPri;
+		case EQueuedWorkPriority::High:
+		case EQueuedWorkPriority::Normal:
+			return FGraphicsPipelineStateInitializer::EPSOPrecacheCompileType::NormalPri;
+		default: checkNoEntry(); [[fallthrough]];
+		case EQueuedWorkPriority::Low:
+		case EQueuedWorkPriority::Lowest:
+			return FGraphicsPipelineStateInitializer::EPSOPrecacheCompileType::MinPri;
+	}
+}
+
+#if CSV_PROFILER_STATS
+class FPSOCSVStatTracker
+{
+	bool bPrecache = false;
+	bool bCompute = false;
+
+	// PSO compiles can take place over multiple frames, use the persistent stats.
+	static inline TCsvPersistentCustomStat<int>* TotalCompletedPSOPrecompiles			= nullptr;
+	static inline TCsvPersistentCustomStat<int>* PSOPrecompileQueueLength				= nullptr;
+	static inline TCsvPersistentCustomStat<int>* PSOComputePrecompilesInProgress		= nullptr;
+	static inline TCsvPersistentCustomStat<int>* MinPriPrecacheInFlight					= nullptr;
+	static inline TCsvPersistentCustomStat<int>* NormalPriPrecacheInFlight				= nullptr;
+	static inline TCsvPersistentCustomStat<int>* MaxPriPrecacheInFlight					= nullptr;
+
+	static inline TCsvPersistentCustomStat<int>* MaxPriPSOPrecompileQueueLength			= nullptr;
+	static inline TCsvPersistentCustomStat<int>* NormalPriPSOPrecompileQueueLength		= nullptr;
+	static inline TCsvPersistentCustomStat<int>* MinPriPSOPrecompileQueueLength			= nullptr;
+
+	FGraphicsPipelineStateInitializer::EPSOPrecacheCompileType Pri = FGraphicsPipelineStateInitializer::EPSOPrecacheCompileType::NotSet;
+
+	static void InitStats()
+	{
+		static std::atomic_bool bFirstRun = true;
+		if (bFirstRun)
+		{
+			static UE::FMutex WindowManagerLock;
+			UE::TUniqueLock Lock(WindowManagerLock);
+			if (bFirstRun)
+			{
+				TotalCompletedPSOPrecompiles		= FCsvProfiler::Get()->GetOrCreatePersistentCustomStatInt(TEXT("TotalCompletedPSOPrecompiles"), CSV_CATEGORY_INDEX(PSOPrecacheCompiling));
+				PSOPrecompileQueueLength			= FCsvProfiler::Get()->GetOrCreatePersistentCustomStatInt(TEXT("PSOPrecompileQueueLength"), CSV_CATEGORY_INDEX(PSOPrecacheCompiling));
+				PSOComputePrecompilesInProgress		= FCsvProfiler::Get()->GetOrCreatePersistentCustomStatInt(TEXT("PSOComputePrecompilesInProgress"), CSV_CATEGORY_INDEX(PSOPrecacheCompiling));
+				MinPriPrecacheInFlight				= FCsvProfiler::Get()->GetOrCreatePersistentCustomStatInt(TEXT("MinPriPrecacheInFlight"), CSV_CATEGORY_INDEX(PSOPrecacheCompiling));
+				NormalPriPrecacheInFlight			= FCsvProfiler::Get()->GetOrCreatePersistentCustomStatInt(TEXT("NormalPriPrecacheInFlight"), CSV_CATEGORY_INDEX(PSOPrecacheCompiling));
+				MaxPriPrecacheInFlight				= FCsvProfiler::Get()->GetOrCreatePersistentCustomStatInt(TEXT("MaxPriPrecacheInFlight"), CSV_CATEGORY_INDEX(PSOPrecacheCompiling));
+
+				MaxPriPSOPrecompileQueueLength		= FCsvProfiler::Get()->GetOrCreatePersistentCustomStatInt(TEXT("MaxPriPSOPrecompileQueueLength"), CSV_CATEGORY_INDEX(PSOPrecacheCompiling));
+				NormalPriPSOPrecompileQueueLength	= FCsvProfiler::Get()->GetOrCreatePersistentCustomStatInt(TEXT("NormalPriPSOPrecompileQueueLength"), CSV_CATEGORY_INDEX(PSOPrecacheCompiling));
+				MinPriPSOPrecompileQueueLength		= FCsvProfiler::Get()->GetOrCreatePersistentCustomStatInt(TEXT("MinPriPSOPrecompileQueueLength"), CSV_CATEGORY_INDEX(PSOPrecacheCompiling));
+				bFirstRun = false;
+			}
+		}
+	}
+	static bool AreStatsEnabled()
+	{
+		return bEnableStatTracker && FCsvProfiler::Get()->IsCategoryEnabled(CSV_CATEGORY_INDEX(PSOPrecacheCompiling));
+	}
+public:
+	static inline bool bEnableStatTracker = true;
+
+	FPSOCSVStatTracker()
+	{
+		if (!AreStatsEnabled())
+		{ 
+			return;
+		}
+
+		InitStats();
+	}
+
+	void SetState(bool bInPrecache, bool bInCompute, FGraphicsPipelineStateInitializer::EPSOPrecacheCompileType InPri)
+	{
+		if (!AreStatsEnabled())
+		{
+			return;
+		}
+
+		check(TotalCompletedPSOPrecompiles);
+		bPrecache = bInPrecache;
+		bCompute = bInCompute;
+
+		if (bPrecache)
+		{
+			PSOPrecompileQueueLength->Add(1);
+
+			if (bCompute)
+			{
+				PSOComputePrecompilesInProgress->Add(1);
+			}
+
+			SetPriority(InPri);
+		}
+	}
+	
+	void SetPriority(FGraphicsPipelineStateInitializer::EPSOPrecacheCompileType InPri)
+	{
+		if (!AreStatsEnabled())
+		{
+			return;
+		}
+
+		if (bPrecache && InPri != Pri)
+		{
+			// unset the previous priority
+			if (Pri != FGraphicsPipelineStateInitializer::EPSOPrecacheCompileType::NotSet)
+			{
+				GetPriStat(Pri, false)->Add(-1);
+			}
+			Pri = InPri;
+			GetPriStat(Pri, false)->Add(1);
+		}
+	}
+
+	void SetInFlightPriority(FGraphicsPipelineStateInitializer::EPSOPrecacheCompileType InPri)
+	{
+		if (!AreStatsEnabled())
+		{
+			return;
+		}
+
+		if (bPrecache)
+		{
+			check(InPri == Pri);
+			check(InPri != FGraphicsPipelineStateInitializer::EPSOPrecacheCompileType::NotSet);
+			GetPriStat(Pri, true)->Add(1);
+		}
+	}
+
+	~FPSOCSVStatTracker()
+	{
+		if (!AreStatsEnabled())
+		{
+			return;
+		}
+
+		if (bPrecache)
+		{
+			TotalCompletedPSOPrecompiles->Add(1);
+			PSOPrecompileQueueLength->Add(-1);
+
+			if (bCompute)
+			{
+				PSOComputePrecompilesInProgress->Add(-1);
+			}
+
+			if (ensure(Pri != FGraphicsPipelineStateInitializer::EPSOPrecacheCompileType::NotSet))
+			{
+				GetPriStat(Pri, true)->Add(-1);
+				GetPriStat(Pri, false)->Add(-1);
+			}
+		}
+	}
+	private:
+		TCsvPersistentCustomStat<int>* GetPriStat(FGraphicsPipelineStateInitializer::EPSOPrecacheCompileType InPri, bool bInFlight) const
+		{
+			switch (InPri)
+			{
+			case FGraphicsPipelineStateInitializer::EPSOPrecacheCompileType::MaxPri:
+				return bInFlight ? MaxPriPrecacheInFlight : MaxPriPSOPrecompileQueueLength;
+			case FGraphicsPipelineStateInitializer::EPSOPrecacheCompileType::NormalPri:
+				return bInFlight ? NormalPriPrecacheInFlight : NormalPriPSOPrecompileQueueLength ;
+			case FGraphicsPipelineStateInitializer::EPSOPrecacheCompileType::MinPri:
+				return bInFlight ? MinPriPrecacheInFlight : MinPriPSOPrecompileQueueLength;
+			default:
+				checkNoEntry();
+				return nullptr;
+			}
+		}
+};
+
+static FAutoConsoleVariableRef CVarPSOPrecacheEnableCompileStats(
+	TEXT("r.PSOPrecache.EnableCompileStats"),
+	FPSOCSVStatTracker::bEnableStatTracker,
+	TEXT("Enable CSV stat tracking of precache compile progress. (default enabled)"),
+	ECVF_RenderThreadSafe
+);
+
+#else
+class FPSOCSVStatTracker
+{
+public:
+	FPSOCSVStatTracker() {}
+	void SetState(bool bInPrecache, bool bInCompute, FGraphicsPipelineStateInitializer::EPSOPrecacheCompileType InPri) {}
+	void SetPriority(FGraphicsPipelineStateInitializer::EPSOPrecacheCompileType InPri) {}
+	void SetInFlightPriority(FGraphicsPipelineStateInitializer::EPSOPrecacheCompileType InPri) {}
+};
+#endif
+
+
 /**
  * Base class for pipeline state intended to be stored in a TSharedPipelineStateCache,
  * with state double buffering for Render and RHI Threads.
@@ -783,6 +982,8 @@ private:
 
 public:
 
+	FPSOCSVStatTracker CSVStat;
+
 	virtual ~FPipelineStateAsync() override
 	{
 		check(IsComplete());
@@ -2127,7 +2328,10 @@ protected:
 			{
 				EQueuedWorkPriority PrevPriority = PipelineState->GetPrecompileTask()->GetPriority();
 				check(PrevPriority > NewPriority);
-				PipelineState->GetPrecompileTask()->Reschedule(&GPSOPrecacheThreadPool.Get(), NewPriority);
+				if( PipelineState->GetPrecompileTask()->Reschedule(&GPSOPrecacheThreadPool.Get(), NewPriority) )
+				{
+					PipelineState->CSVStat.SetPriority(GetPSOCompileTypeFromQueuePri(NewPriority));
+				}
 			}
 		}
 
@@ -3147,6 +3351,7 @@ public:
 		{
 			Initializer.PrecacheCompileType = (uint32)FMath::Clamp((uint32)*OptionalPriorityOverride, (uint32)FGraphicsPipelineStateInitializer::EPSOPrecacheCompileType::MinPri, (uint32)FGraphicsPipelineStateInitializer::EPSOPrecacheCompileType::MaxPri);
 		}
+		Pipeline->CSVStat.SetInFlightPriority((FGraphicsPipelineStateInitializer::EPSOPrecacheCompileType)Initializer.PrecacheCompileType);
 
 		FGraphicsPipelineState* GfxPipeline = static_cast<FGraphicsPipelineState*>(Pipeline);
 
@@ -3267,6 +3472,7 @@ public:
 		FComputePipelineState* ComputePipeline = static_cast<FComputePipelineState*>(Pipeline);
 
 		uint64 StartTime = FPlatformTime::Cycles64();
+		Pipeline->CSVStat.SetInFlightPriority(FGraphicsPipelineStateInitializer::EPSOPrecacheCompileType::NormalPri);
 		ComputePipeline->RHIPipeline = RHICreateComputePipelineState(Initializer);
 		CheckAndUpdateHitchCountStat(FPSOPrecacheRequestID::EType::Compute, !IsPrecachedPSO(Initializer), PSOCompilationDebugData, StartTime, PSOPrecacheResult);
 
@@ -3420,37 +3626,21 @@ static void InternalAsyncCompilePipelineState(const TPipelineInitializer& Initia
 			{
 				// Convert the task priority to PSO precompile priority.
 				// Update here as the task's priority may have changed since creation.
-				FGraphicsPipelineStateInitializer::EPSOPrecacheCompileType PriOverride;
-				switch (ThisTask->GetPriority())
-				{
-					case EQueuedWorkPriority::Blocking:
-					case EQueuedWorkPriority::Highest:
-						PriOverride = FGraphicsPipelineStateInitializer::EPSOPrecacheCompileType::MaxPri;
-						break;
-					case EQueuedWorkPriority::High:
-					case EQueuedWorkPriority::Normal:
-						PriOverride = FGraphicsPipelineStateInitializer::EPSOPrecacheCompileType::NormalPri;
-						break;
-					default: checkNoEntry(); [[fallthrough]];
-					case EQueuedWorkPriority::Low:
-					case EQueuedWorkPriority::Lowest:
-						PriOverride = FGraphicsPipelineStateInitializer::EPSOPrecacheCompileType::MinPri;
-						break;
-				}
+				FGraphicsPipelineStateInitializer::EPSOPrecacheCompileType PriOverride = GetPSOCompileTypeFromQueuePri(ThisTask->GetPriority());
 				ThreadPoolTask->CompilePSO(&PriOverride);
 			}
 		));
 
+		EQueuedWorkPriority QueuedWorkPriority = EQueuedWorkPriority::Normal;
 
 		// Graphics pipeline initializers support overriding the priority for PSO precaching, so start them at low.
 		if constexpr (std::is_same_v<TPipelineInitializer, FGraphicsPipelineStateInitializer>)
 		{
-			CachedState->GetPrecompileTask()->StartBackgroundTask(&GPSOPrecacheThreadPool.Get(), Initializer.bPSOPrecache ? EQueuedWorkPriority::Low : EQueuedWorkPriority::Normal);
-		}
-		else
-		{
-			CachedState->GetPrecompileTask()->StartBackgroundTask(&GPSOPrecacheThreadPool.Get(), EQueuedWorkPriority::Normal);
+			QueuedWorkPriority = Initializer.bPSOPrecache ? EQueuedWorkPriority::Low : EQueuedWorkPriority::Normal;
 		}
+
+		CachedState->CSVStat.SetState(Initializer.bPSOPrecache, CachedState->IsCompute(), GetPSOCompileTypeFromQueuePri(QueuedWorkPriority));
+		CachedState->GetPrecompileTask()->StartBackgroundTask(&GPSOPrecacheThreadPool.Get(), QueuedWorkPriority);
 	}
 }
 
@@ -4260,7 +4450,6 @@ void FPrecacheGraphicsPipelineCache::OnNewPipelineStateCreated(const FGraphicsPi
 
 	// Mark as precache so it will try and use the background thread pool if available
 	FGraphicsPipelineStateInitializer InitializerCopy(Initializer);
-	InitializerCopy.SetPSOPrecacheCompileType(FGraphicsPipelineStateInitializer::EPSOPrecacheCompileType::NormalPri);
 	InitializerCopy.bPSOPrecache = true;
 
 	FPSOCompilationDebugData PSOCompilationDebugData;
