Prevent PSO precache CSV stats appearing on non-precaching platforms.
Slight clean up of PSO precache CSV tracking code.

diff --git a/Engine/Source/Runtime/RHI/Private/PipelineStateCache.cpp b/Engine/Source/Runtime/RHI/Private/PipelineStateCache.cpp
--- a/Engine/Source/Runtime/RHI/Private/PipelineStateCache.cpp
+++ b/Engine/Source/Runtime/RHI/Private/PipelineStateCache.cpp
@@ -758,42 +758,27 @@ class FPSOCSVStatTracker
 	bool bPrecache = false;
 	bool bCompute = false;
 
-	// PSO compiles can take place over multiple frames, use the persistent stats.
-	static inline TCsvPersistentCustomStat<int>* TotalCompletedPSOPrecompiles			= nullptr;
-	static inline TCsvPersistentCustomStat<int>* PSOPrecompileQueueLength				= nullptr;
-	static inline TCsvPersistentCustomStat<int>* PSOComputePrecompilesInProgress		= nullptr;
-	static inline TCsvPersistentCustomStat<int>* MinPriPrecacheInFlight					= nullptr;
-	static inline TCsvPersistentCustomStat<int>* NormalPriPrecacheInFlight				= nullptr;
-	static inline TCsvPersistentCustomStat<int>* MaxPriPrecacheInFlight					= nullptr;
-
-	static inline TCsvPersistentCustomStat<int>* MaxPriPSOPrecompileQueueLength			= nullptr;
-	static inline TCsvPersistentCustomStat<int>* NormalPriPSOPrecompileQueueLength		= nullptr;
-	static inline TCsvPersistentCustomStat<int>* MinPriPSOPrecompileQueueLength			= nullptr;
+	struct FPersistentStats
+	{
+		// PSO compiles can take place over multiple frames, use the persistent stats.
+		TCsvPersistentCustomStat<int>* TotalCompletedPSOPrecompiles			= FCsvProfiler::Get()->GetOrCreatePersistentCustomStatInt(TEXT("TotalCompletedPSOPrecompiles"), CSV_CATEGORY_INDEX(PSOPrecacheCompiling));
+		TCsvPersistentCustomStat<int>* PSOPrecompileQueueLength				= FCsvProfiler::Get()->GetOrCreatePersistentCustomStatInt(TEXT("PSOPrecompileQueueLength"), CSV_CATEGORY_INDEX(PSOPrecacheCompiling));
+		TCsvPersistentCustomStat<int>* PSOComputePrecompilesInProgress		= FCsvProfiler::Get()->GetOrCreatePersistentCustomStatInt(TEXT("PSOComputePrecompilesInProgress"), CSV_CATEGORY_INDEX(PSOPrecacheCompiling));
+		TCsvPersistentCustomStat<int>* MinPriPrecacheInFlight				= FCsvProfiler::Get()->GetOrCreatePersistentCustomStatInt(TEXT("MinPriPrecacheInFlight"), CSV_CATEGORY_INDEX(PSOPrecacheCompiling));
+		TCsvPersistentCustomStat<int>* NormalPriPrecacheInFlight			= FCsvProfiler::Get()->GetOrCreatePersistentCustomStatInt(TEXT("NormalPriPrecacheInFlight"), CSV_CATEGORY_INDEX(PSOPrecacheCompiling));
+		TCsvPersistentCustomStat<int>* MaxPriPrecacheInFlight				= FCsvProfiler::Get()->GetOrCreatePersistentCustomStatInt(TEXT("MaxPriPrecacheInFlight"), CSV_CATEGORY_INDEX(PSOPrecacheCompiling));
+
+		TCsvPersistentCustomStat<int>* MaxPriPSOPrecompileQueueLength		= FCsvProfiler::Get()->GetOrCreatePersistentCustomStatInt(TEXT("MaxPriPSOPrecompileQueueLength"), CSV_CATEGORY_INDEX(PSOPrecacheCompiling));
+		TCsvPersistentCustomStat<int>* NormalPriPSOPrecompileQueueLength	= FCsvProfiler::Get()->GetOrCreatePersistentCustomStatInt(TEXT("NormalPriPSOPrecompileQueueLength"), CSV_CATEGORY_INDEX(PSOPrecacheCompiling));
+		TCsvPersistentCustomStat<int>* MinPriPSOPrecompileQueueLength		= FCsvProfiler::Get()->GetOrCreatePersistentCustomStatInt(TEXT("MinPriPSOPrecompileQueueLength"), CSV_CATEGORY_INDEX(PSOPrecacheCompiling));
+	};
 
 	FGraphicsPipelineStateInitializer::EPSOPrecacheCompileType Pri = FGraphicsPipelineStateInitializer::EPSOPrecacheCompileType::NotSet;
 
-	static void InitStats()
+	static FPersistentStats& GetStats()
 	{
-		static std::atomic_bool bFirstRun = true;
-		if (bFirstRun)
-		{
-			static UE::FMutex WindowManagerLock;
-			UE::TUniqueLock Lock(WindowManagerLock);
-			if (bFirstRun)
-			{
-				TotalCompletedPSOPrecompiles		= FCsvProfiler::Get()->GetOrCreatePersistentCustomStatInt(TEXT("TotalCompletedPSOPrecompiles"), CSV_CATEGORY_INDEX(PSOPrecacheCompiling));
-				PSOPrecompileQueueLength			= FCsvProfiler::Get()->GetOrCreatePersistentCustomStatInt(TEXT("PSOPrecompileQueueLength"), CSV_CATEGORY_INDEX(PSOPrecacheCompiling));
-				PSOComputePrecompilesInProgress		= FCsvProfiler::Get()->GetOrCreatePersistentCustomStatInt(TEXT("PSOComputePrecompilesInProgress"), CSV_CATEGORY_INDEX(PSOPrecacheCompiling));
-				MinPriPrecacheInFlight				= FCsvProfiler::Get()->GetOrCreatePersistentCustomStatInt(TEXT("MinPriPrecacheInFlight"), CSV_CATEGORY_INDEX(PSOPrecacheCompiling));
-				NormalPriPrecacheInFlight			= FCsvProfiler::Get()->GetOrCreatePersistentCustomStatInt(TEXT("NormalPriPrecacheInFlight"), CSV_CATEGORY_INDEX(PSOPrecacheCompiling));
-				MaxPriPrecacheInFlight				= FCsvProfiler::Get()->GetOrCreatePersistentCustomStatInt(TEXT("MaxPriPrecacheInFlight"), CSV_CATEGORY_INDEX(PSOPrecacheCompiling));
-
-				MaxPriPSOPrecompileQueueLength		= FCsvProfiler::Get()->GetOrCreatePersistentCustomStatInt(TEXT("MaxPriPSOPrecompileQueueLength"), CSV_CATEGORY_INDEX(PSOPrecacheCompiling));
-				NormalPriPSOPrecompileQueueLength	= FCsvProfiler::Get()->GetOrCreatePersistentCustomStatInt(TEXT("NormalPriPSOPrecompileQueueLength"), CSV_CATEGORY_INDEX(PSOPrecacheCompiling));
-				MinPriPSOPrecompileQueueLength		= FCsvProfiler::Get()->GetOrCreatePersistentCustomStatInt(TEXT("MinPriPSOPrecompileQueueLength"), CSV_CATEGORY_INDEX(PSOPrecacheCompiling));
-				bFirstRun = false;
-			}
-		}
+		static FPersistentStats PersistentStats;
+		return PersistentStats;
 	}
 	static bool AreStatsEnabled()
 	{
@@ -802,16 +787,6 @@ class FPSOCSVStatTracker
 public:
 	static inline bool bEnableStatTracker = true;
 
-	FPSOCSVStatTracker()
-	{
-		if (!AreStatsEnabled())
-		{ 
-			return;
-		}
-
-		InitStats();
-	}
-
 	void SetState(bool bInPrecache, bool bInCompute, FGraphicsPipelineStateInitializer::EPSOPrecacheCompileType InPri)
 	{
 		if (!AreStatsEnabled())
@@ -819,17 +794,17 @@ public:
 			return;
 		}
 
-		check(TotalCompletedPSOPrecompiles);
 		bPrecache = bInPrecache;
 		bCompute = bInCompute;
 
 		if (bPrecache)
 		{
-			PSOPrecompileQueueLength->Add(1);
+			FPersistentStats& PersistentStats = GetStats();
+			PersistentStats.PSOPrecompileQueueLength->Add(1);
 
 			if (bCompute)
 			{
-				PSOComputePrecompilesInProgress->Add(1);
+				PersistentStats.PSOComputePrecompilesInProgress->Add(1);
 			}
 
 			SetPriority(InPri);
@@ -879,12 +854,14 @@ public:
 
 		if (bPrecache)
 		{
-			TotalCompletedPSOPrecompiles->Add(1);
-			PSOPrecompileQueueLength->Add(-1);
+			FPersistentStats& PersistentStats = GetStats();
+
+			PersistentStats.TotalCompletedPSOPrecompiles->Add(1);
+			PersistentStats.PSOPrecompileQueueLength->Add(-1);
 
 			if (bCompute)
 			{
-				PSOComputePrecompilesInProgress->Add(-1);
+				PersistentStats.PSOComputePrecompilesInProgress->Add(-1);
 			}
 
 			if (ensure(Pri != FGraphicsPipelineStateInitializer::EPSOPrecacheCompileType::NotSet))
@@ -897,14 +874,16 @@ public:
 	private:
 		TCsvPersistentCustomStat<int>* GetPriStat(FGraphicsPipelineStateInitializer::EPSOPrecacheCompileType InPri, bool bInFlight) const
 		{
+			FPersistentStats& PersistentStats = GetStats();
+
 			switch (InPri)
 			{
 			case FGraphicsPipelineStateInitializer::EPSOPrecacheCompileType::MaxPri:
-				return bInFlight ? MaxPriPrecacheInFlight : MaxPriPSOPrecompileQueueLength;
+				return bInFlight ? PersistentStats.MaxPriPrecacheInFlight : PersistentStats.MaxPriPSOPrecompileQueueLength;
 			case FGraphicsPipelineStateInitializer::EPSOPrecacheCompileType::NormalPri:
-				return bInFlight ? NormalPriPrecacheInFlight : NormalPriPSOPrecompileQueueLength ;
+				return bInFlight ? PersistentStats.NormalPriPrecacheInFlight : PersistentStats.NormalPriPSOPrecompileQueueLength ;
 			case FGraphicsPipelineStateInitializer::EPSOPrecacheCompileType::MinPri:
-				return bInFlight ? MinPriPrecacheInFlight : MinPriPSOPrecompileQueueLength;
+				return bInFlight ? PersistentStats.MinPriPrecacheInFlight : PersistentStats.MinPriPSOPrecompileQueueLength;
 			default:
 				checkNoEntry();
 				return nullptr;
