Respect D3D11_ES31 entries in Hardware config files

Make sure DX11 is used and not DX12.

diff --git a/Engine/Source/Runtime/RHI/Private/Windows/WindowsDynamicRHI.cpp b/Engine/Source/Runtime/RHI/Private/Windows/WindowsDynamicRHI.cpp
--- a/Engine/Source/Runtime/RHI/Private/Windows/WindowsDynamicRHI.cpp
+++ b/Engine/Source/Runtime/RHI/Private/Windows/WindowsDynamicRHI.cpp
@@ -448,14 +448,67 @@
 
 // Whether to default to Performance Mode on low-end machines (ES31 feature level).
 // Looks at a few configuration options that control thresholds to enable this as the default.
-static bool ShouldDefaultToPerformanceMode()
+static bool ShouldDefaultToPerformanceMode(bool& bOutForceD3D11)
 {
 	static TOptional<bool> ForceES31;
+	static bool bForceD3D11 = false;
 	if (ForceES31.IsSet())
 	{
+		bOutForceD3D11 = bForceD3D11;
 		return ForceES31.GetValue();
 	}
 
+	FWindowsPlatformApplicationMisc::FGPUInfo BestGPUInfo = FWindowsPlatformApplicationMisc::GetBestGPUInfo();
+
+	TArray<FString> DeviceDefaultRHIList;
+	GConfig->GetArray(TEXT("Devices"), TEXT("DeviceDefaultRHIList"), DeviceDefaultRHIList, GHardwareIni);
+
+	FString GPUBrand = FPlatformMisc::GetPrimaryGPUBrand();
+	for (const FString& DeviceDefaultRHIString : DeviceDefaultRHIList)
+	{
+		const TCHAR* Line = *DeviceDefaultRHIString;
+
+		ensure(Line[0] == TCHAR('('));
+
+		FString RHIName;
+		FParse::Value(Line+1, TEXT("RHI="), RHIName);
+
+		FString DeviceName;
+		FParse::Value(Line+1, TEXT("DeviceName="), DeviceName);
+
+		if (RHIName.Compare("D3D11_ES31", ESearchCase::IgnoreCase) == 0 && GPUBrand.Compare(DeviceName, ESearchCase::IgnoreCase) == 0)
+		{
+			UE_LOG(LogRHI, Log, TEXT("Found RHIName %s with DeviceName %s: default feature level to %s"), *RHIName, *DeviceName, GetLogName(ERHIFeatureLevel::ES3_1));
+
+			ForceES31 = true;
+			bForceD3D11 = true;
+
+			bOutForceD3D11 = true;
+			return true;
+		}
+
+		FString VendorId;
+		FParse::Value(Line + 1, TEXT("VendorId="), VendorId);
+		uint32 VendorIdInt = FParse::HexNumber(*VendorId);
+
+		FString DeviceId;
+		FParse::Value(Line + 1, TEXT("DeviceId="), DeviceId);
+		uint32 DeviceIdInt = FParse::HexNumber(*DeviceId);
+
+		if (BestGPUInfo.VendorId && BestGPUInfo.DeviceId &&
+			BestGPUInfo.VendorId == VendorIdInt && BestGPUInfo.DeviceId == DeviceIdInt &&
+			RHIName.Compare("D3D11_ES31", ESearchCase::IgnoreCase) == 0)
+		{
+			UE_LOG(LogRHI, Log, TEXT("Found RHIName %s with DeviceName %s, VendorId %s, and DeviceId %s: default feature level to %s"), *RHIName, *DeviceName, *VendorId, *DeviceId, GetLogName(ERHIFeatureLevel::ES3_1));
+
+			ForceES31 = true;
+			bForceD3D11 = true;
+
+			bOutForceD3D11 = true;
+			return true;
+		}
+	}
+
 	// Force Performance mode for laptops. We don't have an explicit way to check for laptops
 	// so we use whether we're running on battery as a proxy.
 	bool bIsLaptop = FPlatformMisc::IsRunningOnBattery();
@@ -475,8 +528,6 @@
 		return true;
 	}
 
-	FWindowsPlatformApplicationMisc::FGPUInfo BestGPUInfo = FWindowsPlatformApplicationMisc::GetBestGPUInfo();
-
 	FString MinMemorySizeBucketString;
 	FString MinIntegratedMemorySizeBucketString;
 	if (GConfig->GetString(TEXT("PerformanceMode"), TEXT("MinMemorySizeBucket"), MinMemorySizeBucketString, GEngineIni) && GConfig->GetString(TEXT("PerformanceMode"), TEXT("MinIntegratedMemorySizeBucket"), MinIntegratedMemorySizeBucketString, GEngineIni))
@@ -511,55 +562,14 @@
 		}
 	}
 
-	TArray<FString> DeviceDefaultRHIList;
-	GConfig->GetArray(TEXT("Devices"), TEXT("DeviceDefaultRHIList"), DeviceDefaultRHIList, GHardwareIni);
-
-	FString GPUBrand = FPlatformMisc::GetPrimaryGPUBrand();
-	for (const FString& DeviceDefaultRHIString : DeviceDefaultRHIList)
-	{
-		const TCHAR* Line = *DeviceDefaultRHIString;
-
-		ensure(Line[0] == TCHAR('('));
-
-		FString RHIName;
-		FParse::Value(Line+1, TEXT("RHI="), RHIName);
-
-		FString DeviceName;
-		FParse::Value(Line+1, TEXT("DeviceName="), DeviceName);
-
-		if (RHIName.Compare("D3D11_ES31", ESearchCase::IgnoreCase) == 0 && GPUBrand.Compare(DeviceName, ESearchCase::IgnoreCase) == 0)
-		{
-			UE_LOG(LogRHI, Log, TEXT("Found RHIName %s with DeviceName %s: default feature level to %s"), *RHIName, *DeviceName, GetLogName(ERHIFeatureLevel::ES3_1));
-
-			ForceES31 = true;
-			return true;
-		}
-
-		FString VendorId;
-		FParse::Value(Line + 1, TEXT("VendorId="), VendorId);
-		uint32 VendorIdInt = FParse::HexNumber(*VendorId);
-
-		FString DeviceId;
-		FParse::Value(Line + 1, TEXT("DeviceId="), DeviceId);
-		uint32 DeviceIdInt = FParse::HexNumber(*DeviceId);
-
-		if (BestGPUInfo.VendorId && BestGPUInfo.DeviceId &&
-			BestGPUInfo.VendorId == VendorIdInt && BestGPUInfo.DeviceId == DeviceIdInt &&
-			RHIName.Compare("D3D11_ES31", ESearchCase::IgnoreCase) == 0)
-		{
-			UE_LOG(LogRHI, Log, TEXT("Found RHIName %s with DeviceName %s, VendorId %s, and DeviceId %s: default feature level to %s"), *RHIName, *DeviceName, *VendorId, *DeviceId, GetLogName(ERHIFeatureLevel::ES3_1));
-
-			ForceES31 = true;
-			return true;
-		}
-	}
-
 	ForceES31 = false;
 	return false;
 }
 
-static TOptional<ERHIFeatureLevel::Type> GetPreferredFeatureLevel(EWindowsRHI ChosenRHI, const FParsedWindowsDynamicRHIConfig& Config)
+static TOptional<ERHIFeatureLevel::Type> GetPreferredFeatureLevel(EWindowsRHI ChosenRHI, const FParsedWindowsDynamicRHIConfig& Config, bool& bOutForceD3D11)
 {
+	bOutForceD3D11 = false;
+
 	TOptional<ERHIFeatureLevel::Type> PreferredFeatureLevel;
 	if (!GIsEditor)
 	{
@@ -598,7 +608,7 @@
 			bool bDefaultES31 = false;
 			if (!bFoundPreference && !CVarIgnorePerformanceModeCheck.GetValueOnAnyThread())
 			{
-				bDefaultES31 = ShouldDefaultToPerformanceMode();
+				bDefaultES31 = ShouldDefaultToPerformanceMode(bOutForceD3D11);
 			}
 
 			if (bPreferFeatureLevelES31 || bDefaultES31)
@@ -606,6 +616,10 @@
 				if (!bFoundPreference)
 				{
 					GConfig->SetBool(TEXT("D3DRHIPreference"), TEXT("bPreferFeatureLevelES31"), true, GGameUserSettingsIni);
+					if (bOutForceD3D11)
+					{
+						GConfig->SetString(TEXT("D3DRHIPreference"), TEXT("PreferredRHI"), TEXT("dx11"), GGameUserSettingsIni);
+					}
 				}
 				PreferredFeatureLevel = ERHIFeatureLevel::ES3_1;
 			}
@@ -894,8 +908,10 @@
 	return ForcedFeatureLevel;
 }
 
-static ERHIFeatureLevel::Type ChooseFeatureLevel(EWindowsRHI ChosenRHI, const TOptional<EWindowsRHI> ForcedRHI, const TOptional<ERHIFeatureLevel::Type> ForcedFeatureLevel, const FParsedWindowsDynamicRHIConfig& Config)
+static ERHIFeatureLevel::Type ChooseFeatureLevel(EWindowsRHI ChosenRHI, const TOptional<ERHIFeatureLevel::Type> ForcedFeatureLevel, const FParsedWindowsDynamicRHIConfig& Config, bool& bOutForceD3D11)
 {
+	bOutForceD3D11 = false;
+
 	if (ForcedFeatureLevel)
 	{
 		// Allow the forced feature level if we're in a position to compile its shaders
@@ -913,7 +929,7 @@
 		}
 	}
 
-	TOptional<ERHIFeatureLevel::Type> FeatureLevel = GetPreferredFeatureLevel(ChosenRHI, Config);
+	TOptional<ERHIFeatureLevel::Type> FeatureLevel = GetPreferredFeatureLevel(ChosenRHI, Config, bOutForceD3D11);
 
 	if (!FeatureLevel || (FPlatformProperties::RequiresCookedData() && !Config.IsFeatureLevelTargeted(ChosenRHI, FeatureLevel.GetValue())))
 	{
@@ -1072,7 +1088,13 @@
 		UE_LOG(LogRHI, Log, TEXT("Using Default RHI: %s"), GetLogName(ChosenRHI));
 	}
 
-	DesiredFeatureLevel = ChooseFeatureLevel(ChosenRHI, ForcedRHI, ForcedFeatureLevel, Config);
+	bool bOutForceD3D11 = false;
+	DesiredFeatureLevel = ChooseFeatureLevel(ChosenRHI, ForcedFeatureLevel, Config, bOutForceD3D11);
+	if (bOutForceD3D11)
+	{
+		ChosenRHI = EWindowsRHI::D3D11;
+		UE_LOG(LogRHI, Log, TEXT("Using Device RHI: %s"), GetLogName(ChosenRHI));
+	}
 
 	// Load the dynamic RHI module.
 
