Deprecated BuildFeatureLevelWidget,LoadEditorFeatureLevel,SaveEditorFeatureLevel for BuildShaderPlatformWidget,LoadEditorShaderPlatformPreview,SaveEditorShaderPlatformPreview

Don't consume slots of CustomShaderPlatform enum for each variation of the same ShaderPlatform with different DeviceProfile. Instead only keep track of the DeviceProfiles and the FriendlyNames.

Resolve issue where the Saved state of the Editor would not be correctly loaded if the Saved state had a ShaderPlatform active with a DeviceProfile different from the default one.

Because we are changing the ShaderPlatformName of PreviewPlatforms, which is serialized in FShaderMapBase::Serialize, we require a GUID bump of Global/Material and Niagara Shadermap version.

Deprecated GetCurrentFeatureLevelPreviewText, GetCurrentFeatureLevelPreviewTextVisibility, BuildFeatureLevelWidget  for GetCurrentShaderPlatformPreviewText, GetCurrentShaderPlatformPreviewTextVisibility, BuildShaderPlatformWidget

diff --git a/Engine/Plugins/Enterprise/DataprepEditor/Source/DataprepEditor/Private/Widgets/SDataprepEditorViewport.cpp b/Engine/Plugins/Enterprise/DataprepEditor/Source/DataprepEditor/Private/Widgets/SDataprepEditorViewport.cpp
--- a/Engine/Plugins/Enterprise/DataprepEditor/Source/DataprepEditor/Private/Widgets/SDataprepEditorViewport.cpp
+++ b/Engine/Plugins/Enterprise/DataprepEditor/Source/DataprepEditor/Private/Widgets/SDataprepEditorViewport.cpp
@@ -588,7 +588,7 @@ void SDataprepEditorViewport::PopulateViewportOverlays(TSharedRef<SOverlay> Over
 	.HAlign(HAlign_Right)
 	.Padding(5.0f)
 	[
-		BuildFeatureLevelWidget()
+		BuildShaderPlatformWidget()
 	];
 }
 
diff --git a/Engine/Plugins/Experimental/Dataflow/Source/DataflowEditor/Private/Dataflow/DataflowConstructionViewport.cpp b/Engine/Plugins/Experimental/Dataflow/Source/DataflowEditor/Private/Dataflow/DataflowConstructionViewport.cpp
--- a/Engine/Plugins/Experimental/Dataflow/Source/DataflowEditor/Private/Dataflow/DataflowConstructionViewport.cpp
+++ b/Engine/Plugins/Experimental/Dataflow/Source/DataflowEditor/Private/Dataflow/DataflowConstructionViewport.cpp
@@ -271,7 +271,7 @@ void SDataflowConstructionViewport::PopulateViewportOverlays(TSharedRef<SOverlay
 		.HAlign(HAlign_Right)
 		.Padding(5.0f)
 		[
-			BuildFeatureLevelWidget()
+			BuildShaderPlatformWidget()
 		];
 }
 
diff --git a/Engine/Plugins/FX/Cascade/Source/Cascade/Private/SCascadePreviewViewport.cpp b/Engine/Plugins/FX/Cascade/Source/Cascade/Private/SCascadePreviewViewport.cpp
--- a/Engine/Plugins/FX/Cascade/Source/Cascade/Private/SCascadePreviewViewport.cpp
+++ b/Engine/Plugins/FX/Cascade/Source/Cascade/Private/SCascadePreviewViewport.cpp
@@ -134,7 +134,7 @@ void SCascadePreviewViewport::PopulateViewportOverlays(TSharedRef<SOverlay> Over
 		.HAlign(HAlign_Right)
 		.Padding(5.0f)
 		[
-			BuildFeatureLevelWidget()
+			BuildShaderPlatformWidget()
 		];
 }
 
diff --git a/Engine/Plugins/Runtime/nDisplay/Source/DisplayClusterConfigurator/Private/Views/Viewport/DisplayClusterConfiguratorSCSEditorViewport.cpp b/Engine/Plugins/Runtime/nDisplay/Source/DisplayClusterConfigurator/Private/Views/Viewport/DisplayClusterConfiguratorSCSEditorViewport.cpp
--- a/Engine/Plugins/Runtime/nDisplay/Source/DisplayClusterConfigurator/Private/Views/Viewport/DisplayClusterConfiguratorSCSEditorViewport.cpp
+++ b/Engine/Plugins/Runtime/nDisplay/Source/DisplayClusterConfigurator/Private/Views/Viewport/DisplayClusterConfiguratorSCSEditorViewport.cpp
@@ -597,7 +597,7 @@ void SDisplayClusterConfiguratorSCSEditorViewport::PopulateViewportOverlays(TSha
 		.HAlign(HAlign_Right)
 		.Padding(5.0f)
 		[
-			BuildFeatureLevelWidget()
+			BuildShaderPlatformWidget()
 		];
 }
 
diff --git a/Engine/Source/Editor/Kismet/Private/SSCSEditorViewport.cpp b/Engine/Source/Editor/Kismet/Private/SSCSEditorViewport.cpp
--- a/Engine/Source/Editor/Kismet/Private/SSCSEditorViewport.cpp
+++ b/Engine/Source/Editor/Kismet/Private/SSCSEditorViewport.cpp
@@ -286,7 +286,7 @@ void SSCSEditorViewport::PopulateViewportOverlays(TSharedRef<class SOverlay> Ove
 		.HAlign(HAlign_Right)
 		.Padding(5.0f)
 		[
-			BuildFeatureLevelWidget()
+			BuildShaderPlatformWidget()
 		];
 }
 
diff --git a/Engine/Source/Editor/LevelEditor/Private/LevelEditor.cpp b/Engine/Source/Editor/LevelEditor/Private/LevelEditor.cpp
--- a/Engine/Source/Editor/LevelEditor/Private/LevelEditor.cpp
+++ b/Engine/Source/Editor/LevelEditor/Private/LevelEditor.cpp
@@ -1925,7 +1925,7 @@ void FLevelEditorModule::BindGlobalLevelEditorCommands()
 				}
 
 				const ERHIFeatureLevel::Type FeatureLevel = GetMaxSupportedFeatureLevel(ShaderPlatform);
-				return FPreviewPlatformInfo(FeatureLevel, ShaderPlatform, Item.PlatformName, Item.ShaderFormat, Item.DeviceProfileName, true, Item.PreviewShaderPlatformName, Item.OptionalFriendlyNameOverride);
+				return FPreviewPlatformInfo(FeatureLevel, ShaderPlatform, Item.PlatformName, Item.ShaderFormat, Item.DeviceProfileName, true, Item.PreviewShaderPlatformName);
 			};
 
 			FPreviewPlatformInfo PreviewFeatureLevelInfo = GetPreviewFeatureLevelInfo();
diff --git a/Engine/Source/Editor/LevelEditor/Private/LevelEditorActions.cpp b/Engine/Source/Editor/LevelEditor/Private/LevelEditorActions.cpp
--- a/Engine/Source/Editor/LevelEditor/Private/LevelEditorActions.cpp
+++ b/Engine/Source/Editor/LevelEditor/Private/LevelEditorActions.cpp
@@ -4447,8 +4447,7 @@ void FLevelEditorCommands::RegisterCommands()
 				}
 				else
 				{
-					FriendlyNameBuilder.AppendLine(FDataDrivenShaderPlatformInfo::GetFriendlyName(ShaderPlatform));
-					FPlatformMisc::LowLevelOutputDebugStringf(TEXT("MENU friendly name %s\n"), *FriendlyNameBuilder.ToText().ToString());
+					FriendlyNameBuilder.AppendLine(FDataDrivenShaderPlatformInfo::GetFriendlyName(ShaderPlatform, Item.DeviceProfileName));
 				}
 			}
 		}
@@ -4482,7 +4481,6 @@ void FLevelEditorCommands::RegisterCommands()
 		if (PlatformEngineIni && PlatformEngineIni->GetString(TEXT("DeviceProfileManager"), TEXT("PreviewDeviceProfileSelectionModule"), DeviceProfileSelectionModule, GEngineIni))
 		{
 			EShaderPlatform ShaderPlatform = FDataDrivenShaderPlatformInfo::GetShaderPlatformFromName(Item.ShaderPlatformToPreview);
-			FText PlatformFriendlyName = FDataDrivenShaderPlatformInfo::GetFriendlyName(ShaderPlatform);
 
 			TArray<PreviewPlatformCommand>* PlatformToPreviewJsonPlatformOverridesValue = PlatformToPreviewJsonPlatformOverrides.Find(Item.PlatformName);
 			if (!PlatformToPreviewJsonPlatformOverridesValue)
diff --git a/Engine/Source/Editor/LevelEditor/Private/LevelEditorToolBar.cpp b/Engine/Source/Editor/LevelEditor/Private/LevelEditorToolBar.cpp
--- a/Engine/Source/Editor/LevelEditor/Private/LevelEditorToolBar.cpp
+++ b/Engine/Source/Editor/LevelEditor/Private/LevelEditorToolBar.cpp
@@ -97,17 +97,16 @@ namespace PreviewModeFunctionality
 	{
 #define LOCTEXT_NAMESPACE "LevelEditorToolBar"
 		EShaderPlatform PreviewShaderPlatform = GEditor->PreviewPlatform.GetShaderPlatform();
-
-		EShaderPlatform MaxRHIFeatureLevelPlatform = GetFeatureLevelShaderPlatform(GMaxRHIFeatureLevel);
-
+		FText FriendlyNamePreviewOff = FDataDrivenShaderPlatformInfo::GetFriendlyName(GMaxRHIShaderPlatform);
+		FText FriendlyNamePreviewOn = GEditor->PreviewPlatform.GetFriendlyName();
 		{
-			const FText& RenderingAsPlatformName = GEditor->PreviewPlatform.bPreviewFeatureLevelActive ? GEditor->PreviewPlatform.GetFriendlyName() : FDataDrivenShaderPlatformInfo::GetFriendlyName(MaxRHIFeatureLevelPlatform);
-			const FText& SwitchToPlatformName = GEditor->PreviewPlatform.bPreviewFeatureLevelActive ? FDataDrivenShaderPlatformInfo::GetFriendlyName(MaxRHIFeatureLevelPlatform) :GEditor->PreviewPlatform.GetFriendlyName();
-			if (PreviewShaderPlatform == MaxRHIFeatureLevelPlatform)
+			const FText& RenderingAsPlatformName = GEditor->PreviewPlatform.bPreviewFeatureLevelActive ? FriendlyNamePreviewOn : FriendlyNamePreviewOff;
+			const FText& SwitchToPlatformName = GEditor->PreviewPlatform.bPreviewFeatureLevelActive ? FriendlyNamePreviewOff : FriendlyNamePreviewOn;
+			if (PreviewShaderPlatform == GMaxRHIShaderPlatform)
 			{
 				return FText::Format(LOCTEXT("PreviewModeViewingAs", "Viewing {0}."), RenderingAsPlatformName);
 			}
-			else if (GWorld->GetFeatureLevel() == GMaxRHIFeatureLevel)
+			else if (GEditor->IsFeatureLevelPreviewActive())
 			{
 				return FText::Format(LOCTEXT("PreviewModeViewingAsSwitchTo", "Viewing {0}. Click to preview {1}."), RenderingAsPlatformName, SwitchToPlatformName);
 			}
diff --git a/Engine/Source/Editor/LevelEditor/Private/SLevelViewport.cpp b/Engine/Source/Editor/LevelEditor/Private/SLevelViewport.cpp
--- a/Engine/Source/Editor/LevelEditor/Private/SLevelViewport.cpp
+++ b/Engine/Source/Editor/LevelEditor/Private/SLevelViewport.cpp
@@ -425,7 +425,7 @@ void SLevelViewport::ConstructViewportOverlayContent()
 	auto GetCombinedVisibility = [this]()
 	{
 		if (GetCurrentScreenPercentageVisibility() == EVisibility::Collapsed &&
-			GetCurrentFeatureLevelPreviewTextVisibility() == EVisibility::Collapsed &&
+			GetCurrentShaderPlatformPreviewTextVisibility() == EVisibility::Collapsed &&
 			GetSelectedActorsCurrentLevelTextVisibility() == EVisibility::Collapsed &&
 			!IsActorEditorContextVisible())
 		{
@@ -480,7 +480,7 @@ void SLevelViewport::ConstructViewportOverlayContent()
 				.AutoHeight()
 				.Padding(2.0f, 1.0f, 2.0f, 1.0f)
 				[
-					BuildFeatureLevelWidget()
+					BuildShaderPlatformWidget()
 				]
 				+ SVerticalBox::Slot()
 				.AutoHeight()
diff --git a/Engine/Source/Editor/MaterialEditor/Private/SMaterialEditorViewport.cpp b/Engine/Source/Editor/MaterialEditor/Private/SMaterialEditorViewport.cpp
--- a/Engine/Source/Editor/MaterialEditor/Private/SMaterialEditorViewport.cpp
+++ b/Engine/Source/Editor/MaterialEditor/Private/SMaterialEditorViewport.cpp
@@ -1026,7 +1026,7 @@ void SMaterialEditor3DPreviewViewport::PopulateViewportOverlays(TSharedRef<class
 		.HAlign(HAlign_Right)
 		.Padding(5.0f)
 		[
-			BuildFeatureLevelWidget()
+			BuildShaderPlatformWidget()
 		];
 		// clang-format on
 }
diff --git a/Engine/Source/Editor/Persona/Private/SAnimationEditorViewport.cpp b/Engine/Source/Editor/Persona/Private/SAnimationEditorViewport.cpp
--- a/Engine/Source/Editor/Persona/Private/SAnimationEditorViewport.cpp
+++ b/Engine/Source/Editor/Persona/Private/SAnimationEditorViewport.cpp
@@ -122,7 +122,7 @@ void SAnimationEditorViewport::PopulateViewportOverlays(TSharedRef<SOverlay> Ove
 		.HAlign(HAlign_Right)
 		.Padding(5.0f)
 		[
-			BuildFeatureLevelWidget()
+			BuildShaderPlatformWidget()
 		];
 
 	if (TSharedPtr<SAnimationEditorViewportTabBody> ViewportTab = TabBodyPtr.Pin())
diff --git a/Engine/Source/Editor/StaticMeshEditor/Private/SStaticMeshEditorViewport.cpp b/Engine/Source/Editor/StaticMeshEditor/Private/SStaticMeshEditorViewport.cpp
--- a/Engine/Source/Editor/StaticMeshEditor/Private/SStaticMeshEditorViewport.cpp
+++ b/Engine/Source/Editor/StaticMeshEditor/Private/SStaticMeshEditorViewport.cpp
@@ -106,7 +106,7 @@ void SStaticMeshEditorViewport::PopulateViewportOverlays(TSharedRef<SOverlay> Ov
 		.HAlign(HAlign_Right)
 		.Padding(5.0f)
 		[
-			BuildFeatureLevelWidget()
+			BuildShaderPlatformWidget()
 		];
 }
 
diff --git a/Engine/Source/Editor/UnrealEd/Classes/Editor/EditorEngine.h b/Engine/Source/Editor/UnrealEd/Classes/Editor/EditorEngine.h
--- a/Engine/Source/Editor/UnrealEd/Classes/Editor/EditorEngine.h
+++ b/Engine/Source/Editor/UnrealEd/Classes/Editor/EditorEngine.h
@@ -3208,6 +3208,9 @@ public:
 	/** Return the current active shader platform in the editor*/
 	UNREALED_API EShaderPlatform GetActiveShaderPlatform() const;
 
+	/** Return the current active shader platform friendly name in the editor*/
+	UNREALED_API FText GetActiveShaderPlatformFriendlyName() const;
+
 	/**
 	 * Return the active feature level. This will be the chosen Settings->Preview Rendering Level if the Preview Mode button
 	 * is highlighted or SM5 if the Preview Mode button has been clicked and isn't highlighted.
@@ -3225,10 +3228,15 @@ public:
 
 protected:
 
-	/** Function pair used to save and restore the global feature level */
+	/** Function pair used to save and restore the global Shader Platform */
+	UE_DEPRECATED(5.7, "LoadEditorFeatureLevel is deprecated, please use LoadEditorShaderPlatformPreview")
 	UNREALED_API void LoadEditorFeatureLevel();
+	UE_DEPRECATED(5.7, "SaveEditorFeatureLevel is deprecated, please use LoadEditorShaderPlatformPreview")
 	UNREALED_API void SaveEditorFeatureLevel();
 
+	UNREALED_API void LoadEditorShaderPlatformPreview();
+	UNREALED_API void SaveEditorShaderPlatformPreview();
+
 	/** Utility function that can determine whether some input world is using materials who's shaders are emulated in the editor */
 	UNREALED_API bool IsEditorShaderPlatformEmulated(UWorld* World);
 
diff --git a/Engine/Source/Editor/UnrealEd/Private/EditorEngine.cpp b/Engine/Source/Editor/UnrealEd/Private/EditorEngine.cpp
--- a/Engine/Source/Editor/UnrealEd/Private/EditorEngine.cpp
+++ b/Engine/Source/Editor/UnrealEd/Private/EditorEngine.cpp
@@ -469,9 +469,7 @@ void FPreviewPlatformInfo::InternalSetFriendlyName()
 {
 	if (PreviewShaderPlatformFriendlyName.IsEmpty())
 	{
-		EShaderPlatform PreviewShaderPlatform = GetShaderPlatform();
-		EShaderPlatform MaxRHIFeatureLevelPlatform = GetFeatureLevelShaderPlatform(GMaxRHIFeatureLevel);
-		PreviewShaderPlatformFriendlyName = FDataDrivenShaderPlatformInfo::GetFriendlyName(bPreviewFeatureLevelActive ? PreviewShaderPlatform : MaxRHIFeatureLevelPlatform);
+		PreviewShaderPlatformFriendlyName = FDataDrivenShaderPlatformInfo::GetFriendlyName(GetShaderPlatform(), DeviceProfileName);
 	}
 }
 
@@ -1312,7 +1310,7 @@ void UEditorEngine::Init(IEngineLoop* InEngineLoop)
 	GEditor = this;
 	InitEditor(InEngineLoop);
 
-	LoadEditorFeatureLevel();
+	LoadEditorShaderPlatformPreview();
 
 
 	// Init transactioning.
@@ -6394,7 +6392,7 @@ void UEditorEngine::OnEffectivePreviewShaderPlatformChange()
 	}
 }
 
-static void SaveFeatureLevelAsDisabled(FPreviewPlatformInfo& PreviewPlatform)
+static void SaveShaderPlatformAsDisabled(FPreviewPlatformInfo& PreviewPlatform)
 {
 	auto* Settings = GetMutableDefault<UEditorPerProjectUserSettings>();
 
@@ -6439,7 +6437,7 @@ void UEditorEngine::SetPreviewPlatform(const FPreviewPlatformInfo& NewPreviewPla
 	PreviewPlatform = NewPreviewPlatform;
 
 	// Initially set preview as disabled in case it fails
-	SaveFeatureLevelAsDisabled(PreviewPlatform);
+	SaveShaderPlatformAsDisabled(PreviewPlatform);
 
 	// If we changed the preview platform, we need to update the material quality settings
 	if (bChangedPreviewShaderPlatform)
@@ -6531,7 +6529,7 @@ void UEditorEngine::SetPreviewPlatform(const FPreviewPlatformInfo& NewPreviewPla
 
 	if (bSaveSettings)
 	{
-		SaveEditorFeatureLevel();
+		SaveEditorShaderPlatformPreview();
 	}
 }
 
@@ -6590,7 +6588,7 @@ void UEditorEngine::ToggleFeatureLevelPreview()
 
 	GEditor->RedrawAllViewports();
 	
-	SaveEditorFeatureLevel();
+	SaveEditorShaderPlatformPreview();
 }
 
 bool UEditorEngine::IsFeatureLevelPreviewEnabled() const
@@ -6605,13 +6603,26 @@ bool UEditorEngine::IsFeatureLevelPreviewActive() const
 
 EShaderPlatform UEditorEngine::GetActiveShaderPlatform() const
 {
-	EShaderPlatform ActiveShaderPlatform = GShaderPlatformForFeatureLevel[GMaxRHIFeatureLevel];
 	if (PreviewPlatform.bPreviewFeatureLevelActive)
 	{
-		ActiveShaderPlatform = GShaderPlatformForFeatureLevel[PreviewPlatform.PreviewFeatureLevel];
+		return GShaderPlatformForFeatureLevel[PreviewPlatform.PreviewFeatureLevel];
 	}
+	else
+	{
+		return GMaxRHIShaderPlatform;
+	}
+}
 
-	return ActiveShaderPlatform;
+FText UEditorEngine::GetActiveShaderPlatformFriendlyName() const
+{
+	if (PreviewPlatform.bPreviewFeatureLevelActive)
+	{
+		return PreviewPlatform.GetFriendlyName();
+	}
+	else
+	{
+		return FDataDrivenShaderPlatformInfo::GetFriendlyName(GMaxRHIShaderPlatform);
+	}
 }
 
 ERHIFeatureLevel::Type UEditorEngine::GetActiveFeatureLevelPreviewType() const
@@ -6619,7 +6630,13 @@ ERHIFeatureLevel::Type UEditorEngine::GetActiveFeatureLevelPreviewType() const
 	return PreviewPlatform.bPreviewFeatureLevelActive ? PreviewPlatform.PreviewFeatureLevel : GMaxRHIFeatureLevel;
 }
 
+// Deprecated in 5.7
 void UEditorEngine::LoadEditorFeatureLevel()
+{
+	LoadEditorShaderPlatformPreview();
+}
+
+void UEditorEngine::LoadEditorShaderPlatformPreview()
 {
 	auto* Settings = GetMutableDefault<UEditorPerProjectUserSettings>();
 
@@ -6648,7 +6665,7 @@ void UEditorEngine::LoadEditorFeatureLevel()
 	}
 }
 
-void UEditorEngine::SaveEditorFeatureLevel()
+void UEditorEngine::SaveEditorShaderPlatformPreview()
 {
 	auto* Settings = GetMutableDefault<UEditorPerProjectUserSettings>();
 	Settings->PreviewFeatureLevel = (int32)PreviewPlatform.PreviewFeatureLevel;
diff --git a/Engine/Source/Editor/UnrealEd/Private/PlayLevel.cpp b/Engine/Source/Editor/UnrealEd/Private/PlayLevel.cpp
--- a/Engine/Source/Editor/UnrealEd/Private/PlayLevel.cpp
+++ b/Engine/Source/Editor/UnrealEd/Private/PlayLevel.cpp
@@ -132,7 +132,7 @@ DEFINE_LOG_CATEGORY(LogPlayLevel);
 const static FName NAME_CategoryPIE("PIE");
 
 // Forward declare local utility functions
-FText GeneratePIEViewportWindowTitle(const EPlayNetMode InNetMode, const ERHIFeatureLevel::Type InFeatureLevel, const FRequestPlaySessionParams& InSessionParams, const int32 ClientIndex, const int32 InViewportIndex, const float FixedTick, const bool bVRPreview);
+FText GeneratePIEViewportWindowTitle(const EPlayNetMode InNetMode, const FRequestPlaySessionParams& InSessionParams, const int32 ClientIndex, const int32 InViewportIndex, const float FixedTick, const bool bVRPreview);
 bool IsPrimaryPIEClient(const FRequestPlaySessionParams& InPlaySessionParams, const int32 InClientIndex);
 
 // This class listens to output log messages, and forwards warnings and errors to the message log
@@ -3204,7 +3204,7 @@ UGameInstance* UEditorEngine::CreateInnerProcessPIEGameInstance(FRequestPlaySess
 	return GameInstance;
 }
 
-FText GeneratePIEViewportWindowTitle(const EPlayNetMode InNetMode, const ERHIFeatureLevel::Type InFeatureLevel, const FRequestPlaySessionParams& InSessionParams, const int32 ClientIndex, const float FixedTick, const bool bVRPreview)
+FText GeneratePIEViewportWindowTitle(const EPlayNetMode InNetMode, const FRequestPlaySessionParams& InSessionParams, const int32 ClientIndex, const float FixedTick, const bool bVRPreview)
 {
 #if PLATFORM_64BITS
 	const FString PlatformBitsString(TEXT("64"));
@@ -3217,7 +3217,7 @@ FText GeneratePIEViewportWindowTitle(const EPlayNetMode InNetMode, const ERHIFea
 	FFormatNamedArguments Args;
 	Args.Add(TEXT("GameName"), FText::FromString(FString(WindowTitleOverride.IsEmpty() ? FApp::GetProjectName() : WindowTitleOverride.ToString())));
 	Args.Add(TEXT("PlatformBits"), FText::FromString(PlatformBitsString));
-	Args.Add(TEXT("RHIName"), FDataDrivenShaderPlatformInfo::GetFriendlyName(GetFeatureLevelShaderPlatform(InFeatureLevel)));
+	Args.Add(TEXT("RHIName"), GEditor ? GEditor->GetActiveShaderPlatformFriendlyName(): FDataDrivenShaderPlatformInfo::GetFriendlyName(GMaxRHIShaderPlatform));
 	
 	if (InNetMode == PIE_Client)
 	{
@@ -3322,7 +3322,7 @@ TSharedRef<SPIEViewport> UEditorEngine::GeneratePIEViewportWindow(const FRequest
 	// If they haven't provided a Slate Window (common), we will create one.
 	if (!bHasCustomWindow)
 	{
-		FText ViewportName = GeneratePIEViewportWindowTitle(InNetMode, PreviewPlatform.GetEffectivePreviewFeatureLevel(), InSessionParams, InWorldContext.PIEInstance, InWorldContext.PIEFixedTickSeconds, bVRPreview);
+		FText ViewportName = GeneratePIEViewportWindowTitle(InNetMode, InSessionParams, InWorldContext.PIEInstance, InWorldContext.PIEFixedTickSeconds, bVRPreview);
 		PieWindow = SNew(SWindow)
 			.Title(ViewportName)
 			.Tag("PIEWindow")
@@ -3457,7 +3457,7 @@ TSharedRef<SPIEViewport> UEditorEngine::GeneratePIEViewportWindow(const FRequest
 			{
 				if (!bHasCustomWindow)
 				{
-					FText ViewportName = GeneratePIEViewportWindowTitle(InNetMode, NewFeatureLevel, InSessionParams, InWorldContext.PIEInstance, InWorldContext.PIEFixedTickSeconds, bVRPreview);
+					FText ViewportName = GeneratePIEViewportWindowTitle(InNetMode, InSessionParams, InWorldContext.PIEInstance, InWorldContext.PIEFixedTickSeconds, bVRPreview);
 					PieWindow->SetTitle(ViewportName);
 				}
 				InViewportClient->GetWorld()->ChangeFeatureLevel(NewFeatureLevel);
diff --git a/Engine/Source/Editor/UnrealEd/Private/SEditorViewport.cpp b/Engine/Source/Editor/UnrealEd/Private/SEditorViewport.cpp
--- a/Engine/Source/Editor/UnrealEd/Private/SEditorViewport.cpp
+++ b/Engine/Source/Editor/UnrealEd/Private/SEditorViewport.cpp
@@ -1167,16 +1167,22 @@ EShaderPlatform SEditorViewport::GetShaderPlatformHelper(const ERHIFeatureLevel:
 	return ShaderPlatform;
 }
 
+// Deprecated in 5.7
 TSharedRef<SWidget> SEditorViewport::BuildFeatureLevelWidget() const
+{
+	return BuildShaderPlatformWidget();
+}
+
+TSharedRef<SWidget> SEditorViewport::BuildShaderPlatformWidget() const
 {
 	TSharedRef<SWidget> BoxWidget = SNew(SHorizontalBox)
-		.Visibility(this, &SEditorViewport::GetCurrentFeatureLevelPreviewTextVisibility)
+		.Visibility(this, &SEditorViewport::GetCurrentShaderPlatformPreviewTextVisibility)
 		+ SHorizontalBox::Slot()
 		.AutoWidth()
 		.Padding(2.0f, 1.0f, 2.0f, 1.0f)
 		[
 			SNew(STextBlock)
-			.Text(this, &SEditorViewport::GetCurrentFeatureLevelPreviewText, true)
+			.Text(this, &SEditorViewport::GetCurrentShaderPlatformPreviewText, true)
 			.ShadowOffset(FVector2D(1, 1))
 		]
 		+ SHorizontalBox::Slot()
@@ -1184,14 +1190,20 @@ TSharedRef<SWidget> SEditorViewport::BuildFeatureLevelWidget() const
 		.Padding(4.0f, 1.0f, 2.0f, 1.0f)
 		[
 			SNew(STextBlock)
-			.Text(this, &SEditorViewport::GetCurrentFeatureLevelPreviewText, false)
+			.Text(this, &SEditorViewport::GetCurrentShaderPlatformPreviewText, false)
 			.ShadowOffset(FVector2D(1, 1))
 		];
 
 	return BoxWidget;
 }
 
+// Deprecated in 5.7
 EVisibility SEditorViewport::GetCurrentFeatureLevelPreviewTextVisibility() const
+{
+	return GetCurrentShaderPlatformPreviewTextVisibility();
+}
+
+EVisibility SEditorViewport::GetCurrentShaderPlatformPreviewTextVisibility() const
 {
 	if (Client->GetWorld() && !GLevelEditorModeTools().IsViewportUIHidden())
 	{
@@ -1203,7 +1215,13 @@ EVisibility SEditorViewport::GetCurrentFeatureLevelPreviewTextVisibility() const
 	}
 }
 
+// Deprecated in 5.7
 FText SEditorViewport::GetCurrentFeatureLevelPreviewText(bool bDrawOnlyLabel) const
+{
+	return GetCurrentShaderPlatformPreviewText(bDrawOnlyLabel);
+}
+
+FText SEditorViewport::GetCurrentShaderPlatformPreviewText(bool bDrawOnlyLabel) const
 {
 	FText LabelName;
 
@@ -1213,13 +1231,8 @@ FText SEditorViewport::GetCurrentFeatureLevelPreviewText(bool bDrawOnlyLabel) co
 	}
 	else
 	{
-		UWorld* World = Client->GetWorld();
-		if (World != nullptr)
-		{
-			ERHIFeatureLevel::Type TargetFeatureLevel = World->GetFeatureLevel();
-			const FText& PlatformText = GEditor->PreviewPlatform.GetFriendlyName();
-			LabelName = FText::Format(LOCTEXT("WorldFeatureLevel", "{0}"), PlatformText);
-		}
+		const FText& PlatformText = GEditor->PreviewPlatform.GetFriendlyName();
+		LabelName = FText::Format(LOCTEXT("WorldFeatureLevel", "{0}"), PlatformText);
 	}
 
 	return LabelName;
diff --git a/Engine/Source/Editor/UnrealEd/Private/Subsystems/EditorSubsystemBlueprintLibrary.cpp b/Engine/Source/Editor/UnrealEd/Private/Subsystems/EditorSubsystemBlueprintLibrary.cpp
--- a/Engine/Source/Editor/UnrealEd/Private/Subsystems/EditorSubsystemBlueprintLibrary.cpp
+++ b/Engine/Source/Editor/UnrealEd/Private/Subsystems/EditorSubsystemBlueprintLibrary.cpp
@@ -76,7 +76,7 @@ void UEditorSubsystemBlueprintLibrary::SetPreviewPlatform(FName PreviewShaderPla
 						}
 
 						const ERHIFeatureLevel::Type FeatureLevel = GetMaxSupportedFeatureLevel(ShaderPlatform);
-						return FPreviewPlatformInfo(FeatureLevel, ShaderPlatform, Item.PlatformName, Item.ShaderFormat, Item.DeviceProfileName, true, Item.PreviewShaderPlatformName, Item.OptionalFriendlyNameOverride);
+						return FPreviewPlatformInfo(FeatureLevel, ShaderPlatform, Item.PlatformName, Item.ShaderFormat, Item.DeviceProfileName, true, Item.PreviewShaderPlatformName);
 					};
 
 				FPreviewPlatformInfo PreviewFeatureLevelInfo = GetPreviewFeatureLevelInfo();
diff --git a/Engine/Source/Editor/UnrealEd/Public/SEditorViewport.h b/Engine/Source/Editor/UnrealEd/Public/SEditorViewport.h
--- a/Engine/Source/Editor/UnrealEd/Public/SEditorViewport.h
+++ b/Engine/Source/Editor/UnrealEd/Public/SEditorViewport.h
@@ -123,26 +123,32 @@ public:
 	void MarkLegacyToolbarChildAsAutomaticallyUpgradable(const TSharedRef<SWidget>& ExpectedChild);
 
 ///////////////////////////////////////////////////////////////////////////////
-// begin feature level control functions block
+// begin shader platform control functions block
 ///////////////////////////////////////////////////////////////////////////////
 private:
-	/** Called to get the feature level preview text */
+	/** Called to get the shader platform preview text */
+	UE_DEPRECATED(5.7, "GetCurrentFeatureLevelPreviewText is deprecated, please use GetCurrentShaderPlatformPreviewText")
 	UNREALED_API FText GetCurrentFeatureLevelPreviewText(bool bDrawOnlyLabel) const;
+	UNREALED_API FText GetCurrentShaderPlatformPreviewText(bool bDrawOnlyLabel) const;
 
 	/** Helper function that, for some FeatureLevel argument, will retrieve the required shader platform */
 	UNREALED_API EShaderPlatform GetShaderPlatformHelper(const ERHIFeatureLevel::Type InFeatureLevel) const;
 
 protected:
-	/** @return The visibility of the current feature level preview text display */
+	/** @return The visibility of the current shader platform preview text display */
+	UE_DEPRECATED(5.7, "GetCurrentFeatureLevelPreviewTextVisibility is deprecated, please use GetCurrentShaderPlatformPreviewTextVisibility")
 	UNREALED_API EVisibility GetCurrentFeatureLevelPreviewTextVisibility() const;
+	UNREALED_API EVisibility GetCurrentShaderPlatformPreviewTextVisibility() const;
 
 	/** @return true if realtime can be toggled (it cannot be toggled directly if there is an override in place) */
 	UNREALED_API bool CanToggleRealtime() const;
 
-	/** call this function to build a 'text' widget that can display the present feature level */
+	/** call this function to build a 'text' widget that can display the current Shader Platform */
+	UE_DEPRECATED(5.7, "BuildFeatureLevelWidget is deprecated, please use BuildShaderPlatformWidget")
 	UNREALED_API TSharedRef<SWidget> BuildFeatureLevelWidget() const;
+	UNREALED_API TSharedRef<SWidget> BuildShaderPlatformWidget() const;
 ///////////////////////////////////////////////////////////////////////////////
-// end feature level control functions block
+// end shader platform control functions block
 ///////////////////////////////////////////////////////////////////////////////
 
 	/** Called by the fixed EV100 slider to get the fixed EV100 value */
diff --git a/Engine/Source/Runtime/Core/Private/Misc/DataDrivenPlatformInfoRegistry.cpp b/Engine/Source/Runtime/Core/Private/Misc/DataDrivenPlatformInfoRegistry.cpp
--- a/Engine/Source/Runtime/Core/Private/Misc/DataDrivenPlatformInfoRegistry.cpp
+++ b/Engine/Source/Runtime/Core/Private/Misc/DataDrivenPlatformInfoRegistry.cpp
@@ -19,6 +19,7 @@ namespace
 	TArray<FName> SortedPlatformNames;
 	TArray<const FDataDrivenPlatformInfo*> SortedPlatformInfos;
 #if DDPI_HAS_EXTENDED_PLATFORMINFO_DATA
+	TMap<FName, TArray<FPreviewPlatformMenuItem>> ShaderPlatformNameToMenuItem;
 	TArray<struct FPreviewPlatformMenuItem> PreviewPlatformMenuItems;
 	TSet<FName> PlatformsHiddenFromUI;
 #endif
@@ -265,7 +266,6 @@ static void ParsePreviewPlatforms(const FConfigFile& IniFile)
 			if (Section.Key.StartsWith(TEXT("PreviewPlatform ")))
 			{
 				const FString& SectionName = Section.Key;
-				FName PreviewPlatformName = *(SectionName.Mid(16) + TEXT("_Preview"));
 
 				// Early-out if enabled / disabled cvar is specified and not set (or the other way around)
 				auto ShouldEarlyOut = [&](const FName& KeyName, bool bValueToEarlyOut) -> bool
@@ -299,6 +299,10 @@ static void ParsePreviewPlatforms(const FConfigFile& IniFile)
 				
 				FPreviewPlatformMenuItem Item;
 				Item.PlatformName = PlatformName;
+				Item.ShaderPlatformToPreview = *GetSectionString(Section.Value, FName("ShaderPlatform"));
+				checkf(Item.ShaderPlatformToPreview != NAME_None, TEXT("DataDrivenPlatformInfo section [PreviewPlatform %s] must specify a ShaderPlatformToPreview"), *SectionName);
+
+				FName PreviewPlatformName = FName(Item.ShaderPlatformToPreview.ToString() + TEXT("_Preview"));
 				Item.PreviewShaderPlatformName = GetSectionString(Section.Value, FName("PreviewShaderPlatform"), PreviewPlatformName);
 				Item.ShaderFormat = *GetSectionString(Section.Value, FName("ShaderFormat"));
 				checkf(Item.ShaderFormat != NAME_None, TEXT("DataDrivenPlatformInfo section [PreviewPlatform %s] must specify a ShaderFormat"), *SectionName);
@@ -306,10 +310,10 @@ static void ParsePreviewPlatforms(const FConfigFile& IniFile)
 				Item.ActiveIconName = *GetSectionString(Section.Value, FName("ActiveIconName"));
 				Item.InactiveIconPath = GetSectionString(Section.Value, FName("InactiveIconPath"));
 				Item.InactiveIconName = *GetSectionString(Section.Value, FName("InactiveIconName"));
-				Item.ShaderPlatformToPreview = *GetSectionString(Section.Value, FName("ShaderPlatform"));
+				
 				Item.PreviewFeatureLevelName = *GetSectionString(Section.Value, FName("PreviewFeatureLevel"));
 
-				checkf(Item.ShaderPlatformToPreview != NAME_None, TEXT("DataDrivenPlatformInfo section [PreviewPlatform %s] must specify a ShaderPlatformToPreview"), *SectionName);
+				
 				FTextStringHelper::ReadFromBuffer(*GetSectionString(Section.Value, FName("MenuTooltip")), Item.MenuTooltip);
 				FTextStringHelper::ReadFromBuffer(*GetSectionString(Section.Value, FName("IconText")), Item.IconText);
 
@@ -338,6 +342,7 @@ static void ParsePreviewPlatforms(const FConfigFile& IniFile)
 						Item.OptionalFriendlyNameOverride = FText::FromString(Item.DeviceProfileName.ToString());
 					}
 									
+					ShaderPlatformNameToMenuItem.FindOrAdd(Item.ShaderPlatformToPreview).Add(Item);
 					PreviewPlatformMenuItems.Add(Item);
 				}
 			}
@@ -667,6 +672,11 @@ const TArray<struct FPreviewPlatformMenuItem>& FDataDrivenPlatformInfoRegistry::
 	return PreviewPlatformMenuItems;
 }
 
+const TArray<struct FPreviewPlatformMenuItem>* FDataDrivenPlatformInfoRegistry::GetMenuItemForShaderPlatformName(FName ShaderPlatformName)
+{
+	return ShaderPlatformNameToMenuItem.Find(ShaderPlatformName);
+}
+
 bool FDataDrivenPlatformInfoRegistry::IsPlatformHiddenFromUI(FName PlatformName)
 {
 	return PlatformsHiddenFromUI.Contains(PlatformName);
diff --git a/Engine/Source/Runtime/Core/Private/UObject/FortniteReleaseBranchCustomObjectVersion.cpp b/Engine/Source/Runtime/Core/Private/UObject/FortniteReleaseBranchCustomObjectVersion.cpp
--- a/Engine/Source/Runtime/Core/Private/UObject/FortniteReleaseBranchCustomObjectVersion.cpp
+++ b/Engine/Source/Runtime/Core/Private/UObject/FortniteReleaseBranchCustomObjectVersion.cpp
@@ -6,12 +6,12 @@ TMap<FGuid, FGuid> FFortniteReleaseBranchCustomObjectVersion::GetSystemGuids()
 	TMap<FGuid, FGuid> SystemGuids;
 	const FDevSystemGuids& DevGuids = FDevSystemGuids::Get();
 
-	SystemGuids.Add(DevGuids.Niagara_LatestScriptCompileVersion, FGuid("8BC8C557C87C4F629468F10EA4161F0A"));
+	SystemGuids.Add(DevGuids.Niagara_LatestScriptCompileVersion, FGuid("213B319D239D438C82A946BB164152DE"));
 	SystemGuids.Add(DevGuids.SkeletalMeshDerivedDataVersion, FGuid("C96DEBAD45E3443D804017D4C29F7C4B"));
 	SystemGuids.Add(DevGuids.MaterialTranslationDDCVersion, FGuid("5958B746DC5E4046AA7FB7E30E7AC509"));
-	SystemGuids.Add(DevGuids.GLOBALSHADERMAP_DERIVEDDATA_VER, FGuid("0A0A1BF5D2A94105BB395EF0DA7CAEB9"));
-	SystemGuids.Add(DevGuids.MATERIALSHADERMAP_DERIVEDDATA_VER, FGuid("C0E1F59386894C9D855B85178B2553CB"));
+	SystemGuids.Add(DevGuids.GLOBALSHADERMAP_DERIVEDDATA_VER, FGuid("43D841E68BB14DF4A45F6D09B35501F7"));
+	SystemGuids.Add(DevGuids.MATERIALSHADERMAP_DERIVEDDATA_VER, FGuid("D52C1D148C294A6A8C6CD75505081A31"));
 	SystemGuids.Add(DevGuids.STATICMESH_DERIVEDDATA_VER, FGuid("0D7C30DE6E5F489A831CEBF1346725C4"));
 	SystemGuids.Add(DevGuids.NANITE_DERIVEDDATA_VER, FGuid("B08A8584408545A88E5ED193054085E3"));
 	return SystemGuids;
-}
+}
\ No newline at end of file
diff --git a/Engine/Source/Runtime/Core/Public/Misc/DataDrivenPlatformInfoRegistry.h b/Engine/Source/Runtime/Core/Public/Misc/DataDrivenPlatformInfoRegistry.h
--- a/Engine/Source/Runtime/Core/Public/Misc/DataDrivenPlatformInfoRegistry.h
+++ b/Engine/Source/Runtime/Core/Public/Misc/DataDrivenPlatformInfoRegistry.h
@@ -355,6 +355,11 @@ struct FDataDrivenPlatformInfoRegistry
 	 * Retrieve the full list of all known preview platforms from all DDPIs
 	*/
 	static CORE_API const TArray<struct FPreviewPlatformMenuItem>& GetAllPreviewPlatformMenuItems();
+	
+	/**
+	 * Retrieve the PreviewPlatform MenuItems for ShaderPlatformName
+	*/
+	static CORE_API const TArray<struct FPreviewPlatformMenuItem>* GetMenuItemForShaderPlatformName(FName ShaderPlatformName);
 
 	/** Option to hide a platform from the user interface at runtime */
 	static CORE_API bool IsPlatformHiddenFromUI(FName PlatformName);
diff --git a/Engine/Source/Runtime/RHI/Private/DataDrivenShaderPlatformInfo.cpp b/Engine/Source/Runtime/RHI/Private/DataDrivenShaderPlatformInfo.cpp
--- a/Engine/Source/Runtime/RHI/Private/DataDrivenShaderPlatformInfo.cpp
+++ b/Engine/Source/Runtime/RHI/Private/DataDrivenShaderPlatformInfo.cpp
@@ -24,7 +24,7 @@ static TMap<FName, EShaderPlatform> PlatformNameToShaderPlatformMap;
 #if WITH_EDITOR
 struct FDataDrivenShaderPlatformInfoEditorOnly
 {
-	FText FriendlyName;
+	TMap<FName,FText> DeviceProfileToFriendlyName;
 	EShaderPlatform PreviewShaderPlatformParent = SP_NumPlatforms;
 	uint32 bCanUsePreviewPlatformForMaterialValidation : 1;
 };
@@ -330,7 +330,9 @@ void FGenericDataDrivenShaderPlatformInfo::ParseDataDrivenShaderInfo(const FConf
 
 #if WITH_EDITOR
 	FDataDrivenShaderPlatformInfoEditorOnly& EditorInfo = DataDrivenShaderPlatformInfoEditorOnlyInfos[Index];
-	FTextStringHelper::ReadFromBuffer(*GetSectionString(Section, FName("FriendlyName")), EditorInfo.FriendlyName);
+	FText FriendlyName;
+	FTextStringHelper::ReadFromBuffer(*GetSectionString(Section, FName("FriendlyName")), FriendlyName);
+	EditorInfo.DeviceProfileToFriendlyName.FindOrAdd(NAME_None) = FriendlyName;
 
 	EditorInfo.bCanUsePreviewPlatformForMaterialValidation = GetSectionBool(Section, "bCanUsePreviewPlatformForMaterialValidation", false);
 #endif
@@ -393,27 +395,26 @@ void FGenericDataDrivenShaderPlatformInfo::Initialize()
 #if WITH_EDITOR
 				if (!FParse::Param(FCommandLine::Get(), TEXT("NoPreviewPlatforms")))
 				{
-					const FName& CurrentPlatformName = Infos[ShaderPlatform].Name;
+					const TArray<FPreviewPlatformMenuItem>* Items = FDataDrivenPlatformInfoRegistry::GetMenuItemForShaderPlatformName(Infos[ShaderPlatform].Name);
+					if (!Items)
+					{
+						continue;
+					}
 
-					for (const FPreviewPlatformMenuItem& Item : FDataDrivenPlatformInfoRegistry::GetAllPreviewPlatformMenuItems())
+					const EShaderPlatform PreviewShaderPlatform = EShaderPlatform(CustomShaderPlatform++);
+					ParseDataDrivenShaderInfo(SectionSettings, PreviewShaderPlatform);
+					FGenericDataDrivenShaderPlatformInfo& PreviewInfo = Infos[PreviewShaderPlatform];
+					FDataDrivenShaderPlatformInfoEditorOnly& PreviewEditorInfo = DataDrivenShaderPlatformInfoEditorOnlyInfos[PreviewShaderPlatform];
+					PreviewEditorInfo.PreviewShaderPlatformParent = ShaderPlatform;
+
+					for (bool bFirst = true; const FPreviewPlatformMenuItem& Item : *Items)
 					{
-						if (Item.ShaderPlatformToPreview == CurrentPlatformName)
+						if (bFirst)
 						{
-							const EShaderPlatform PreviewShaderPlatform = EShaderPlatform(CustomShaderPlatform++);
-							FGenericDataDrivenShaderPlatformInfo& PreviewInfo = Infos[PreviewShaderPlatform];
 							PreviewInfo.Name = Item.PreviewShaderPlatformName;
 							PreviewInfo.PlatformName = Infos[ShaderPlatform].PlatformName;
-							ParseDataDrivenShaderInfo(SectionSettings, PreviewShaderPlatform);
 							PreviewInfo.bIsPreviewPlatform = true;
 							PreviewInfo.bContainsValidPlatformInfo = true;
-
-							FDataDrivenShaderPlatformInfoEditorOnly& PreviewEditorInfo = DataDrivenShaderPlatformInfoEditorOnlyInfos[PreviewShaderPlatform];
-							PreviewEditorInfo.PreviewShaderPlatformParent = ShaderPlatform;
-							if (!Item.OptionalFriendlyNameOverride.IsEmpty())
-							{
-								PreviewEditorInfo.FriendlyName = Item.OptionalFriendlyNameOverride;
-							}
-
 							ERHIFeatureLevel::Type PreviewFeatureLevel = ERHIFeatureLevel::Num;
 							if (GetFeatureLevelFromName(Item.PreviewFeatureLevelName, PreviewFeatureLevel))
 							{
@@ -421,6 +422,16 @@ void FGenericDataDrivenShaderPlatformInfo::Initialize()
 							}
 
 							PlatformNameToShaderPlatformMap.FindOrAdd(PreviewInfo.Name) = PreviewShaderPlatform;
+							bFirst = false;
+						}
+						
+						if (!Item.OptionalFriendlyNameOverride.IsEmpty())
+						{
+							PreviewEditorInfo.DeviceProfileToFriendlyName.FindOrAdd(Item.DeviceProfileName) = Item.OptionalFriendlyNameOverride;
+						}
+						else
+						{
+							PreviewEditorInfo.DeviceProfileToFriendlyName.FindOrAdd(Item.DeviceProfileName) = GetFriendlyName(PreviewShaderPlatform);
 						}
 					}
 				}
@@ -462,8 +473,25 @@ void FGenericDataDrivenShaderPlatformInfo::Initialize()
 
 				PlatformNameToShaderPlatformMap.FindOrAdd(Info.Name) = ShaderPlatform;
 
-				FDataDrivenShaderPlatformInfoEditorOnly& PreviewEditorInfo = DataDrivenShaderPlatformInfoEditorOnlyInfos[ShaderPlatform];
-				PreviewEditorInfo.PreviewShaderPlatformParent = ShaderPlatform;
+				const TArray<FPreviewPlatformMenuItem>* Items = FDataDrivenPlatformInfoRegistry::GetMenuItemForShaderPlatformName(Info.Name);
+				if (!Items)
+				{
+					continue;
+				}
+
+				for (const FPreviewPlatformMenuItem & Item : *Items)
+				{
+					FDataDrivenShaderPlatformInfoEditorOnly& PreviewEditorInfo = DataDrivenShaderPlatformInfoEditorOnlyInfos[ShaderPlatform];
+					PreviewEditorInfo.PreviewShaderPlatformParent = ShaderPlatform;
+					if (!Item.OptionalFriendlyNameOverride.IsEmpty())
+					{
+						PreviewEditorInfo.DeviceProfileToFriendlyName.FindOrAdd(Item.DeviceProfileName) = Item.OptionalFriendlyNameOverride;
+					}
+					else
+					{
+						PreviewEditorInfo.DeviceProfileToFriendlyName.FindOrAdd(Item.DeviceProfileName) = GetFriendlyName(ShaderPlatform);
+					}
+				}
 			}
 #endif
 		}
@@ -575,14 +603,21 @@ void FGenericDataDrivenShaderPlatformInfo::UpdatePreviewPlatforms()
 	}
 }
 
-FText FGenericDataDrivenShaderPlatformInfo::GetFriendlyName(const FStaticShaderPlatform Platform)
+FText FGenericDataDrivenShaderPlatformInfo::GetFriendlyName(const FStaticShaderPlatform Platform, FName DeviceProfileName)
 {
 	if (IsRunningCommandlet() || GUsingNullRHI)
 	{
 		return FText();
 	}
 	check(IsValid(Platform));
-	return DataDrivenShaderPlatformInfoEditorOnlyInfos[Platform].FriendlyName;
+	
+	TMap<FName, FText>& DeviceProfileToFriendlyName = DataDrivenShaderPlatformInfoEditorOnlyInfos[Platform].DeviceProfileToFriendlyName;
+	if (FText* FriendlyName = DeviceProfileToFriendlyName.Find(DeviceProfileName))
+	{
+		return *FriendlyName;
+	}
+
+	return FText::FromString(TEXT(""));
 }
 
 const EShaderPlatform FGenericDataDrivenShaderPlatformInfo::GetPreviewShaderPlatformParent(const FStaticShaderPlatform Platform)
diff --git a/Engine/Source/Runtime/RHI/Public/DataDrivenShaderPlatformInfo.h b/Engine/Source/Runtime/RHI/Public/DataDrivenShaderPlatformInfo.h
--- a/Engine/Source/Runtime/RHI/Public/DataDrivenShaderPlatformInfo.h
+++ b/Engine/Source/Runtime/RHI/Public/DataDrivenShaderPlatformInfo.h
@@ -876,7 +876,7 @@ PRAGMA_ENABLE_DEPRECATION_WARNINGS
 
 #if WITH_EDITOR
 	RHI_API static void UpdatePreviewPlatforms();
-	RHI_API static FText GetFriendlyName(const FStaticShaderPlatform Platform);
+	RHI_API static FText GetFriendlyName(const FStaticShaderPlatform Platform, FName DeviceProfileName = NAME_None);
 	RHI_API static const EShaderPlatform GetPreviewShaderPlatformParent(const FStaticShaderPlatform Platform);
 	RHI_API static const bool CanUseForMaterialValidation(const FStaticShaderPlatform Platform);
 
