Fix shadow variable build error. Enum should be an enum class

diff --git a/Engine/Source/Runtime/RHI/Public/RHIBreadcrumbs.h b/Engine/Source/Runtime/RHI/Public/RHIBreadcrumbs.h
--- a/Engine/Source/Runtime/RHI/Public/RHIBreadcrumbs.h
+++ b/Engine/Source/Runtime/RHI/Public/RHIBreadcrumbs.h
@@ -223,7 +223,7 @@ namespace UE::RHI::GPUProfiler
 			}
 		};
 
-		enum EVerbosity
+		enum class EVerbosity
 		{
 			Log,
 			Warning,
