- Update the Pipeline File Cache to serialize/hash bAllowVariableRateShading and ShadingRate
- Fix broken deserialization logic for pre-bHasFragmentDensityAttachment pipeline file cache archives
- De-duplicate vertex sorting code in FVulkanVertexDeclarationKey and FPipelineCacheFileFormatPSO

diff --git a/Engine/Source/Runtime/RHI/Private/PipelineFileCache.cpp b/Engine/Source/Runtime/RHI/Private/PipelineFileCache.cpp
--- a/Engine/Source/Runtime/RHI/Private/PipelineFileCache.cpp
+++ b/Engine/Source/Runtime/RHI/Private/PipelineFileCache.cpp
@@ -85,6 +85,7 @@ enum class EPipelineCacheFileFormatVersions : uint32
 	RemovingLineAA = 27,
 	AddingDepthBounds = 28,
 	AddRTPSOShaderBindingLayout = 29,
+	VariableRateShading = 30,
 };
 
 const uint64 FPipelineCacheFileFormatMagic = 0x5049504543414348; // PIPECACH
@@ -594,6 +595,11 @@ FString FPipelineCacheFileFormatPSO::GraphicsDescriptor::StateToString() const
 		, uint32(bHasFragmentDensityAttachment)
 	);
 
+	Result += FString::Printf(TEXT("%d,%d,")
+		, uint32(bAllowVariableRateShading)
+		, uint32(ShadingRate)
+	);
+
 	Result += FString::Printf(TEXT("%d,")
 		, uint32(bDepthBounds)
 	);
@@ -682,6 +688,12 @@ void FPipelineCacheFileFormatPSO::GraphicsDescriptor::AddStateToReadableString(T
 	OutBuilder << bHasFragmentDensityAttachment;
 	OutBuilder << TEXT("\n");
 
+	OutBuilder << TEXT(" AllowVRS:");
+	OutBuilder << bAllowVariableRateShading;
+	OutBuilder << TEXT(" SR:");
+	OutBuilder << uint32(ShadingRate);
+	OutBuilder << TEXT("\n");
+
 	OutBuilder << TEXT(" DB:");
 	OutBuilder << bDepthBounds;
 	OutBuilder << TEXT("\n");
@@ -782,6 +794,17 @@ bool FPipelineCacheFileFormatPSO::GraphicsDescriptor::StateFromString(const FStr
 		bHasFragmentDensityAttachment = (bool)LocalHasFDM;
 	}
 
+	// parse variable rate shading information
+	{
+		uint32 LocalAllowVariableRateShading = 0;
+		uint32 LocalShadingRate = 0;
+		check(PartEnd - PartIt >= 2);
+		LexFromString(LocalAllowVariableRateShading, *PartIt++);
+		LexFromString(LocalShadingRate, *PartIt++);
+		bAllowVariableRateShading = (bool)LocalAllowVariableRateShading;
+		ShadingRate = (EVRSShadingRate)LocalShadingRate;
+	}
+
 	// parse depth bounds
 	{
 		uint32 DepthBounds = 0;
@@ -806,34 +829,7 @@ bool FPipelineCacheFileFormatPSO::GraphicsDescriptor::StateFromString(const FStr
 
 	check(PartIt + MaxVertexElementCount == PartEnd + VertDescNum);
 
-	VertexDescriptor.Sort([](FVertexElement const& A, FVertexElement const& B)
-	  {
-		  if (A.StreamIndex < B.StreamIndex)
-		  {
-			  return true;
-		  }
-		  if (A.StreamIndex > B.StreamIndex)
-		  {
-			  return false;
-		  }
-		  if (A.Offset < B.Offset)
-		  {
-			  return true;
-		  }
-		  if (A.Offset > B.Offset)
-		  {
-			  return false;
-		  }
-		  if (A.AttributeIndex < B.AttributeIndex)
-		  {
-			  return true;
-		  }
-		  if (A.AttributeIndex > B.AttributeIndex)
-		  {
-			  return false;
-		  }
-		  return false;
-	  });
+	SortVertexElements(VertexDescriptor);
 
 	return true;
 }
@@ -1185,6 +1181,9 @@ bool FPipelineCacheFileFormatPSO::Verify() const
 			KeyHash = FCrc::MemCrc32(&Key.GraphicsDesc.MultiViewCount, sizeof(Key.GraphicsDesc.MultiViewCount), KeyHash);
 			KeyHash = FCrc::MemCrc32(&Key.GraphicsDesc.bHasFragmentDensityAttachment, sizeof(Key.GraphicsDesc.bHasFragmentDensityAttachment), KeyHash);
 
+			KeyHash = FCrc::MemCrc32(&Key.GraphicsDesc.bAllowVariableRateShading, sizeof(Key.GraphicsDesc.bAllowVariableRateShading), KeyHash);
+			KeyHash = FCrc::MemCrc32(&Key.GraphicsDesc.ShadingRate, sizeof(Key.GraphicsDesc.ShadingRate), KeyHash);
+
 			KeyHash = FCrc::MemCrc32(&Key.GraphicsDesc.bDepthBounds, sizeof(Key.GraphicsDesc.bDepthBounds), KeyHash);
 
 			for(auto const& Element : Key.GraphicsDesc.VertexDescriptor)
@@ -1282,34 +1281,7 @@ bool FPipelineCacheFileFormatPSO::Verify() const
 				
 				FVertexDeclarationElementList Elements;
 				Ar << Elements;
-				Elements.Sort([](FVertexElement const& A, FVertexElement const& B)
-							  {
-								  if (A.StreamIndex < B.StreamIndex)
-								  {
-									  return true;
-								  }
-								  if (A.StreamIndex > B.StreamIndex)
-								  {
-									  return false;
-								  }
-								  if (A.Offset < B.Offset)
-								  {
-									  return true;
-								  }
-								  if (A.Offset > B.Offset)
-								  {
-									  return false;
-								  }
-								  if (A.AttributeIndex < B.AttributeIndex)
-								  {
-									  return true;
-								  }
-								  if (A.AttributeIndex > B.AttributeIndex)
-								  {
-									  return false;
-								  }
-								  return false;
-							  });
+				FPipelineCacheFileFormatPSO::SortVertexElements(Elements);
 				
 				Info.GraphicsDesc.VertexDescriptor.AddZeroed(Elements.Num());
 				for (uint32 i = 0; i < (uint32)Elements.Num(); i++)
@@ -1382,20 +1354,20 @@ bool FPipelineCacheFileFormatPSO::Verify() const
 			Ar << Info.GraphicsDesc.SubpassHint;
 			Ar << Info.GraphicsDesc.SubpassIndex;
 
-			if (Ar.GameNetVer() < (uint32)EPipelineCacheFileFormatVersions::FragmentDensityAttachment)
-			{
-				uint8 MultiViewCount = 0;
-				Ar << MultiViewCount;
-				
-				bool bHasFragmentDensityAttachment = false;
-				Ar << bHasFragmentDensityAttachment;
-			}
-			else
+			if (Ar.GameNetVer() >= (uint32)EPipelineCacheFileFormatVersions::FragmentDensityAttachment)
 			{
 				Ar << Info.GraphicsDesc.MultiViewCount;
 				Ar << Info.GraphicsDesc.bHasFragmentDensityAttachment;
 			}
 
+			if (Ar.GameNetVer() >= (uint32)EPipelineCacheFileFormatVersions::VariableRateShading)
+			{
+				Ar << Info.GraphicsDesc.bAllowVariableRateShading;
+				uint32 ShadingRate = (uint32)Info.GraphicsDesc.ShadingRate;
+				Ar << ShadingRate;
+				Info.GraphicsDesc.ShadingRate = (EVRSShadingRate)ShadingRate;
+			}
+
 			if (Ar.GameNetVer() >= (uint32)EPipelineCacheFileFormatVersions::AddingDepthBounds)
 			{
 				Ar << Info.GraphicsDesc.bDepthBounds;
@@ -1488,34 +1460,7 @@ FPipelineCacheFileFormatPSO::FPipelineCacheFileFormatPSO()
 			bOK &= Init.BoundShaderState.VertexDeclarationRHI->GetInitializer(PSO.GraphicsDesc.VertexDescriptor);
 			check(bOK);
 		
-			PSO.GraphicsDesc.VertexDescriptor.Sort([](FVertexElement const& A, FVertexElement const& B)
-			{
-				if (A.StreamIndex < B.StreamIndex)
-				{
-					return true;
-				}
-				if (A.StreamIndex > B.StreamIndex)
-				{
-					return false;
-				}
-				if (A.Offset < B.Offset)
-				{
-					return true;
-				}
-				if (A.Offset > B.Offset)
-				{
-					return false;
-				}
-				if (A.AttributeIndex < B.AttributeIndex)
-				{
-					return true;
-				}
-				if (A.AttributeIndex > B.AttributeIndex)
-				{
-					return false;
-				}
-				return false;
-			});
+			SortVertexElements(PSO.GraphicsDesc.VertexDescriptor);
 		}
 
 		PSO.GraphicsDesc.VertexShader = Init.BoundShaderState.VertexShaderRHI->GetHash();
@@ -1585,6 +1530,9 @@ FPipelineCacheFileFormatPSO::FPipelineCacheFileFormatPSO()
 	PSO.GraphicsDesc.MultiViewCount = (uint8)Init.MultiViewCount;
 	PSO.GraphicsDesc.bHasFragmentDensityAttachment = Init.bHasFragmentDensityAttachment;
 
+	PSO.GraphicsDesc.bAllowVariableRateShading = Init.bAllowVariableRateShading;
+	PSO.GraphicsDesc.ShadingRate = Init.ShadingRate;
+
 	PSO.GraphicsDesc.bDepthBounds = Init.bDepthBounds;
 
 #if !UE_BUILD_SHIPPING
@@ -1594,6 +1542,38 @@ FPipelineCacheFileFormatPSO::FPipelineCacheFileFormatPSO()
 	return bOK;
 }
 
+void FPipelineCacheFileFormatPSO::SortVertexElements(FVertexDeclarationElementList& VertexElements)
+{
+	VertexElements.Sort([](FVertexElement const& A, FVertexElement const& B)
+	{
+		if (A.StreamIndex < B.StreamIndex)
+		{
+			return true;
+		}
+		if (A.StreamIndex > B.StreamIndex)
+		{
+			return false;
+		}
+		if (A.Offset < B.Offset)
+		{
+			return true;
+		}
+		if (A.Offset > B.Offset)
+		{
+			return false;
+		}
+		if (A.AttributeIndex < B.AttributeIndex)
+		{
+			return true;
+		}
+		if (A.AttributeIndex > B.AttributeIndex)
+		{
+			return false;
+		}
+		return false;
+	});
+}
+
 FPipelineCacheFileFormatPSO::~FPipelineCacheFileFormatPSO()
 {
 	
@@ -1642,6 +1622,7 @@ bool FPipelineCacheFileFormatPSO::operator==(const FPipelineCacheFileFormatPSO&
 						GraphicsDesc.DepthStore == Other.GraphicsDesc.DepthStore && GraphicsDesc.StencilLoad == Other.GraphicsDesc.StencilLoad && GraphicsDesc.StencilStore == Other.GraphicsDesc.StencilStore &&
 						GraphicsDesc.SubpassHint == Other.GraphicsDesc.SubpassHint && GraphicsDesc.SubpassIndex == Other.GraphicsDesc.SubpassIndex &&
 						GraphicsDesc.MultiViewCount == Other.GraphicsDesc.MultiViewCount && GraphicsDesc.bHasFragmentDensityAttachment == Other.GraphicsDesc.bHasFragmentDensityAttachment &&
+						GraphicsDesc.bAllowVariableRateShading == Other.GraphicsDesc.bAllowVariableRateShading && GraphicsDesc.ShadingRate == Other.GraphicsDesc.ShadingRate &&
 						GraphicsDesc.bDepthBounds == Other.GraphicsDesc.bDepthBounds &&
 					FMemory::Memcmp(&GraphicsDesc.BlendState, &Other.GraphicsDesc.BlendState, sizeof(FBlendStateInitializerRHI)) == 0 &&
 					FMemory::Memcmp(&GraphicsDesc.RasterizerState, &Other.GraphicsDesc.RasterizerState, sizeof(FPipelineFileCacheRasterizerState)) == 0 &&
diff --git a/Engine/Source/Runtime/RHI/Public/PipelineFileCache.h b/Engine/Source/Runtime/RHI/Public/PipelineFileCache.h
--- a/Engine/Source/Runtime/RHI/Public/PipelineFileCache.h
+++ b/Engine/Source/Runtime/RHI/Public/PipelineFileCache.h
@@ -148,6 +148,9 @@ struct FPipelineCacheFileFormatPSO
 		uint8	MultiViewCount;
 		bool	bHasFragmentDensityAttachment;
 
+		bool	bAllowVariableRateShading;
+		EVRSShadingRate ShadingRate;
+
 		bool	bDepthBounds;
 		
 		RHI_API FString ToString() const;
@@ -231,6 +234,8 @@ struct FPipelineCacheFileFormatPSO
 	static bool Init(FPipelineCacheFileFormatPSO& PSO, FGraphicsPipelineStateInitializer const& Init);
 	static bool Init(FPipelineCacheFileFormatPSO & PSO, FPipelineFileCacheRayTracingDesc const& Desc);
 
+	RHI_API static void SortVertexElements(FVertexDeclarationElementList& VertexElements);
+
 	RHI_API FString CommonToString() const;
 	RHI_API static FString CommonHeaderLine();
 	RHI_API void CommonFromString(const FStringView& Src);
diff --git a/Engine/Source/Runtime/VulkanRHI/Private/VulkanRHI.cpp b/Engine/Source/Runtime/VulkanRHI/Private/VulkanRHI.cpp
--- a/Engine/Source/Runtime/VulkanRHI/Private/VulkanRHI.cpp
+++ b/Engine/Source/Runtime/VulkanRHI/Private/VulkanRHI.cpp
@@ -1884,7 +1884,7 @@ IRHITransientResourceAllocator* FVulkanDynamicRHI::RHICreateTransientResourceAll
 
 uint32 FVulkanDynamicRHI::GetPrecachePSOHashVersion()
 {
-	static const uint32 PrecacheHashVersion = 4;
+	static const uint32 PrecacheHashVersion = 5;
 	return PrecacheHashVersion;
 }
 
@@ -1956,11 +1956,7 @@ uint64 FVulkanDynamicRHI::RHIComputeStatePrecachePSOHash(const FGraphicsPipeline
 	HashKey.DrawShadingRate = Initializer.ShadingRate;
 	HashKey.PrimitiveType = Initializer.PrimitiveType;
 	HashKey.bDepthBounds = Initializer.bDepthBounds;
-
-	// For hotfix 5.6.1, bAllowVariableRateShading is ignored for precaching hashes because it is not reflected in Pipeline File Cache entries
-	// This previously led cache reads for PSOs with bAllowVariableRateShading=false to be directed to the wrong chunk in the Vulkan Chunked Pipeline Cache and miss
-	// This fix will still result in a miss for drivers where toggling bAllowVariableRateShading requires a PSO recompile
-	HashKey.bAllowVariableRateShading = true;
+	HashKey.bAllowVariableRateShading = Initializer.bAllowVariableRateShading;
 
 	uint64 PrecachePSOHash = CityHash64((const char*)&HashKey, sizeof(FHashKey));
 
diff --git a/Engine/Source/Runtime/VulkanRHI/Private/VulkanVertexDeclaration.cpp b/Engine/Source/Runtime/VulkanRHI/Private/VulkanVertexDeclaration.cpp
--- a/Engine/Source/Runtime/VulkanRHI/Private/VulkanVertexDeclaration.cpp
+++ b/Engine/Source/Runtime/VulkanRHI/Private/VulkanVertexDeclaration.cpp
@@ -4,6 +4,7 @@
 	VulkanVertexDeclaration.cpp: Vulkan vertex declaration RHI implementation.
 =============================================================================*/
 
+#include "PipelineFileCache.h"
 #include "VulkanRHIPrivate.h"
 
 struct FVulkanVertexDeclarationKey
@@ -19,34 +20,7 @@ struct FVulkanVertexDeclarationKey
 		Hash = FCrc::MemCrc_DEPRECATED(VertexElements.GetData(), VertexElements.Num()*sizeof(FVertexElement));
 
 		FVertexDeclarationElementList SortedElements = VertexElements;
-		SortedElements.Sort([](FVertexElement const& A, FVertexElement const& B)
-		{
-			if (A.StreamIndex < B.StreamIndex)
-			{
-				return true;
-			}
-			if (A.StreamIndex > B.StreamIndex)
-			{
-				return false;
-			}
-			if (A.Offset < B.Offset)
-			{
-				return true;
-			}
-			if (A.Offset > B.Offset)
-			{
-				return false;
-			}
-			if (A.AttributeIndex < B.AttributeIndex)
-			{
-				return true;
-			}
-			if (A.AttributeIndex > B.AttributeIndex)
-			{
-				return false;
-			}
-			return false;
-		});
+		FPipelineCacheFileFormatPSO::SortVertexElements(SortedElements);
 
 		HashNoStride = 0;
 		for (int32 ElementIndex = 0; ElementIndex < SortedElements.Num(); ElementIndex++)
