Make RDG emit all scope and pass names when profiling the GPU via the "profilegpu" command.
- This used to happen with the legacy profiler, but was missing in the new one.

diff --git a/Engine/Source/Runtime/RHI/Private/GPUProfiler.cpp b/Engine/Source/Runtime/RHI/Private/GPUProfiler.cpp
--- a/Engine/Source/Runtime/RHI/Private/GPUProfiler.cpp
+++ b/Engine/Source/Runtime/RHI/Private/GPUProfiler.cpp
@@ -2251,7 +2251,7 @@ namespace UE::RHI::GPUProfiler
 				Timestamps.TotalBusyCycles = TotalGPUTime;
 			}
 
-			void FrameBoundary(FEvent::FFrameBoundary const& Event, FFrameState& FrameState, uint32 FrameNumber)
+			void FrameBoundary(FEvent::FFrameBoundary const& Event, FFrameState& FrameState, uint32 InProfileFrameNumber)
 			{
 				check(!bBusy);
 				Timestamps.CPUFrameBoundary = Event.CPUTimestamp;
@@ -2278,7 +2278,7 @@ namespace UE::RHI::GPUProfiler
 			#endif
 
 			#if WITH_PROFILEGPU
-				if (FrameNumber == Event.FrameNumber + 1)
+				if (InProfileFrameNumber == Event.FrameNumber + 1)
 				{
 					Profile.bProfileFrame = true;
 
@@ -2306,6 +2306,8 @@ namespace UE::RHI::GPUProfiler
 		};
 
 		std::atomic<bool> bTriggerProfile{ false };
+		std::atomic<bool> bIsProfiling{ false };
+
 		uint32 ProfileFrameNumber = 0;
 		uint32 MaxFrameNumber = 0;
 
@@ -2498,6 +2500,13 @@ namespace UE::RHI::GPUProfiler
 							ProcessFrame(FrameState);
 
 							Frames.Remove(FrameBoundary.FrameNumber);
+
+						#if WITH_PROFILEGPU
+							if (ProfileFrameNumber == FrameBoundary.FrameNumber)
+							{
+								bIsProfiling = false;
+							}
+						#endif
 						}
 					}
 					break;
@@ -2645,9 +2654,19 @@ namespace UE::RHI::GPUProfiler
 		TEXT("Captures statistics about a frame of GPU work and prints the results to the log."),
 		FConsoleCommandWithArgsDelegate::CreateLambda([](const TArray<FString>& Args)
 		{
+			GGPUProfilerSink_StatSystem.bIsProfiling = true;
 			GGPUProfilerSink_StatSystem.bTriggerProfile = true;
 		}));
 #endif // WITH_PROFILEGPU
+
+	RHI_API bool IsProfiling()
+	{
+	#if WITH_PROFILEGPU
+		return GGPUProfilerSink_StatSystem.bIsProfiling;
+	#else
+		return false;
+	#endif
+	}
 }
 
 #endif // RHI_NEW_GPU_PROFILER
diff --git a/Engine/Source/Runtime/RHI/Public/GPUProfiler.h b/Engine/Source/Runtime/RHI/Public/GPUProfiler.h
--- a/Engine/Source/Runtime/RHI/Public/GPUProfiler.h
+++ b/Engine/Source/Runtime/RHI/Public/GPUProfiler.h
@@ -589,6 +589,8 @@ namespace UE::RHI::GPUProfiler
 		FRHIDrawStatsCategory DrawcallCategory;
 	#endif
 	};
+
+	RHI_API bool IsProfiling();
 }
 
 #else
diff --git a/Engine/Source/Runtime/RenderCore/Private/RenderGraphPrivate.cpp b/Engine/Source/Runtime/RenderCore/Private/RenderGraphPrivate.cpp
--- a/Engine/Source/Runtime/RenderCore/Private/RenderGraphPrivate.cpp
+++ b/Engine/Source/Runtime/RenderCore/Private/RenderGraphPrivate.cpp
@@ -659,7 +659,12 @@ FRDGScopeState::FState::FState(bool bInImmediate, bool bInParallelExecute)
 		bool bRDGChannelEnabled = false;
 		IF_RDG_ENABLE_TRACE(bRDGChannelEnabled = UE_TRACE_CHANNELEXPR_IS_ENABLED(RDGChannel));
 
-		if (FRDGBuilder::IsDumpingFrame() || GTriggerGPUProfile)
+		if (FRDGBuilder::IsDumpingFrame()
+			|| GTriggerGPUProfile
+		#if RHI_NEW_GPU_PROFILER
+			|| UE::RHI::GPUProfiler::IsProfiling()
+		#endif
+			)
 		{
 			// We want all possible scope and pass names in a DumpGPU/profilegpu trace.
 			return ERDGScopeMode::AllEventsAndPassNames;
