Prefer DXGI GPU adapter name for ShouldDefaultToPerformanceMode

We don't want "Microsoft Remote Display Adapter" over Remote Desktop from EnumDisplayDevices.

diff --git a/Engine/Source/Runtime/RHI/Private/Windows/WindowsDynamicRHI.cpp b/Engine/Source/Runtime/RHI/Private/Windows/WindowsDynamicRHI.cpp
--- a/Engine/Source/Runtime/RHI/Private/Windows/WindowsDynamicRHI.cpp
+++ b/Engine/Source/Runtime/RHI/Private/Windows/WindowsDynamicRHI.cpp
@@ -11,6 +11,12 @@
 
 #include "Windows/WindowsPlatformApplicationMisc.h"
 
+THIRD_PARTY_INCLUDES_START
+#include "Windows/AllowWindowsPlatformTypes.h"
+#include <dxgi.h>
+#include "Windows/HideWindowsPlatformTypes.h"
+THIRD_PARTY_INCLUDES_END
+
 #define LOCTEXT_NAMESPACE "WindowsDynamicRHI"
 
 bool GDynamicRHIFailedToInitializeAdvancedPlatform = false;
@@ -446,6 +452,47 @@ FParsedWindowsDynamicRHIConfig ParseWindowsDynamicRHIConfig()
 	return Config;
 }
 
+struct FWindowsGPUInfo
+{
+	uint32 VendorId = 0;
+	uint32 DeviceId = 0;
+	uint64 DedicatedVideoMemory = 0;
+	FString AdapterName;
+};
+
+static void GetGPUInfo(FWindowsGPUInfo& Info)
+{
+	DXGI_ADAPTER_DESC BestDesc = {};
+	TRefCountPtr<IDXGIFactory1> DXGIFactory1;
+	if (SUCCEEDED(CreateDXGIFactory1(IID_PPV_ARGS(DXGIFactory1.GetInitReference()))))
+	{
+		TRefCountPtr<IDXGIAdapter> TempAdapter;
+		for (uint32 AdapterIndex = 0; DXGIFactory1->EnumAdapters(AdapterIndex, TempAdapter.GetInitReference()) != DXGI_ERROR_NOT_FOUND; ++AdapterIndex)
+		{
+			if (TempAdapter)
+			{
+				DXGI_ADAPTER_DESC Desc;
+				TempAdapter->GetDesc(&Desc);
+
+				if (Desc.DedicatedVideoMemory > BestDesc.DedicatedVideoMemory || AdapterIndex == 0)
+				{
+					BestDesc = Desc;
+				}
+			}
+		}
+	}
+
+	Info.VendorId = BestDesc.VendorId;
+	Info.DeviceId = BestDesc.DeviceId;
+	Info.DedicatedVideoMemory = BestDesc.DedicatedVideoMemory;
+	Info.AdapterName = BestDesc.Description;
+
+	if (Info.AdapterName.IsEmpty())
+	{
+		Info.AdapterName = FPlatformMisc::GetPrimaryGPUBrand();
+	}
+}
+
 // Whether to default to Performance Mode on low-end machines (ES31 feature level).
 // Looks at a few configuration options that control thresholds to enable this as the default.
 static bool ShouldDefaultToPerformanceMode(bool& bOutForceD3D11)
@@ -458,12 +505,13 @@ static bool ShouldDefaultToPerformanceMode(bool& bOutForceD3D11)
 		return ForceES31.GetValue();
 	}
 
-	FWindowsPlatformApplicationMisc::FGPUInfo BestGPUInfo = FWindowsPlatformApplicationMisc::GetBestGPUInfo();
+	FWindowsGPUInfo BestGPUInfo;
+	GetGPUInfo(BestGPUInfo);
 
 	TArray<FString> DeviceDefaultRHIList;
 	GConfig->GetArray(TEXT("Devices"), TEXT("DeviceDefaultRHIList"), DeviceDefaultRHIList, GHardwareIni);
 
-	FString GPUBrand = FPlatformMisc::GetPrimaryGPUBrand();
+	const FString& GPUBrand = BestGPUInfo.AdapterName;
 	for (const FString& DeviceDefaultRHIString : DeviceDefaultRHIList)
 	{
 		const TCHAR* Line = *DeviceDefaultRHIString;
