Do not allow editor to select a feature level below SM5 since that will result in failure anyway

diff --git a/Engine/Source/Runtime/RHI/Private/Windows/WindowsDynamicRHI.cpp b/Engine/Source/Runtime/RHI/Private/Windows/WindowsDynamicRHI.cpp
--- a/Engine/Source/Runtime/RHI/Private/Windows/WindowsDynamicRHI.cpp
+++ b/Engine/Source/Runtime/RHI/Private/Windows/WindowsDynamicRHI.cpp
@@ -397,6 +397,12 @@ static TArray<EShaderPlatform> ParseShaderPlatformsConfig(const TCHAR* InSetting
 		ShaderPlatforms.AddUnique(ShaderFormatToLegacyShaderPlatform(FName(*ShaderFormat)));
 	}
 
+	ShaderPlatforms.RemoveAll(
+		[](EShaderPlatform ShaderPlatform)
+		{
+			return GIsEditor && (FDataDrivenShaderPlatformInfo::GetMaxFeatureLevel(ShaderPlatform) < ERHIFeatureLevel::SM5);
+		});
+
 	return ShaderPlatforms;
 }
 
