Add new config array DeviceForceRHIList to lock limited GPUs to a particular RHI and feature level.

diff --git a/Engine/Source/Runtime/RHI/Private/Windows/WindowsDynamicRHI.cpp b/Engine/Source/Runtime/RHI/Private/Windows/WindowsDynamicRHI.cpp
--- a/Engine/Source/Runtime/RHI/Private/Windows/WindowsDynamicRHI.cpp
+++ b/Engine/Source/Runtime/RHI/Private/Windows/WindowsDynamicRHI.cpp
@@ -458,6 +458,57 @@ FParsedWindowsDynamicRHIConfig ParseWindowsDynamicRHIConfig()
 	return Config;
 }
 
+// TODO: Deduplicate code between this and ShouldDefaultToPerformanceMode
+static bool IsForcedToD3D11_ES31()
+{
+	bool bRequireD3D11ES31 = false;
+
+	FWindowsPlatformApplicationMisc::FGPUInfo BestGPUInfo = FWindowsPlatformApplicationMisc::GetBestGPUInfo();
+
+	TArray<FString> DeviceDefaultRHIList;
+	GConfig->GetArray(TEXT("Devices"), TEXT("DeviceForceRHIList"), DeviceDefaultRHIList, GHardwareIni);
+
+	FString GPUBrand = FPlatformMisc::GetPrimaryGPUBrand();
+	for (const FString& DeviceDefaultRHIString : DeviceDefaultRHIList)
+	{
+		const TCHAR* Line = *DeviceDefaultRHIString;
+
+		ensure(Line[0] == TCHAR('('));
+
+		FString RHIName;
+		FParse::Value(Line + 1, TEXT("RHI="), RHIName);
+
+		FString DeviceName;
+		FParse::Value(Line + 1, TEXT("DeviceName="), DeviceName);
+
+		if (RHIName.Compare("D3D11_ES31", ESearchCase::IgnoreCase) == 0 && GPUBrand.Compare(DeviceName, ESearchCase::IgnoreCase) == 0)
+		{
+			UE_LOG(LogRHI, Log, TEXT("Found RHIName %s with DeviceName %s: forced feature level to D3D11 %s"), *RHIName, *DeviceName, GetLogName(ERHIFeatureLevel::ES3_1));
+			bRequireD3D11ES31 = true;
+			break;
+		}
+
+		FString VendorId;
+		FParse::Value(Line + 1, TEXT("VendorId="), VendorId);
+		uint32 VendorIdInt = FParse::HexNumber(*VendorId);
+
+		FString DeviceId;
+		FParse::Value(Line + 1, TEXT("DeviceId="), DeviceId);
+		uint32 DeviceIdInt = FParse::HexNumber(*DeviceId);
+
+		if (BestGPUInfo.VendorId && BestGPUInfo.DeviceId &&
+			BestGPUInfo.VendorId == VendorIdInt && BestGPUInfo.DeviceId == DeviceIdInt &&
+			RHIName.Compare("D3D11_ES31", ESearchCase::IgnoreCase) == 0)
+		{
+			UE_LOG(LogRHI, Log, TEXT("Found RHIName %s with DeviceName %s, VendorId %s, and DeviceId %s: forced feature level to D3D11 %s"), *RHIName, *DeviceName, *VendorId, *DeviceId, GetLogName(ERHIFeatureLevel::ES3_1));
+			bRequireD3D11ES31 = true;
+			break;
+		}
+	}
+
+	return bRequireD3D11ES31;
+}
+
 struct FWindowsGPUInfo
 {
 	uint32 VendorId = 0;
@@ -627,55 +678,66 @@ static TOptional<ERHIFeatureLevel::Type> GetPreferredFeatureLevel(EWindowsRHI Ch
 	TOptional<ERHIFeatureLevel::Type> PreferredFeatureLevel;
 	if (!GIsEditor)
 	{
-		FString PreferredFeatureLevelName;
-		if (GConfig->GetString(TEXT("D3DRHIPreference"), TEXT("PreferredFeatureLevel"), PreferredFeatureLevelName, GGameUserSettingsIni))
+		bool bSetD3D11_ES31 = IsForcedToD3D11_ES31();
+		if (bSetD3D11_ES31)
 		{
-			UE_LOG(LogRHI, Log, TEXT("Found D3DRHIPreference PreferredFeatureLevel: %s"), *PreferredFeatureLevelName);
-
-			if (PreferredFeatureLevelName == TEXT("sm5"))
-			{
-				PreferredFeatureLevel = ERHIFeatureLevel::SM5;
-			}
-			else if (PreferredFeatureLevelName == TEXT("sm6"))
-			{
-				PreferredFeatureLevel = ERHIFeatureLevel::SM6;
-			}
-			else if (PreferredFeatureLevelName == TEXT("es31"))
-			{
-				PreferredFeatureLevel = ERHIFeatureLevel::ES3_1;
-			}
-			else
-			{
-				UE_LOG(LogRHI, Error, TEXT("unknown feature level name \"%s\" in game user settings, using default"), *PreferredFeatureLevelName);
-			}
+			bOutForceD3D11 = true;
+			PreferredFeatureLevel = ERHIFeatureLevel::ES3_1;
 		}
-		else if (Config.IsFeatureLevelTargeted(ChosenRHI, ERHIFeatureLevel::ES3_1))
+		else
 		{
-			bool bPreferFeatureLevelES31 = false;
-			bool bFoundPreference = GConfig->GetBool(TEXT("D3DRHIPreference"), TEXT("bPreferFeatureLevelES31"), bPreferFeatureLevelES31, GGameUserSettingsIni);
-			if (bFoundPreference)
+			FString PreferredFeatureLevelName;
+			if (GConfig->GetString(TEXT("D3DRHIPreference"), TEXT("PreferredFeatureLevel"), PreferredFeatureLevelName, GGameUserSettingsIni))
 			{
-				UE_LOG(LogRHI, Log, TEXT("Found D3DRHIPreference bPreferFeatureLevelES31: %s"), bPreferFeatureLevelES31 ? TEXT("true") : TEXT("false"));
-			}
+				UE_LOG(LogRHI, Log, TEXT("Found D3DRHIPreference PreferredFeatureLevel: %s"), *PreferredFeatureLevelName);
 
-			// Force low-spec users into performance mode but respect their choice once they have set a preference
-			bool bDefaultES31 = false;
-			if (!bFoundPreference && !CVarIgnorePerformanceModeCheck.GetValueOnAnyThread())
-			{
-				bDefaultES31 = ShouldDefaultToPerformanceMode(bOutForceD3D11);
+				if (PreferredFeatureLevelName == TEXT("sm5"))
+				{
+					PreferredFeatureLevel = ERHIFeatureLevel::SM5;
+				}
+				else if (PreferredFeatureLevelName == TEXT("sm6"))
+				{
+					PreferredFeatureLevel = ERHIFeatureLevel::SM6;
+				}
+				else if (PreferredFeatureLevelName == TEXT("es31"))
+				{
+					PreferredFeatureLevel = ERHIFeatureLevel::ES3_1;
+				}
+				else
+				{
+					UE_LOG(LogRHI, Error, TEXT("unknown feature level name \"%s\" in game user settings, using default"), *PreferredFeatureLevelName);
+				}
 			}
-
-			if (bPreferFeatureLevelES31 || bDefaultES31)
+			else if (Config.IsFeatureLevelTargeted(ChosenRHI, ERHIFeatureLevel::ES3_1))
 			{
-				if (!bFoundPreference)
+				bool bPreferFeatureLevelES31 = false;
+				bool bFoundPreference = GConfig->GetBool(TEXT("D3DRHIPreference"), TEXT("bPreferFeatureLevelES31"), bPreferFeatureLevelES31, GGameUserSettingsIni);
+				if (bFoundPreference)
+				{
+					UE_LOG(LogRHI, Log, TEXT("Found D3DRHIPreference bPreferFeatureLevelES31: %s"), bPreferFeatureLevelES31 ? TEXT("true") : TEXT("false"));
+				}
+
+				// Force low-spec users into performance mode but respect their choice once they have set a preference
+				bool bDefaultES31 = false;
+				if (!bFoundPreference && !CVarIgnorePerformanceModeCheck.GetValueOnAnyThread())
 				{
-					GConfig->SetBool(TEXT("D3DRHIPreference"), TEXT("bPreferFeatureLevelES31"), true, GGameUserSettingsIni);
-					if (bOutForceD3D11)
-					{
-						GConfig->SetString(TEXT("D3DRHIPreference"), TEXT("PreferredRHI"), TEXT("dx11"), GGameUserSettingsIni);
-					}
+					bDefaultES31 = ShouldDefaultToPerformanceMode(bOutForceD3D11);
 				}
-				PreferredFeatureLevel = ERHIFeatureLevel::ES3_1;
+
+				if (bPreferFeatureLevelES31 || bDefaultES31)
+				{
+					bSetD3D11_ES31 = !bFoundPreference;
+					PreferredFeatureLevel = ERHIFeatureLevel::ES3_1;
+				}
+			}
+		}
+
+		if (bSetD3D11_ES31)
+		{
+			GConfig->SetBool(TEXT("D3DRHIPreference"), TEXT("bPreferFeatureLevelES31"), true, GGameUserSettingsIni);
+			if (bOutForceD3D11)
+			{
+				GConfig->SetString(TEXT("D3DRHIPreference"), TEXT("PreferredRHI"), TEXT("dx11"), GGameUserSettingsIni);
 			}
 		}
 	}
