Prefer DXGI GPU adapter name for IsForcedToD3D11_ES31 and ShouldDefaultToPerformanceMode

We don't want "Microsoft Remote Display Adapter" over Remote Desktop from EnumDisplayDevices.

diff --git a/Engine/Source/Runtime/RHI/Private/Windows/WindowsDynamicRHI.cpp b/Engine/Source/Runtime/RHI/Private/Windows/WindowsDynamicRHI.cpp
--- a/Engine/Source/Runtime/RHI/Private/Windows/WindowsDynamicRHI.cpp
+++ b/Engine/Source/Runtime/RHI/Private/Windows/WindowsDynamicRHI.cpp
@@ -458,20 +458,62 @@ FParsedWindowsDynamicRHIConfig ParseWindowsDynamicRHIConfig()
 	return Config;
 }
 
+struct FWindowsGPUInfo
+{
+	uint32 VendorId = 0;
+	uint32 DeviceId = 0;
+	uint64 DedicatedVideoMemory = 0;
+	FString AdapterName;
+};
+
+static void GetGPUInfo(FWindowsGPUInfo& Info)
+{
+	DXGI_ADAPTER_DESC BestDesc = {};
+	TRefCountPtr<IDXGIFactory1> DXGIFactory1;
+	if (SUCCEEDED(CreateDXGIFactory1(IID_PPV_ARGS(DXGIFactory1.GetInitReference()))))
+	{
+		TRefCountPtr<IDXGIAdapter> TempAdapter;
+		for (uint32 AdapterIndex = 0; DXGIFactory1->EnumAdapters(AdapterIndex, TempAdapter.GetInitReference()) != DXGI_ERROR_NOT_FOUND; ++AdapterIndex)
+		{
+			if (TempAdapter)
+			{
+				DXGI_ADAPTER_DESC Desc;
+				TempAdapter->GetDesc(&Desc);
+
+				if (Desc.DedicatedVideoMemory > BestDesc.DedicatedVideoMemory || AdapterIndex == 0)
+				{
+					BestDesc = Desc;
+				}
+			}
+		}
+	}
+
+	Info.VendorId = BestDesc.VendorId;
+	Info.DeviceId = BestDesc.DeviceId;
+	Info.DedicatedVideoMemory = BestDesc.DedicatedVideoMemory;
+	Info.AdapterName = BestDesc.Description;
+
+	if (Info.AdapterName.IsEmpty())
+	{
+		Info.AdapterName = FPlatformMisc::GetPrimaryGPUBrand();
+	}
+}
+
 // TODO: Deduplicate code between this and ShouldDefaultToPerformanceMode
 static bool IsForcedToD3D11_ES31()
 {
 	bool bRequireD3D11ES31 = false;
 
-	FWindowsPlatformApplicationMisc::FGPUInfo BestGPUInfo = FWindowsPlatformApplicationMisc::GetBestGPUInfo();
+	FWindowsGPUInfo BestGPUInfo;
+	GetGPUInfo(BestGPUInfo);
 
-	TArray<FString> DeviceDefaultRHIList;
-	GConfig->GetArray(TEXT("Devices"), TEXT("DeviceForceRHIList"), DeviceDefaultRHIList, GHardwareIni);
+	TArray<FString> DeviceForceRHIList;
+	GConfig->GetArray(TEXT("Devices"), TEXT("DeviceForceRHIList"), DeviceForceRHIList, GHardwareIni);
 
-	FString GPUBrand = FPlatformMisc::GetPrimaryGPUBrand();
-	for (const FString& DeviceDefaultRHIString : DeviceDefaultRHIList)
+	const FString& GPUBrand = BestGPUInfo.AdapterName;
+	for (const FString& DeviceForceRHIString : DeviceForceRHIList)
 	{
-		const TCHAR* Line = *DeviceDefaultRHIString;
+		const TCHAR* Line = *DeviceForceRHIString;
 
 		ensure(Line[0] == TCHAR('('));
 
