Fix bad merge

diff --git a/Engine/Source/Runtime/RHI/Private/Windows/WindowsDynamicRHI.cpp b/Engine/Source/Runtime/RHI/Private/Windows/WindowsDynamicRHI.cpp
--- a/Engine/Source/Runtime/RHI/Private/Windows/WindowsDynamicRHI.cpp
+++ b/Engine/Source/Runtime/RHI/Private/Windows/WindowsDynamicRHI.cpp
@@ -551,47 +551,6 @@ static bool IsForcedToD3D11_ES31()
 	return bRequireD3D11ES31;
 }
 
-struct FWindowsGPUInfo
-{
-	uint32 VendorId = 0;
-	uint32 DeviceId = 0;
-	uint64 DedicatedVideoMemory = 0;
-	FString AdapterName;
-};
-
-static void GetGPUInfo(FWindowsGPUInfo& Info)
-{
-	DXGI_ADAPTER_DESC BestDesc = {};
-	TRefCountPtr<IDXGIFactory1> DXGIFactory1;
-	if (SUCCEEDED(CreateDXGIFactory1(IID_PPV_ARGS(DXGIFactory1.GetInitReference()))))
-	{
-		TRefCountPtr<IDXGIAdapter> TempAdapter;
-		for (uint32 AdapterIndex = 0; DXGIFactory1->EnumAdapters(AdapterIndex, TempAdapter.GetInitReference()) != DXGI_ERROR_NOT_FOUND; ++AdapterIndex)
-		{
-			if (TempAdapter)
-			{
-				DXGI_ADAPTER_DESC Desc;
-				TempAdapter->GetDesc(&Desc);
-
-				if (Desc.DedicatedVideoMemory > BestDesc.DedicatedVideoMemory || AdapterIndex == 0)
-				{
-					BestDesc = Desc;
-				}
-			}
-		}
-	}
-
-	Info.VendorId = BestDesc.VendorId;
-	Info.DeviceId = BestDesc.DeviceId;
-	Info.DedicatedVideoMemory = BestDesc.DedicatedVideoMemory;
-	Info.AdapterName = BestDesc.Description;
-
-	if (Info.AdapterName.IsEmpty())
-	{
-		Info.AdapterName = FPlatformMisc::GetPrimaryGPUBrand();
-	}
-}
-
 // Whether to default to Performance Mode on low-end machines (ES31 feature level).
 // Looks at a few configuration options that control thresholds to enable this as the default.
 static bool ShouldDefaultToPerformanceMode(bool& bOutForceD3D11)
