Consolidate GPU hardware config handling in WindowsDynamicRHI.cpp

diff --git a/Engine/Source/Runtime/RHI/Private/Windows/WindowsDynamicRHI.cpp b/Engine/Source/Runtime/RHI/Private/Windows/WindowsDynamicRHI.cpp
--- a/Engine/Source/Runtime/RHI/Private/Windows/WindowsDynamicRHI.cpp
+++ b/Engine/Source/Runtime/RHI/Private/Windows/WindowsDynamicRHI.cpp
@@ -499,21 +499,15 @@ static void GetGPUInfo(FWindowsGPUInfo& Info)
 	}
 }
 
-// TODO: Deduplicate code between this and ShouldDefaultToPerformanceMode
-static bool IsForcedToD3D11_ES31()
+static bool IsGPUInHardwareListForD3D11_ES31(const FWindowsGPUInfo& BestGPUInfo, const TCHAR* const ConfigKey)
 {
-	bool bRequireD3D11ES31 = false;
-
-	FWindowsGPUInfo BestGPUInfo;
-	GetGPUInfo(BestGPUInfo);
-
-	TArray<FString> DeviceForceRHIList;
-	GConfig->GetArray(TEXT("Devices"), TEXT("DeviceForceRHIList"), DeviceForceRHIList, GHardwareIni);
+	TArray<FString> DeviceRHIList;
+	GConfig->GetArray(TEXT("Devices"), ConfigKey, DeviceRHIList, GHardwareIni);
 
 	const FString& GPUBrand = BestGPUInfo.AdapterName;
-	for (const FString& DeviceForceRHIString : DeviceForceRHIList)
+	for (const FString& DeviceRHIString : DeviceRHIList)
 	{
-		const TCHAR* Line = *DeviceForceRHIString;
+		const TCHAR* Line = *DeviceRHIString;
 
 		ensure(Line[0] == TCHAR('('));
 
@@ -525,9 +519,8 @@ static bool IsForcedToD3D11_ES31()
 
 		if (RHIName.Compare("D3D11_ES31", ESearchCase::IgnoreCase) == 0 && GPUBrand.Compare(DeviceName, ESearchCase::IgnoreCase) == 0)
 		{
-			UE_LOG(LogRHI, Log, TEXT("Found RHIName %s with DeviceName %s: forced feature level to D3D11 %s"), *RHIName, *DeviceName, GetLogName(ERHIFeatureLevel::ES3_1));
-			bRequireD3D11ES31 = true;
-			break;
+			UE_LOG(LogRHI, Log, TEXT("%s: Found RHIName %s with DeviceName %s"), ConfigKey, *RHIName, *DeviceName);
+			return true;
 		}
 
 		FString VendorId;
@@ -542,10 +535,35 @@ static bool IsForcedToD3D11_ES31()
 			BestGPUInfo.VendorId == VendorIdInt && BestGPUInfo.DeviceId == DeviceIdInt &&
 			RHIName.Compare("D3D11_ES31", ESearchCase::IgnoreCase) == 0)
 		{
-			UE_LOG(LogRHI, Log, TEXT("Found RHIName %s with DeviceName %s, VendorId %s, and DeviceId %s: forced feature level to D3D11 %s"), *RHIName, *DeviceName, *VendorId, *DeviceId, GetLogName(ERHIFeatureLevel::ES3_1));
-			bRequireD3D11ES31 = true;
-			break;
+			UE_LOG(LogRHI, Log, TEXT("%s: Found RHIName %s with DeviceName %s, VendorId %s, and DeviceId %s"), ConfigKey, *RHIName, *DeviceName, *VendorId, *DeviceId);
+			return true;
+		}
+	}
+
+	return false;
+}
+
+static bool IsForcedToD3D11_ES31()
+{
+	bool bRequireD3D11ES31 = false;
+
+	static TOptional<bool> ForceD3D11_ES31;
+	if (ForceD3D11_ES31.IsSet())
+	{
+		bRequireD3D11ES31 = ForceD3D11_ES31.GetValue();
+	}
+	else
+	{
+		FWindowsGPUInfo BestGPUInfo;
+		GetGPUInfo(BestGPUInfo);
+
+		bRequireD3D11ES31 = IsGPUInHardwareListForD3D11_ES31(BestGPUInfo, TEXT("DeviceForceRHIList"));
+		if (bRequireD3D11ES31)
+		{
+			UE_LOG(LogRHI, Log, TEXT("Forced feature level to D3D11 %s"), GetLogName(ERHIFeatureLevel::ES3_1));
 		}
+
+		ForceD3D11_ES31 = bRequireD3D11ES31;
 	}
 
 	return bRequireD3D11ES31;
@@ -566,53 +584,15 @@ static bool ShouldDefaultToPerformanceMode(bool& bOutForceD3D11)
 	FWindowsGPUInfo BestGPUInfo;
 	GetGPUInfo(BestGPUInfo);
 
-	TArray<FString> DeviceDefaultRHIList;
-	GConfig->GetArray(TEXT("Devices"), TEXT("DeviceDefaultRHIList"), DeviceDefaultRHIList, GHardwareIni);
-
-	const FString& GPUBrand = BestGPUInfo.AdapterName;
-	for (const FString& DeviceDefaultRHIString : DeviceDefaultRHIList)
+	if (IsGPUInHardwareListForD3D11_ES31(BestGPUInfo, TEXT("DeviceDefaultRHIList")))
 	{
-		const TCHAR* Line = *DeviceDefaultRHIString;
+		UE_LOG(LogRHI, Log, TEXT("Default feature level to %s"), GetLogName(ERHIFeatureLevel::ES3_1));
 
-		ensure(Line[0] == TCHAR('('));
-
-		FString RHIName;
-		FParse::Value(Line+1, TEXT("RHI="), RHIName);
-
-		FString DeviceName;
-		FParse::Value(Line+1, TEXT("DeviceName="), DeviceName);
-
-		if (RHIName.Compare("D3D11_ES31", ESearchCase::IgnoreCase) == 0 && GPUBrand.Compare(DeviceName, ESearchCase::IgnoreCase) == 0)
-		{
-			UE_LOG(LogRHI, Log, TEXT("Found RHIName %s with DeviceName %s: default feature level to %s"), *RHIName, *DeviceName, GetLogName(ERHIFeatureLevel::ES3_1));
-
-			ForceES31 = true;
-			bForceD3D11 = true;
-
-			bOutForceD3D11 = true;
-			return true;
-		}
-
-		FString VendorId;
-		FParse::Value(Line + 1, TEXT("VendorId="), VendorId);
-		uint32 VendorIdInt = FParse::HexNumber(*VendorId);
-
-		FString DeviceId;
-		FParse::Value(Line + 1, TEXT("DeviceId="), DeviceId);
-		uint32 DeviceIdInt = FParse::HexNumber(*DeviceId);
-
-		if (BestGPUInfo.VendorId && BestGPUInfo.DeviceId &&
-			BestGPUInfo.VendorId == VendorIdInt && BestGPUInfo.DeviceId == DeviceIdInt &&
-			RHIName.Compare("D3D11_ES31", ESearchCase::IgnoreCase) == 0)
-		{
-			UE_LOG(LogRHI, Log, TEXT("Found RHIName %s with DeviceName %s, VendorId %s, and DeviceId %s: default feature level to %s"), *RHIName, *DeviceName, *VendorId, *DeviceId, GetLogName(ERHIFeatureLevel::ES3_1));
-
-			ForceES31 = true;
-			bForceD3D11 = true;
+		ForceES31 = true;
+		bForceD3D11 = true;
 
-			bOutForceD3D11 = true;
-			return true;
-		}
+		bOutForceD3D11 = true;
+		return true;
 	}
 
 	// Force Performance mode for laptops. We don't have an explicit way to check for laptops
