Addressing platform preview issues
- bSupportsBindless and StaticShaderBindingLayoutSupport should be set to the current platform's values
- When looking up the "[ShaderPlatformConfig XXX] BindlessConfiguration=" section in the inis, preview shader platforms should use its runtime platform name.

diff --git a/Engine/Source/Runtime/RHI/Private/DataDrivenShaderPlatformInfo.cpp b/Engine/Source/Runtime/RHI/Private/DataDrivenShaderPlatformInfo.cpp
--- a/Engine/Source/Runtime/RHI/Private/DataDrivenShaderPlatformInfo.cpp
+++ b/Engine/Source/Runtime/RHI/Private/DataDrivenShaderPlatformInfo.cpp
@@ -578,6 +578,8 @@ void FGenericDataDrivenShaderPlatformInfo::UpdatePreviewPlatforms()
 				PREVIEW_USE_RUNTIME_VALUE(bSupportsInlineRayTracing);
 				PREVIEW_USE_RUNTIME_VALUE(bInlineRayTracingRequiresBindless);
 				PREVIEW_USE_RUNTIME_VALUE(bSupportsRayTracingShaders);
+				PREVIEW_USE_RUNTIME_VALUE(bSupportsBindless);
+				PREVIEW_USE_RUNTIME_VALUE(StaticShaderBindingLayoutSupport);
 
 				// Settings that will never be supported in preview
 				PREVIEW_FORCE_DISABLE(bSupportsShaderRootConstants);
diff --git a/Engine/Source/Runtime/RHI/Private/RHI.cpp b/Engine/Source/Runtime/RHI/Private/RHI.cpp
--- a/Engine/Source/Runtime/RHI/Private/RHI.cpp
+++ b/Engine/Source/Runtime/RHI/Private/RHI.cpp
@@ -1128,7 +1128,21 @@ bool RHIGetShaderPlatformConfigurationInt(int32& OutSetting, FConfigCacheIni* Co
 bool RHIGetShaderPlatformConfigurationString(FString& OutSetting, FConfigCacheIni* Config, EShaderPlatform Platform, const TCHAR* SettingName)
 {
 	FStringBuilderBase CategoryName;
-	FDataDrivenShaderPlatformInfo::GetName(Platform).ToString(CategoryName);
+
+#if WITH_EDITOR
+	if (FDataDrivenShaderPlatformInfo::GetIsPreviewPlatform(Platform))
+	{
+		// Get the runtime shader platform's category name via the format->platform conversion
+		// This is used to simulate "PREVIEW_USE_RUNTIME_VALUE" from DDSPI. A "ShaderPlatformConfig" system would provide this for us.
+		const FName ShaderFormat = FDataDrivenShaderPlatformInfo::GetShaderFormat(Platform);
+		const EShaderPlatform RuntimePlatform = ShaderFormatToLegacyShaderPlatform(ShaderFormat);
+		FDataDrivenShaderPlatformInfo::GetName(RuntimePlatform).ToString(CategoryName);
+	}
+	else
+#endif
+	{
+		FDataDrivenShaderPlatformInfo::GetName(Platform).ToString(CategoryName);
+	}
 
 	if (CategoryName.Len() > 0)
 	{
