FTextureRWBuffer and FTextureReadBuffer2D should require a passed in command list and not use the immediate command list.

diff --git a/Engine/Plugins/FX/Niagara/Source/Niagara/Private/NiagaraDataInterfaceLandscape.cpp b/Engine/Plugins/FX/Niagara/Source/Niagara/Private/NiagaraDataInterfaceLandscape.cpp
--- a/Engine/Plugins/FX/Niagara/Source/Niagara/Private/NiagaraDataInterfaceLandscape.cpp
+++ b/Engine/Plugins/FX/Niagara/Source/Niagara/Private/NiagaraDataInterfaceLandscape.cpp
@@ -733,18 +733,18 @@ FLandscapeTextureResource::FLandscapeTextureResource(const FIntPoint& InCellCoun
 {
 }
 
-void FLandscapeTextureResource::InitRHI(FRHICommandListBase&)
+void FLandscapeTextureResource::InitRHI(FRHICommandListBase& RHICmdList)
 {
 	if (HeightValues.Num())
 	{
 		FResourceBulkDataArrayView BulkData(HeightValues);
-		HeightTexture.Initialize(TEXT("FLandscapeTextureResource_HeightTexture"), sizeof(float), CellCount.X, CellCount.Y, EPixelFormat::PF_R32_FLOAT, FTextureReadBuffer2D::DefaultTextureInitFlag, &BulkData);
+		HeightTexture.Initialize(RHICmdList, TEXT("FLandscapeTextureResource_HeightTexture"), sizeof(float), CellCount.X, CellCount.Y, EPixelFormat::PF_R32_FLOAT, FTextureReadBuffer2D::DefaultTextureInitFlag, &BulkData);
 	}
 
 	if (PhysMatValues.Num())
 	{
 		FResourceBulkDataArrayView BulkData(PhysMatValues);
-		PhysMatTexture.Initialize(TEXT("FLandscapeTextureResource_PhysMatTexture"), sizeof(uint8), CellCount.X, CellCount.Y, EPixelFormat::PF_R8_UINT, FTextureReadBuffer2D::DefaultTextureInitFlag, &BulkData);
+		PhysMatTexture.Initialize(RHICmdList, TEXT("FLandscapeTextureResource_PhysMatTexture"), sizeof(uint8), CellCount.X, CellCount.Y, EPixelFormat::PF_R8_UINT, FTextureReadBuffer2D::DefaultTextureInitFlag, &BulkData);
 	}
 
 	ReleaseSourceData();
diff --git a/Engine/Plugins/PCG/Source/PCG/Private/Compute/DataInterfaces/PCGLandscapeDataInterface.cpp b/Engine/Plugins/PCG/Source/PCG/Private/Compute/DataInterfaces/PCGLandscapeDataInterface.cpp
--- a/Engine/Plugins/PCG/Source/PCG/Private/Compute/DataInterfaces/PCGLandscapeDataInterface.cpp
+++ b/Engine/Plugins/PCG/Source/PCG/Private/Compute/DataInterfaces/PCGLandscapeDataInterface.cpp
@@ -52,7 +52,7 @@ public:
 		if (HeightValues.Num())
 		{
 			FResourceBulkDataArrayView BulkData(HeightValues);
-			HeightTexture.Initialize(TEXT("FLandscapeTextureResource_HeightTexture"), sizeof(float), CellCount.X, CellCount.Y, EPixelFormat::PF_R32_FLOAT, FTextureReadBuffer2D::DefaultTextureInitFlag, &BulkData);
+			HeightTexture.Initialize(RHICmdList, TEXT("FLandscapeTextureResource_HeightTexture"), sizeof(float), CellCount.X, CellCount.Y, EPixelFormat::PF_R32_FLOAT, FTextureReadBuffer2D::DefaultTextureInitFlag, &BulkData);
 		}
 
 		ReleaseSourceData();
diff --git a/Engine/Source/Runtime/RHI/Public/RHIUtilities.h b/Engine/Source/Runtime/RHI/Public/RHIUtilities.h
--- a/Engine/Source/Runtime/RHI/Public/RHIUtilities.h
+++ b/Engine/Source/Runtime/RHI/Public/RHIUtilities.h
@@ -125,9 +125,9 @@ struct FTextureRWBuffer
 		Release();
 	}
 
-	static constexpr ETextureCreateFlags DefaultTextureInitFlag = TexCreate_ShaderResource | TexCreate_UAV;
+	static constexpr ETextureCreateFlags DefaultTextureInitFlag = ETextureCreateFlags::ShaderResource | ETextureCreateFlags::UAV;
 
-	void Initialize2D(const TCHAR* InDebugName, uint32 BytesPerElement, uint32 SizeX, uint32 SizeY, EPixelFormat Format, ETextureCreateFlags Flags = DefaultTextureInitFlag)
+	void Initialize2D(FRHICommandListBase& RHICmdList, const TCHAR* InDebugName, uint32 BytesPerElement, uint32 SizeX, uint32 SizeY, EPixelFormat Format, ETextureCreateFlags Flags = DefaultTextureInitFlag)
 	{
 		NumBytes = SizeX * SizeY * BytesPerElement;
 
@@ -135,14 +135,19 @@ struct FTextureRWBuffer
 			FRHITextureCreateDesc::Create2D(InDebugName, SizeX, SizeY, Format)
 			.SetFlags(Flags);
 
-		FRHICommandListBase& RHICmdList = FRHICommandListImmediate::Get();
-
-		Buffer = RHICreateTexture(Desc);
+		Buffer = RHICmdList.CreateTexture(Desc);
 		UAV = RHICmdList.CreateUnorderedAccessView(Buffer, FRHIViewDesc::CreateTextureUAV().SetDimensionFromTexture(Buffer));
 		SRV = RHICmdList.CreateShaderResourceView(Buffer, FRHIViewDesc::CreateTextureSRV().SetDimensionFromTexture(Buffer));
 	}
 
-	void Initialize3D(const TCHAR* InDebugName, uint32 BytesPerElement, uint32 SizeX, uint32 SizeY, uint32 SizeZ, EPixelFormat Format, ETextureCreateFlags Flags = DefaultTextureInitFlag)
+	UE_DEPRECATED(5.7, "A command list argument is now required")
+	void Initialize2D(const TCHAR* InDebugName, uint32 BytesPerElement, uint32 SizeX, uint32 SizeY, EPixelFormat Format, ETextureCreateFlags Flags = DefaultTextureInitFlag)
+	{
+		FRHICommandListBase& RHICmdList = FRHICommandListImmediate::Get();
+		Initialize2D(RHICmdList, InDebugName, BytesPerElement, SizeX, SizeY, Format, Flags);
+	}
+
+	void Initialize3D(FRHICommandListBase& RHICmdList, const TCHAR* InDebugName, uint32 BytesPerElement, uint32 SizeX, uint32 SizeY, uint32 SizeZ, EPixelFormat Format, ETextureCreateFlags Flags = DefaultTextureInitFlag)
 	{
 		NumBytes = SizeX * SizeY * SizeZ * BytesPerElement;
 
@@ -150,13 +155,18 @@ struct FTextureRWBuffer
 			FRHITextureCreateDesc::Create3D(InDebugName, SizeX, SizeY, SizeZ, Format)
 			.SetFlags(Flags);
 
-		FRHICommandListBase& RHICmdList = FRHICommandListImmediate::Get();
-
-		Buffer = RHICreateTexture(Desc);
+		Buffer = RHICmdList.CreateTexture(Desc);
 		UAV = RHICmdList.CreateUnorderedAccessView(Buffer, FRHIViewDesc::CreateTextureUAV().SetDimensionFromTexture(Buffer));
 		SRV = RHICmdList.CreateShaderResourceView(Buffer, FRHIViewDesc::CreateTextureSRV().SetDimensionFromTexture(Buffer));
 	}
 
+	UE_DEPRECATED(5.7, "A command list argument is now required")
+	void Initialize3D(const TCHAR* InDebugName, uint32 BytesPerElement, uint32 SizeX, uint32 SizeY, uint32 SizeZ, EPixelFormat Format, ETextureCreateFlags Flags = DefaultTextureInitFlag)
+	{
+		FRHICommandListBase& RHICmdList = FRHICommandListImmediate::Get();
+		Initialize3D(RHICmdList, InDebugName, BytesPerElement, SizeX, SizeY, SizeZ, Format, Flags);
+	}
+
 	void Release()
 	{
 		NumBytes = 0;
@@ -308,22 +318,18 @@ struct FTextureReadBuffer2D
 {
 	FTextureRHIRef Buffer;
 	FShaderResourceViewRHIRef SRV;
-	uint32 NumBytes;
-
-	FTextureReadBuffer2D()
-		: NumBytes(0)
-	{}
+	uint32 NumBytes = 0;
 
+	FTextureReadBuffer2D() = default;
 	~FTextureReadBuffer2D()
 	{
 		Release();
 	}
 
-	const static ETextureCreateFlags DefaultTextureInitFlag = ETextureCreateFlags::ShaderResource;
-	void Initialize(const TCHAR* InDebugName, const uint32 BytesPerElement, const uint32 SizeX, const uint32 SizeY, const EPixelFormat Format, ETextureCreateFlags Flags = DefaultTextureInitFlag, FResourceBulkDataInterface* InBulkData = nullptr)
-	{
-		FRHICommandListBase& RHICmdList = FRHICommandListImmediate::Get();
+	static constexpr ETextureCreateFlags DefaultTextureInitFlag = ETextureCreateFlags::ShaderResource;
 
+	void Initialize(FRHICommandListBase& RHICmdList, const TCHAR* InDebugName, uint32 BytesPerElement, uint32 SizeX, uint32 SizeY, EPixelFormat Format, ETextureCreateFlags Flags = DefaultTextureInitFlag, FResourceBulkDataInterface* InBulkData = nullptr)
+	{
 		NumBytes = SizeX * SizeY * BytesPerElement;
 
 		const FRHITextureCreateDesc Desc =
@@ -331,10 +337,17 @@ struct FTextureReadBuffer2D
 			.SetFlags(Flags)
 			.SetInitActionBulkData(InBulkData);
 
-		Buffer = RHICreateTexture(Desc);
+		Buffer = RHICmdList.CreateTexture(Desc);
 		SRV = RHICmdList.CreateShaderResourceView(Buffer, FRHIViewDesc::CreateTextureSRV().SetDimensionFromTexture(Buffer));
 	}
 
+	UE_DEPRECATED(5.7, "A command list argument is now required")
+	void Initialize(const TCHAR* InDebugName, uint32 BytesPerElement, uint32 SizeX, uint32 SizeY, EPixelFormat Format, ETextureCreateFlags Flags = DefaultTextureInitFlag, FResourceBulkDataInterface* InBulkData = nullptr)
+	{
+		FRHICommandListBase& RHICmdList = FRHICommandListImmediate::Get();
+		Initialize(RHICmdList, InDebugName, BytesPerElement, SizeX, SizeY, Format, Flags, InBulkData);
+	}
+
 	void Release()
 	{
 		NumBytes = 0;
