Track long PSO creation times suspected to be due to an unhealthy driver cache and add an API to retrieve such statistics.

Also change tracking globals to be atomic since they can be modified concurrently by multiple threads creating PSOs.

#rnx

diff --git a/Engine/Source/Runtime/RHI/Private/PipelineStateCache.cpp b/Engine/Source/Runtime/RHI/Private/PipelineStateCache.cpp
--- a/Engine/Source/Runtime/RHI/Private/PipelineStateCache.cpp
+++ b/Engine/Source/Runtime/RHI/Private/PipelineStateCache.cpp
@@ -223,9 +223,29 @@ static FAutoConsoleCommand DumpPipelineCmd(
 	FConsoleCommandDelegate::CreateStatic(DumpPipelineCacheStats)
 );
 
-static int32 GraphicsPSOCreationHitchCount = 0;
-static int32 ComputePSOCreationHitchCount = 0;
-static int32 PrecachedPSOCreationHitchCount = 0;
+int32 GPSOPrecacheUnhealthyCacheHitchThresholdMs = 80;
+static FAutoConsoleVariableRef CVarPSOPrecacheUnhealthyCacheHitchThresholdMs(
+	TEXT("r.PSOPrecache.UnhealthyCacheHitchThresholdMs"),
+	GPSOPrecacheUnhealthyCacheHitchThresholdMs,
+	TEXT("Threshold for runtime previously-precached PSO creation to count towards suspected unhealthy driver cache behavior, in milliseconds (default 80)."),
+	ECVF_Default
+);
+
+int32 GPSOPrecacheUnhealthyCacheMaxHitches = 100;
+static FAutoConsoleVariableRef CVarPSOPrecacheUnhealthyCacheMaxHitches(
+	TEXT("r.PSOPrecache.UnhealthyCacheMaxHitches"),
+	GPSOPrecacheUnhealthyCacheMaxHitches,
+	TEXT("Threshold for the number of runtime previously-precached PSO creation hitches to set a flag indicating a suspected unhealthy driver cache (default 100)."),
+	ECVF_Default
+);
+
+static std::atomic<uint32> RuntimePSOCreationCount = 0;
+static std::atomic<uint32> TotalPSOCreationHitchCount = 0;
+static std::atomic<uint32> GraphicsPSOCreationHitchCount = 0;
+static std::atomic<uint32> ComputePSOCreationHitchCount = 0;
+static std::atomic<uint32> PrecachedPSOCreationHitchCount = 0;
+static std::atomic<uint32> UnhealthyDriverCachePSOCreationHitchCount = 0;
+static std::atomic<bool> bDriverCacheSuspectedUnhealthy = false;
 
 struct FPSOCompilationDebugData
 {
@@ -238,48 +258,67 @@ struct FPSOCompilationDebugData
 
 static inline void CheckAndUpdateHitchCountStat(FPSOPrecacheRequestID::EType PSOType, bool bIsRuntimePSO, const FPSOCompilationDebugData& PSOCompilationDebugData, uint64 StartTime, EPSOPrecacheResult PSOPrecacheResult)
 {
-	if (bIsRuntimePSO)
+	if (!bIsRuntimePSO)
 	{
-		int32 RuntimePSOCreationHitchThreshold = CVarPSORuntimeCreationHitchThreshold.GetValueOnAnyThread();
-		double PSOCreationTimeMs = FPlatformTime::ToMilliseconds64(FPlatformTime::Cycles64() - StartTime);
-		if (PSOCreationTimeMs > RuntimePSOCreationHitchThreshold)
+		return;
+	}
+
+	RuntimePSOCreationCount.fetch_add(1, std::memory_order_relaxed);
+
+	int32 RuntimePSOCreationHitchThreshold = CVarPSORuntimeCreationHitchThreshold.GetValueOnAnyThread();
+	double PSOCreationTimeMs = FPlatformTime::ToMilliseconds64(FPlatformTime::Cycles64() - StartTime);
+	if (PSOCreationTimeMs > RuntimePSOCreationHitchThreshold)
+	{
+		if (PSOType == FPSOPrecacheRequestID::EType::Graphics)
 		{
-			if (PSOType == FPSOPrecacheRequestID::EType::Graphics)
-			{
 #if WITH_RHI_BREADCRUMBS
-				UE_LOG(LogPSOHitching, Verbose, TEXT("Runtime graphics PSO creation hitch (%.2f msec) for %s (precache status: %s) - Breadcrumbs: %s"), PSOCreationTimeMs, *PSOCompilationDebugData.PSOCompilationEventName, LexToString(PSOPrecacheResult), PSOCompilationDebugData.BreadcrumbNode ? *PSOCompilationDebugData.BreadcrumbNode->GetFullPath() : TEXT("Unknown"));
+			UE_LOG(LogPSOHitching, Verbose, TEXT("Runtime graphics PSO creation hitch (%.2f msec) for %s (precache status: %s) - Breadcrumbs: %s"), PSOCreationTimeMs, *PSOCompilationDebugData.PSOCompilationEventName, LexToString(PSOPrecacheResult), PSOCompilationDebugData.BreadcrumbNode ? *PSOCompilationDebugData.BreadcrumbNode->GetFullPath() : TEXT("Unknown"));
 #else
-				UE_LOG(LogPSOHitching, Verbose, TEXT("Runtime graphics PSO creation hitch (%.2f msec) for %s (precache status: %s)"), PSOCreationTimeMs, *PSOCompilationDebugData.PSOCompilationEventName, LexToString(PSOPrecacheResult));
+			UE_LOG(LogPSOHitching, Verbose, TEXT("Runtime graphics PSO creation hitch (%.2f msec) for %s (precache status: %s)"), PSOCreationTimeMs, *PSOCompilationDebugData.PSOCompilationEventName, LexToString(PSOPrecacheResult));
 #endif // WITH_RHI_BREADCRUMBS
-				INC_DWORD_STAT(STAT_RuntimeGraphicsPSOHitchCount);
-				CSV_CUSTOM_STAT(PSO, GraphicsPSOHitch, 1, ECsvCustomStatOp::Accumulate);
-				CSV_CUSTOM_STAT(PSO, GraphicsPSOHitchTime, PSOCreationTimeMs, ECsvCustomStatOp::Accumulate);
-				GraphicsPSOCreationHitchCount++;
-			}
-			else if (PSOType == FPSOPrecacheRequestID::EType::Compute)
-			{
+			INC_DWORD_STAT(STAT_RuntimeGraphicsPSOHitchCount);
+			CSV_CUSTOM_STAT(PSO, GraphicsPSOHitch, 1, ECsvCustomStatOp::Accumulate);
+			CSV_CUSTOM_STAT(PSO, GraphicsPSOHitchTime, PSOCreationTimeMs, ECsvCustomStatOp::Accumulate);
+			GraphicsPSOCreationHitchCount.fetch_add(1, std::memory_order_relaxed);
+		}
+		else if (PSOType == FPSOPrecacheRequestID::EType::Compute)
+		{
 #if WITH_RHI_BREADCRUMBS
-				UE_LOG(LogPSOHitching, Verbose, TEXT("Runtime compute PSO creation hitch (%.2f msec) for %s (precache status: %s) - Breadcrumbs: %s"), PSOCreationTimeMs, *PSOCompilationDebugData.PSOCompilationEventName, LexToString(PSOPrecacheResult), PSOCompilationDebugData.BreadcrumbNode ? *PSOCompilationDebugData.BreadcrumbNode->GetFullPath() : TEXT("Unknown"));
+			UE_LOG(LogPSOHitching, Verbose, TEXT("Runtime compute PSO creation hitch (%.2f msec) for %s (precache status: %s) - Breadcrumbs: %s"), PSOCreationTimeMs, *PSOCompilationDebugData.PSOCompilationEventName, LexToString(PSOPrecacheResult), PSOCompilationDebugData.BreadcrumbNode ? *PSOCompilationDebugData.BreadcrumbNode->GetFullPath() : TEXT("Unknown"));
 #else
-				UE_LOG(LogPSOHitching, Verbose, TEXT("Runtime compute PSO creation hitch (%.2f msec) for %s (precache status: %s)"), PSOCreationTimeMs, *PSOCompilationDebugData.PSOCompilationEventName, LexToString(PSOPrecacheResult));
+			UE_LOG(LogPSOHitching, Verbose, TEXT("Runtime compute PSO creation hitch (%.2f msec) for %s (precache status: %s)"), PSOCreationTimeMs, *PSOCompilationDebugData.PSOCompilationEventName, LexToString(PSOPrecacheResult));
 #endif // WITH_RHI_BREADCRUMBS
-				INC_DWORD_STAT(STAT_RuntimeComputePSOHitchCount);
-				CSV_CUSTOM_STAT(PSO, ComputePSOHitch, 1, ECsvCustomStatOp::Accumulate);
-				CSV_CUSTOM_STAT(PSO, ComputePSOHitchTime, PSOCreationTimeMs, ECsvCustomStatOp::Accumulate);
-				ComputePSOCreationHitchCount++;
-			}
+			INC_DWORD_STAT(STAT_RuntimeComputePSOHitchCount);
+			CSV_CUSTOM_STAT(PSO, ComputePSOHitch, 1, ECsvCustomStatOp::Accumulate);
+			CSV_CUSTOM_STAT(PSO, ComputePSOHitchTime, PSOCreationTimeMs, ECsvCustomStatOp::Accumulate);
+			ComputePSOCreationHitchCount.fetch_add(1, std::memory_order_relaxed);
+		}
 
-			if (PSOPrecacheResult == EPSOPrecacheResult::Complete)
-			{
-				PrecachedPSOCreationHitchCount++;
-			}
+		// Hitches coming from precached PSOs are more concerning. Try to detect whether
+		// the underlying driver cache is unhealthy.
+		if (PSOPrecacheResult == EPSOPrecacheResult::Complete)
+		{
+			PrecachedPSOCreationHitchCount.fetch_add(1, std::memory_order_relaxed);
 
-			int32 TotalHitches = GraphicsPSOCreationHitchCount + ComputePSOCreationHitchCount;
-			if (TotalHitches > 0 && TotalHitches % 50 == 0)
+			if (PSOCreationTimeMs > GPSOPrecacheUnhealthyCacheHitchThresholdMs)
 			{
-				UE_LOG(LogPSOHitching, Log, TEXT("Encountered %d PSO creation hitches so far (%d graphics, %d compute). %d of them were precached."), TotalHitches, GraphicsPSOCreationHitchCount, ComputePSOCreationHitchCount, PrecachedPSOCreationHitchCount);
+				int32 CurrentDriverCacheHitchCount = UnhealthyDriverCachePSOCreationHitchCount.fetch_add(1, std::memory_order_relaxed) + 1;
+
+				if (CurrentDriverCacheHitchCount == GPSOPrecacheUnhealthyCacheMaxHitches)
+				{
+					UE_LOG(LogPSOHitching, Warning, TEXT("Encountered %d PSO creation hitches (threshold: %dms) from PSOs that were precached. This could be an indication of an unhealthy driver cache."), CurrentDriverCacheHitchCount, GPSOPrecacheUnhealthyCacheHitchThresholdMs);
+					bDriverCacheSuspectedUnhealthy.store(true, std::memory_order_relaxed);
+				}
 			}
 		}
+
+		uint32 TotalHitches = TotalPSOCreationHitchCount.fetch_add(1, std::memory_order_relaxed) + 1;
+		if (TotalHitches > 0 && TotalHitches % 50 == 0)
+		{
+			// The numbers logged here might be slightly out-of-sync but it's good enough for a general picture.
+			UE_LOG(LogPSOHitching, Log, TEXT("Encountered %d PSO creation hitches so far (%d graphics, %d compute). %d of them were precached."), TotalHitches, 
+				GraphicsPSOCreationHitchCount.load(std::memory_order_relaxed), ComputePSOCreationHitchCount.load(std::memory_order_relaxed), PrecachedPSOCreationHitchCount.load(std::memory_order_relaxed));
+		}
 	}
 }
 
@@ -4558,15 +4597,29 @@ uint32 PipelineStateCache::NumActivePrecacheRequests()
 	return GPrecacheGraphicsPipelineCache->NumActivePrecacheRequests() + GPrecacheComputePipelineCache->NumActivePrecacheRequests();
 }
 
+FPSORuntimeCreationStats PipelineStateCache::GetPSORuntimeCreationStats()
+{
+	FPSORuntimeCreationStats Stats;
+	Stats.TotalPSOCreations = RuntimePSOCreationCount.load(std::memory_order_relaxed);
+	Stats.ComputePSOHitches = ComputePSOCreationHitchCount.load(std::memory_order_relaxed);
+	Stats.GraphicsPSOHitches = GraphicsPSOCreationHitchCount.load(std::memory_order_relaxed);
+	Stats.PreviouslyPrecachedPSOHitches = PrecachedPSOCreationHitchCount.load(std::memory_order_relaxed);
+	Stats.SuspectedUnhealthyDriverCachePSOHitches = UnhealthyDriverCachePSOCreationHitchCount.load(std::memory_order_relaxed);
+	Stats.bDriverCacheSuspectedUnhealthy = bDriverCacheSuspectedUnhealthy.load(std::memory_order_relaxed);
+	return Stats;
+}
+
 void PipelineStateCache::ResetPSOHitchTrackingStats()
 {
 	SET_DWORD_STAT(STAT_RuntimeGraphicsPSOHitchCount, 0);
-	GraphicsPSOCreationHitchCount = 0;
+	GraphicsPSOCreationHitchCount.store(0, std::memory_order_relaxed);
 
 	SET_DWORD_STAT(STAT_RuntimeComputePSOHitchCount, 0);
-	ComputePSOCreationHitchCount = 0;
+	ComputePSOCreationHitchCount.store(0, std::memory_order_relaxed);
 
-	PrecachedPSOCreationHitchCount = 0;
+	PrecachedPSOCreationHitchCount.store(0, std::memory_order_relaxed);
+	UnhealthyDriverCachePSOCreationHitchCount.store(0, std::memory_order_relaxed);
+	bDriverCacheSuspectedUnhealthy.store(false, std::memory_order_relaxed);
 }
 
 FRHIGraphicsPipelineState* ExecuteSetGraphicsPipelineState(FGraphicsPipelineState* GraphicsPipelineState)
diff --git a/Engine/Source/Runtime/RHI/Public/PipelineStateCache.h b/Engine/Source/Runtime/RHI/Public/PipelineStateCache.h
--- a/Engine/Source/Runtime/RHI/Public/PipelineStateCache.h
+++ b/Engine/Source/Runtime/RHI/Public/PipelineStateCache.h
@@ -111,6 +111,31 @@ struct FPSOPrecacheRequestResult
 	FGraphEventRef AsyncCompileEvent;
 };
 
+// Statistics about PSOs created at runtime, including stats for PSOs which took over a pre-defined threshold to be created. 
+// These only track PSOs requested by the renderer on the critical path (not from background precaching).
+// Note that while slow PSO creation times usually result in a hitch at runtime, it's not a 1-to-1 correlation.
+// There are cases where multiple PSOs are slow in a single frame and therefore cause a single hitch, and cases
+// where a PSO might not be slow enough to cause a visible hitch for a certain frame (depending on the threshold).
+struct FPSORuntimeCreationStats
+{
+	// The total number of PSOs created after a request by the renderer. Does not include PSOs created from background precaching.
+	uint32 TotalPSOCreations = 0;
+
+	// The number of compute and graphics PSOs that took too long to create.
+	uint32 ComputePSOHitches = 0;
+	uint32 GraphicsPSOHitches = 0;
+
+	// How many of the PSOs that took too long to create were previously precached in the background.
+	uint32 PreviouslyPrecachedPSOHitches = 0;
+
+	// How many of the PSOs that too long to create were previously precached in the background and are suspected to have hit a full compilation again.
+	// These use a higher threshold and usually indicate the driver cache is being missed or is not working.
+	uint32 SuspectedUnhealthyDriverCachePSOHitches = 0;
+
+	// Whether the driver cache is suspected to be unhealthy due to the number of very long creation times for precached PSO (controlled by a configurable threshold).
+	bool bDriverCacheSuspectedUnhealthy = false;
+};
+
 extern RHI_API void SetComputePipelineState(FRHIComputeCommandList& RHICmdList, FRHIComputeShader* ComputeShader);
 extern RHI_API void SetGraphicsPipelineStateCheckApply(FRHICommandList& RHICmdList, const FGraphicsPipelineStateInitializer& Initializer, uint32 StencilRef, bool bApplyAdditionalState = true);
 extern RHI_API void SetGraphicsPipelineState(FRHICommandList& RHICmdList, const FGraphicsPipelineStateInitializer& Initializer, uint32 StencilRef, EApplyRendertargetOption ApplyFlags = EApplyRendertargetOption::CheckApply, bool bApplyAdditionalState = true);
@@ -201,6 +226,9 @@ namespace PipelineStateCache
 	/* Set all subsequent high priority requests to highest priority, useful in non-interactive scenarios where maximum PSO throughput is preferable. */
 	extern RHI_API void						PrecachePSOsBoostToHighestPriority(bool bForceHighest);
 
+	/* Return stats on PSOs created at runtime, including slow PSO creation stats. */
+	extern RHI_API FPSORuntimeCreationStats	GetPSORuntimeCreationStats();
+
 	/* Reset the PSO hitch tracking counters */
 	extern RHI_API void						ResetPSOHitchTrackingStats();
 }
