GPU Profiler crash investigation
- We're seeing crashes where the Breadcrumb field of GPU Profiler FEvents are actually GPU (or CPU) timestamps. It is not clear whether these values are garbage left over from previous use of the cache memory, or stomps after the memory was allocated.
- Memset the GPU profiler event stream memory on allocation/free to help distinguish between the two cases.

diff --git a/Engine/Source/Runtime/RHI/Public/GPUProfiler.h b/Engine/Source/Runtime/RHI/Public/GPUProfiler.h
--- a/Engine/Source/Runtime/RHI/Public/GPUProfiler.h
+++ b/Engine/Source/Runtime/RHI/Public/GPUProfiler.h
@@ -106,7 +106,7 @@ namespace UE::RHI::GPUProfiler
 
 		#if WITH_RHI_BREADCRUMBS
 			// The RHI breadcrumb currently at the top of the stack at the frame boundary.
-			FRHIBreadcrumbNode* Breadcrumb;
+			FRHIBreadcrumbNode* const Breadcrumb;
 		#endif
 
 			FFrameBoundary(
@@ -340,11 +340,18 @@ namespace UE::RHI::GPUProfiler
 				{
 					Memory = FMemory::Malloc(sizeof(FChunk), alignof(FChunk));
 				}
+
+				// UE-295331 Investigation : Fill memory with garbage on allocation to catch use-after-free etc.
+				FMemory::Memset(Memory, 0xf7, sizeof(FChunk));
+
 				return Memory;
 			}
 
 			void operator delete(void* Pointer)
 			{
+				// UE-295331 Investigation : Fill memory with garbage on deallocation to catch use-after-free etc.
+				FMemory::Memset(Pointer, 0xe5, sizeof(FChunk));
+
 				MemoryPool.Push(Pointer);
 			}
 
