Fix build

#rnx

diff --git a/Engine/Source/Runtime/Experimental/Chaos/Public/Chaos/CollisionFilterData.h b/Engine/Source/Runtime/Experimental/Chaos/Public/Chaos/CollisionFilterData.h
--- a/Engine/Source/Runtime/Experimental/Chaos/Public/Chaos/CollisionFilterData.h
+++ b/Engine/Source/Runtime/Experimental/Chaos/Public/Chaos/CollisionFilterData.h
@@ -141,7 +141,7 @@ namespace Chaos::Filter
 
 	struct FQueryFilterData
 	{
-		enum CHAOS_API EQueryType
+		enum EQueryType
 		{
 			Channel = 0,
 			ObjectType,
diff --git a/Engine/Source/Runtime/RHI/Public/RHICommandList.h b/Engine/Source/Runtime/RHI/Public/RHICommandList.h
--- a/Engine/Source/Runtime/RHI/Public/RHICommandList.h
+++ b/Engine/Source/Runtime/RHI/Public/RHICommandList.h
@@ -2617,7 +2617,7 @@ FRHICOMMAND_UNNAMED(FRHICommandExecuteMultiIndirectClusterOperation)
 		: Params(InParams)
 	{}
 
-	RHI_API void Execute(FRHICommandListBase & CmdList);
+	RHI_EXECUTE_API void Execute(FRHICommandListBase & CmdList);
 };
 
 FRHICOMMAND_MACRO(FRHICommandRayTraceDispatch)
