Add various functionality for Qualcomm GPUs on Windows..

* Add a function to check whether the current RHI device is a Qualcomm GPU
* Allow the use of r.PSOPrecache.KeepInMemoryUntilUse for Qualcomm
* Support -ClearPSODriverCache for Qualcomm (D3DSCache)

diff --git a/Engine/Source/Runtime/D3D12RHI/Private/Windows/WindowsD3D12Device.cpp b/Engine/Source/Runtime/D3D12RHI/Private/Windows/WindowsD3D12Device.cpp
--- a/Engine/Source/Runtime/D3D12RHI/Private/Windows/WindowsD3D12Device.cpp
+++ b/Engine/Source/Runtime/D3D12RHI/Private/Windows/WindowsD3D12Device.cpp
@@ -1299,6 +1299,11 @@ static void ClearPSODriverCache()
 	FString LocalAppDataFolder = FPlatformMisc::GetEnvironmentVariable(TEXT("LOCALAPPDATA"));
 	if (!LocalAppDataFolder.IsEmpty())
 	{
+		auto DeleteFolder = [](const FString& PSOFolderPath) {
+			IFileManager& FileManager = IFileManager::Get();
+			FileManager.DeleteDirectory(*PSOFolderPath, /*RequireExists*/ false, /*Tree*/ true);
+		};
+
 		auto ClearFolder = [](const FString& PSOPath, const TCHAR* Extension)
 		{
 			TArray<FString> Files;
@@ -1340,6 +1345,12 @@ static void ClearPSODriverCache()
 			FString PSOPath = FPaths::Combine(*LocalAppDataFolder, TEXT(".."), TEXT("LocalLow"), TEXT("Intel"), TEXT("ShaderCache"));
 			ClearFolder(PSOPath, nullptr);
 		}
+		else if (IsRHIDeviceQualcomm())
+		{
+			// Qualcomm uses D3DSCache. The folder contains sub-folders for each application's cache entry so it needs recursive deletion.
+			FString PSOPath = FPaths::Combine(*LocalAppDataFolder, TEXT("D3DSCache"));
+			DeleteFolder(PSOPath);
+		}
 	}
 	else
 	{
diff --git a/Engine/Source/Runtime/RHI/Private/PipelineStateCache.cpp b/Engine/Source/Runtime/RHI/Private/PipelineStateCache.cpp
--- a/Engine/Source/Runtime/RHI/Private/PipelineStateCache.cpp
+++ b/Engine/Source/Runtime/RHI/Private/PipelineStateCache.cpp
@@ -376,8 +376,8 @@ int32 GPSOPrecacheKeepInMemoryUntilUsed = 0;
 static FAutoConsoleVariableRef CVarPSOPrecacheKeepInMemoryUntilUsed(
 	TEXT("r.PSOPrecache.KeepInMemoryUntilUsed"),
 	GPSOPrecacheKeepInMemoryUntilUsed,
-	TEXT("If enabled and if the underlying GPU vendor is NVIDIA, precached PSOs will be kept in memory instead of being deleted immediately after creation, and will only be deleted once they are actually used for rendering or when the maximum limit of in-memory PSOs is reached.\n")
-	TEXT("This can speed up the re-creation of precached PSOs for NVIDIA drivers and avoid small hitches, at the cost of memory.\n")
+	TEXT("If enabled and if the underlying GPU vendor is NVIDIA or Qualcomm, precached PSOs will be kept in memory instead of being deleted immediately after creation, and will only be deleted once they are actually used for rendering or when the maximum limit of in-memory PSOs is reached.\n")
+	TEXT("This can speed up the re-creation of precached PSOs for NVIDIA and Qualcomm drivers and avoid small hitches, at the cost of memory.\n")
 	TEXT("It's recommended to set r.PSOPrecache.KeepInMemoryGraphicsMaxNum and r.PSOPrecache.KeepInMemoryComputeMaxNum to a non-zero value to ensure the number of in-memory PSOs is bounded.\n")
 	TEXT("Valid options:\n")
 	TEXT("0 = off (default)\n")
@@ -403,7 +403,7 @@ static FAutoConsoleVariableRef CVarPSOPrecacheKeepInMemoryComputeMaxNum(
 
 bool ShouldKeepPrecachedPSOsInMemory()
 {
-	return GPSOPrecacheKeepInMemoryUntilUsed && IsRHIDeviceNVIDIA();
+	return GPSOPrecacheKeepInMemoryUntilUsed && (IsRHIDeviceNVIDIA() || IsRHIDeviceQualcomm());
 }
 
 bool ShouldTrackUsedPrecachedPSOs()
diff --git a/Engine/Source/Runtime/RHI/Private/RHI.cpp b/Engine/Source/Runtime/RHI/Private/RHI.cpp
--- a/Engine/Source/Runtime/RHI/Private/RHI.cpp
+++ b/Engine/Source/Runtime/RHI/Private/RHI.cpp
@@ -1355,6 +1355,14 @@ bool IsRHIDeviceNVIDIA()
 	return GRHIVendorId == 0x10DE;
 }
 
+bool IsRHIDeviceQualcomm()
+{
+	check(GRHIVendorId != 0);
+	// On Windows, the DXGI vendor ID is reported as 0x4D4F4351.
+	// On other platforms, it's 0x5143.
+	return GRHIVendorId == 0x5143 || GRHIVendorId == 0x4D4F4351;
+}
+
 bool IsRHIDeviceApple()
 {
     check(GRHIVendorId != 0);
diff --git a/Engine/Source/Runtime/RHI/Public/RHI.h b/Engine/Source/Runtime/RHI/Public/RHI.h
--- a/Engine/Source/Runtime/RHI/Public/RHI.h
+++ b/Engine/Source/Runtime/RHI/Public/RHI.h
@@ -52,6 +52,9 @@ RHI_API bool IsRHIDeviceIntel();
 // to trigger GPU specific optimizations and fallbacks
 RHI_API bool IsRHIDeviceNVIDIA();
 
+// to trigger GPU specific optimizations and fallbacks
+RHI_API bool IsRHIDeviceQualcomm();
+
 // to trigger GPU specific optimizations and fallbacks
 RHI_API bool IsRHIDeviceApple();
 
