Use actual compute shader hash instead of pointer hash for compute PSO precache caching

diff --git a/Engine/Source/Runtime/RHI/Private/PipelineStateCache.cpp b/Engine/Source/Runtime/RHI/Private/PipelineStateCache.cpp
--- a/Engine/Source/Runtime/RHI/Private/PipelineStateCache.cpp
+++ b/Engine/Source/Runtime/RHI/Private/PipelineStateCache.cpp
@@ -2836,7 +2836,7 @@ public:
 
 	static FORCEINLINE uint64 PipelineStateInitializerHash(const FComputePipelineStateInitializer& Key)
 	{
-		return (uint64)Key.ComputeShader;
+		return GetTypeHash(Key.ComputeShader->GetHash());
 	}
 };
 
