Speculative fix for RHI breadcrumb use-after-free crashes in the new GPU profiler
- Crashes were seen in FQueueState::EndBreadcrumb, calling TraceEndGPU on the current Breadcrumb pointer, before the pointer is replaced with the new breadcrumb from the event structure.
- Usually, the breadcrumbs are kept alive by the FEventStream they belong to, however the FQueueState is only holding a raw pointer to the last known breadcrumb.
- It is plausible that the previous breadcrumb belonged to an FEventStream which has completed iteration, and therefore has been destroyed, releasing the allocator references. The next call to EndBreadcrumb will refer to this deleted breadcrumb node.
- Converting the FQueueState::Breadcrumb raw pointer to an FRHIBreadcrumbNodeRef should keep it alive and avoid this use-after-free.

diff --git a/Engine/Source/Runtime/RHI/Private/GPUProfiler.cpp b/Engine/Source/Runtime/RHI/Private/GPUProfiler.cpp
--- a/Engine/Source/Runtime/RHI/Private/GPUProfiler.cpp
+++ b/Engine/Source/Runtime/RHI/Private/GPUProfiler.cpp
@@ -1754,7 +1754,7 @@ namespace UE::RHI::GPUProfiler
 		#if WITH_RHI_BREADCRUMBS
 			TMap<FRHIBreadcrumbData_Stats, int32> ActiveStats;
 			TArray<FRHIBreadcrumbData_Stats> ActiveStatsStack;
-			FRHIBreadcrumbNode* Breadcrumb = nullptr;
+			FRHIBreadcrumbNodeRef Breadcrumb {};
 		#endif
 
 		#if WITH_PROFILEGPU
@@ -2182,10 +2182,8 @@ namespace UE::RHI::GPUProfiler
 				}
 				
 				Breadcrumb->TraceEndGPU(Queue.Value, Event.GPUTimestampBOP);
-
 				Breadcrumb = Event.Breadcrumb->GetParent();
 
-
 			#if WITH_PROFILEGPU
 				if (Profile.bProfileFrame)
 				{
