Minimal version of RHI breadcrumb lifetime fix from CL 48294643 that doesn't modify header files.
- Unlike 48294643, GetBreadcrumbAllocator() has not been modified to hold a reference on new allocator instances immediately. Instead, we take the reference at the point the RDG allocator is swapped in, via AttachBreadcrumbSubTree.

diff --git a/Engine/Source/Runtime/RHI/Private/RHICommandList.cpp b/Engine/Source/Runtime/RHI/Private/RHICommandList.cpp
--- a/Engine/Source/Runtime/RHI/Private/RHICommandList.cpp
+++ b/Engine/Source/Runtime/RHI/Private/RHICommandList.cpp
@@ -329,9 +329,14 @@ void FRHICommandListBase::AttachBreadcrumbSubTree(FRHIBreadcrumbAllocator& Alloc
 	{
 		if (BreadcrumbAllocator)
 		{
+			// Make sure we have a reference to the allocator that was already active on the command list.
 			BreadcrumbAllocatorRefs.AddUnique(BreadcrumbAllocator.Get());
 		}
+
 		BreadcrumbAllocator = Allocator.AsShared();
+
+		// Take a reference on the new allocator we just swapped to.
+		BreadcrumbAllocatorRefs.AddUnique(BreadcrumbAllocator.Get());
 	}
 }
 #endif
