- Keep precached PSOs in memory for certain IHVs to reduce hitches coming from reading PSOs from the driver cache
- Enable IHV specific PSO precaching to reduce the amount of PSOs to compile depending on the IHV

diff --git a/Engine/Source/Runtime/D3D12RHI/Private/D3D12RHI.cpp b/Engine/Source/Runtime/D3D12RHI/Private/D3D12RHI.cpp
--- a/Engine/Source/Runtime/D3D12RHI/Private/D3D12RHI.cpp
+++ b/Engine/Source/Runtime/D3D12RHI/Private/D3D12RHI.cpp
@@ -42,7 +42,7 @@ TAutoConsoleVariable<int32> CVarD3D12ZeroBufferSizeInMB(
 	ECVF_ReadOnly
 	);
 
-static bool GPSOPrecacheD3D12DriverCacheAware = false;
+static bool GPSOPrecacheD3D12DriverCacheAware = true;
 static FAutoConsoleVariableRef CVarPSOPrecacheD3D12DriverCacheAware(
 	TEXT("r.PSOPrecache.D3D12.DriverCacheAware"),
 	GPSOPrecacheD3D12DriverCacheAware,
diff --git a/Engine/Source/Runtime/RHI/Private/PipelineStateCache.cpp b/Engine/Source/Runtime/RHI/Private/PipelineStateCache.cpp
--- a/Engine/Source/Runtime/RHI/Private/PipelineStateCache.cpp
+++ b/Engine/Source/Runtime/RHI/Private/PipelineStateCache.cpp
@@ -372,7 +372,7 @@ static FAutoConsoleVariableRef CVarPrecompileThreadPoolThreadPriority(
 	TEXT("Thread priority for the PSO precompile pool"),
 	ECVF_RenderThreadSafe);
 
-int32 GPSOPrecacheKeepInMemoryUntilUsed = 0;
+int32 GPSOPrecacheKeepInMemoryUntilUsed = 2;
 static FAutoConsoleVariableRef CVarPSOPrecacheKeepInMemoryUntilUsed(
 	TEXT("r.PSOPrecache.KeepInMemoryUntilUsed"),
 	GPSOPrecacheKeepInMemoryUntilUsed,
