Increase default number of compute PSOs to keep memory from PSO precaching to reduce runtime hitching (when r.PSOPrecache.KeepInMemoryUntilUsed is enabled)

Increase system memory load with around 60MBs and fixes all PSO related precached hitches in CitySample

diff --git a/Engine/Source/Runtime/RHI/Private/PipelineStateCache.cpp b/Engine/Source/Runtime/RHI/Private/PipelineStateCache.cpp
--- a/Engine/Source/Runtime/RHI/Private/PipelineStateCache.cpp
+++ b/Engine/Source/Runtime/RHI/Private/PipelineStateCache.cpp
@@ -393,7 +393,7 @@ static FAutoConsoleVariableRef CVarPSOPrecacheKeepInMemoryGraphicsMaxNum(
 	TEXT("If set to 0, no limit will be applied (not recommended outside of testing, as it can cause unbounded memory usage)."),
 	ECVF_RenderThreadSafe);
 
-int32 GPSOPrecacheKeepInMemoryComputeMaxNum = 200;
+int32 GPSOPrecacheKeepInMemoryComputeMaxNum = 1000;
 static FAutoConsoleVariableRef CVarPSOPrecacheKeepInMemoryComputeMaxNum(
 	TEXT("r.PSOPrecache.KeepInMemoryComputeMaxNum"),
 	GPSOPrecacheKeepInMemoryComputeMaxNum,
