Merge (as edit) CL 49806200

Fix RHI task pipe chaining
- RHI task pipes must enforce FIFO ordering of tasks that are enqueued, however this wasn't always happening. Tasks on the RHIThread pipe could run out-of-order.
- The SubmitToGPU task closes the RHIThread pipe so it can hold a reference the LastSubmit event. The next enqueued task on the RHIThread pipe would create a new TaskGraph task (because FTaskPipe::Current is null), but the previous chain's completion event was not inserted as a prerequisite.
- This could lead to chains on the RHI thread running in reverse order. This may have been causing deleted RHI resource crashes or other RHI / GPU crashes.

diff --git a/Engine/Source/Runtime/RHI/Private/RHICommandList.cpp b/Engine/Source/Runtime/RHI/Private/RHICommandList.cpp
--- a/Engine/Source/Runtime/RHI/Private/RHICommandList.cpp
+++ b/Engine/Source/Runtime/RHI/Private/RHICommandList.cpp
@@ -627,12 +627,7 @@ void FRHICommandListExecutor::FTaskPipe::Enqueue(ENamedThreads::Type NamedThread
 	if (LastThread != NamedThread)
 	{
 		// The target thread is changing. End the previous task chain and start a new one.
-		FGraphEventRef PrevEvent = Close();
-		if (PrevEvent)
-		{
-			Prereqs.Add(PrevEvent);
-		}
-
+		Close();
 		LastThread = NamedThread;
 	}
 
@@ -658,6 +653,11 @@ void FRHICommandListExecutor::FTaskPipe::Enqueue(ENamedThreads::Type NamedThread
 		}
 	}
 
+	if (LastEvent)
+	{
+		Current->Prereqs.Add(LastEvent);
+	}
+
 	// Failed to append, or no running task. Start a new one...
 	LastEvent = LaunchTask(Current);
 }
@@ -1480,14 +1480,14 @@ RHI_API FGraphEventRef FRHICommandListExecutor::Submit(TConstArrayView<FRHIComma
 			#endif
 
 				// Submission thread
-				EnqueueSubmitTask(MoveTemp(State->TranslateEvents),
+				FTaskPipe* Pipe = EnqueueSubmitTask(MoveTemp(State->TranslateEvents),
 					[State, Args]() mutable
 					{
 						SCOPED_NAMED_EVENT(RHI_SubmitToGPU, FColor::White);
 						State->Submit(Args);
 					}
 				);
-				LastSubmit = RHIThreadPipe.Close();
+				LastSubmit = Pipe->Close();
 			}
 		);
 
