Fixing preview of platforms with bStripNaniteVertexShaders set.

diff --git a/Engine/Source/Runtime/RHI/Private/DataDrivenShaderPlatformInfo.cpp b/Engine/Source/Runtime/RHI/Private/DataDrivenShaderPlatformInfo.cpp
--- a/Engine/Source/Runtime/RHI/Private/DataDrivenShaderPlatformInfo.cpp
+++ b/Engine/Source/Runtime/RHI/Private/DataDrivenShaderPlatformInfo.cpp
@@ -563,6 +563,7 @@ void FGenericDataDrivenShaderPlatformInfo::UpdatePreviewPlatforms()
 				PREVIEW_DISABLE_IF_RUNTIME_UNSUPPORTED(bSupportsGen5TemporalAA);
 				PREVIEW_DISABLE_IF_RUNTIME_UNSUPPORTED(bSupportsMeshShadersTier0);
 				PREVIEW_DISABLE_IF_RUNTIME_UNSUPPORTED(bSupportsMeshShadersTier1);
+				PREVIEW_DISABLE_IF_RUNTIME_UNSUPPORTED(bStripNaniteVertexShaders);
 				PREVIEW_DISABLE_IF_RUNTIME_UNSUPPORTED(bSupportsMobileMultiView);
 
 				// Settings that need to match the runtime
