Duplicating the null fence version of RHILockStagingBuffer inside the command list.

diff --git a/Engine/Source/Runtime/RHI/Public/RHICommandList.h b/Engine/Source/Runtime/RHI/Public/RHICommandList.h
--- a/Engine/Source/Runtime/RHI/Public/RHICommandList.h
+++ b/Engine/Source/Runtime/RHI/Public/RHICommandList.h
@@ -4728,6 +4728,11 @@ public:
 	{
 		return GDynamicRHI->LockStagingBuffer_RenderThread(*this, StagingBuffer, Fence, Offset, SizeRHI);
 	}
+
+	inline void* LockStagingBuffer(FRHIStagingBuffer* StagingBuffer, uint32 Offset, uint32 SizeRHI)
+	{
+		return LockStagingBuffer(StagingBuffer, nullptr, Offset, SizeRHI);
+	}
 	
 	inline void UnlockStagingBuffer(FRHIStagingBuffer* StagingBuffer)
 	{
@@ -5692,7 +5697,7 @@ inline FRHIShaderLibraryRef RHICreateShaderLibrary(EShaderPlatform Platform, FSt
 
 inline void* RHILockStagingBuffer(FRHIStagingBuffer* StagingBuffer, uint32 Offset, uint32 Size)
 {
-	return FRHICommandListExecutor::GetImmediateCommandList().LockStagingBuffer(StagingBuffer, nullptr, Offset, Size);
+	return FRHICommandListExecutor::GetImmediateCommandList().LockStagingBuffer(StagingBuffer, Offset, Size);
 }
 
 inline void* RHILockStagingBuffer(FRHIStagingBuffer* StagingBuffer, FRHIGPUFence* Fence, uint32 Offset, uint32 Size)
