[RHIBreadcrumbs]
- Fix potential stack overflow by flattening the destruction of breadcrumbs when using instrumentation

diff --git a/Engine/Source/Runtime/RHI/Private/RHIBreadcrumbs.cpp b/Engine/Source/Runtime/RHI/Private/RHIBreadcrumbs.cpp
--- a/Engine/Source/Runtime/RHI/Private/RHIBreadcrumbs.cpp
+++ b/Engine/Source/Runtime/RHI/Private/RHIBreadcrumbs.cpp
@@ -8,6 +8,38 @@
 
 #if WITH_RHI_BREADCRUMBS
 
+#if USING_INSTRUMENTATION
+
+namespace Impl
+{
+	static thread_local TArray<TSharedRef<class FRHIBreadcrumbAllocator>> RHIBreadcrumbAllocator_PendingDestruction;
+	static thread_local int32 RHIBreadcrumbAllocator_DestructionDepth = 0;
+} // namespace Impl
+
+// The destruction can generate incredibly deep callstack so flatten it to avoid stack overflow
+FRHIBreadcrumbAllocator::~FRHIBreadcrumbAllocator()
+{
+	if (Impl::RHIBreadcrumbAllocator_DestructionDepth > 0)
+	{
+		for (TSharedRef<class FRHIBreadcrumbAllocator>& Item : Parents)
+		{
+			Impl::RHIBreadcrumbAllocator_PendingDestruction.Emplace(Item);
+		}
+	}
+	else
+	{
+		TGuardValue Guard(Impl::RHIBreadcrumbAllocator_DestructionDepth, Impl::RHIBreadcrumbAllocator_DestructionDepth + 1);
+		Parents.Reset();
+
+		while (Impl::RHIBreadcrumbAllocator_PendingDestruction.Num())
+		{
+			TArray<TSharedRef<class FRHIBreadcrumbAllocator>> ToDelete = MoveTemp(Impl::RHIBreadcrumbAllocator_PendingDestruction);
+		}
+	}
+}
+
+#endif // USING_INSTRUMENTATION
+
 // Constructor for the sentinel node
 FRHIBreadcrumbNode::FRHIBreadcrumbNode(FRHIBreadcrumbData const& Data)
 	: Data(Data)
diff --git a/Engine/Source/Runtime/RHI/Public/RHIBreadcrumbs.h b/Engine/Source/Runtime/RHI/Public/RHIBreadcrumbs.h
--- a/Engine/Source/Runtime/RHI/Public/RHIBreadcrumbs.h
+++ b/Engine/Source/Runtime/RHI/Public/RHIBreadcrumbs.h
@@ -352,6 +352,9 @@ namespace UE::RHI::GPUProfiler
 		FRHIBreadcrumbAllocatorArray Parents;
 
 	public:
+#if USING_INSTRUMENTATION
+		RHI_API ~FRHIBreadcrumbAllocator();
+#endif
 		FRHIBreadcrumbAllocatorArray const& GetParents() const { return Parents; }
 
 		template <typename TType, typename... TArgs>
