Use the stack flattening RHI breacrumb allocator destructor in all builds
- Long RHI breadcrumb allocator chains can be formed on frames that use multiple RDG builders, so include the fix in all builds, not just when USING_INSTRUMENTATION is enabled.
- Typically this doesn't happen during gameplay, so allow some recursion to avoid the cost of using the PendingDestruction array when the stack isn't deep.

diff --git a/Engine/Source/Runtime/RHI/Private/RHIBreadcrumbs.cpp b/Engine/Source/Runtime/RHI/Private/RHIBreadcrumbs.cpp
--- a/Engine/Source/Runtime/RHI/Private/RHIBreadcrumbs.cpp
+++ b/Engine/Source/Runtime/RHI/Private/RHIBreadcrumbs.cpp
@@ -8,38 +8,45 @@
 
 #if WITH_RHI_BREADCRUMBS
 
-#if USING_INSTRUMENTATION
+thread_local FRHIBreadcrumbAllocatorArray FRHIBreadcrumbAllocator::PendingDestruction;
+thread_local int32 FRHIBreadcrumbAllocator::DestructionDepth = 0;
 
-namespace Impl
-{
-	static thread_local TArray<TSharedRef<class FRHIBreadcrumbAllocator>> RHIBreadcrumbAllocator_PendingDestruction;
-	static thread_local int32 RHIBreadcrumbAllocator_DestructionDepth = 0;
-} // namespace Impl
-
-// The destruction can generate incredibly deep callstack so flatten it to avoid stack overflow
 FRHIBreadcrumbAllocator::~FRHIBreadcrumbAllocator()
 {
-	if (Impl::RHIBreadcrumbAllocator_DestructionDepth > 0)
+	//
+	// There can be frames, particularly in the editor, where there are hundreds of RHI breadcrumb allocators
+	// chained together via the Parents field. Limit the recursive stack depth of the destructor to prevent
+	// stack overflow in these cases.
+	// 
+	// Allowing some recursion avoids paying the cost of using the PendingDestruction array in the general
+	// gameplay case where the chain is only a few allocators long.
+	//
+	if (DestructionDepth > 2)
 	{
-		for (TSharedRef<class FRHIBreadcrumbAllocator>& Item : Parents)
+		if (PendingDestruction.IsEmpty())
+		{
+			PendingDestruction = MoveTemp(Parents);
+		}
+		else
 		{
-			Impl::RHIBreadcrumbAllocator_PendingDestruction.Emplace(Item);
+			for (TSharedRef<FRHIBreadcrumbAllocator>& Item : Parents)
+			{
+				PendingDestruction.Emplace(Item);
+			}
 		}
 	}
 	else
 	{
-		TGuardValue Guard(Impl::RHIBreadcrumbAllocator_DestructionDepth, Impl::RHIBreadcrumbAllocator_DestructionDepth + 1);
+		TGuardValue Guard(DestructionDepth, DestructionDepth + 1);
 		Parents.Reset();
 
-		while (Impl::RHIBreadcrumbAllocator_PendingDestruction.Num())
+		while (PendingDestruction.Num())
 		{
-			TArray<TSharedRef<class FRHIBreadcrumbAllocator>> ToDelete = MoveTemp(Impl::RHIBreadcrumbAllocator_PendingDestruction);
+			FRHIBreadcrumbAllocatorArray ToDelete = MoveTemp(PendingDestruction);
 		}
 	}
 }
 
-#endif // USING_INSTRUMENTATION
-
 // Constructor for the sentinel node
 FRHIBreadcrumbNode::FRHIBreadcrumbNode(FRHIBreadcrumbData const& Data)
 	: Data(Data)
diff --git a/Engine/Source/Runtime/RHI/Public/RHIBreadcrumbs.h b/Engine/Source/Runtime/RHI/Public/RHIBreadcrumbs.h
--- a/Engine/Source/Runtime/RHI/Public/RHIBreadcrumbs.h
+++ b/Engine/Source/Runtime/RHI/Public/RHIBreadcrumbs.h
@@ -346,15 +346,17 @@ namespace UE::RHI::GPUProfiler
 
 	class FRHIBreadcrumbAllocator : public TSharedFromThis<FRHIBreadcrumbAllocator>
 	{
+		static thread_local FRHIBreadcrumbAllocatorArray PendingDestruction;
+		static thread_local int32 DestructionDepth;
+
 		friend FRHIBreadcrumbNode;
 
 		FMemStackBase Inner;
 		FRHIBreadcrumbAllocatorArray Parents;
 
 	public:
-#if USING_INSTRUMENTATION
 		RHI_API ~FRHIBreadcrumbAllocator();
-#endif
+
 		FRHIBreadcrumbAllocatorArray const& GetParents() const { return Parents; }
 
 		template <typename TType, typename... TArgs>
