Add an experimental cvar that prevents the PSO precache compiling tasks that are below a threshold (r.PSOPrecaching.TaskPriorityThreshold).
- This can place less demand on the CPU, but increases memory costs as the precache queue will not empty.
- Increases likelihood of the stand-in default material appearing but also potentially reducing its on-screen duration.

diff --git a/Engine/Source/Runtime/Core/Public/Misc/QueuedThreadPoolWrapper.h b/Engine/Source/Runtime/Core/Public/Misc/QueuedThreadPoolWrapper.h
--- a/Engine/Source/Runtime/Core/Public/Misc/QueuedThreadPoolWrapper.h
+++ b/Engine/Source/Runtime/Core/Public/Misc/QueuedThreadPoolWrapper.h
@@ -49,7 +49,7 @@ public:
 	 * InPriorityMapper           Thread-safe function used to map any priority from this Queue to the priority that should be used when scheduling the task on the underlying thread pool.
 	 */
 	CORE_API FQueuedThreadPoolWrapper(FQueuedThreadPool* InWrappedQueuedThreadPool, int32 InMaxConcurrency = -1, TFunction<EQueuedWorkPriority(EQueuedWorkPriority)> InPriorityMapper = DefaultQueuedWorkPriorityWrapper);
-	CORE_API ~FQueuedThreadPoolWrapper();
+	CORE_API ~FQueuedThreadPoolWrapper() override;
 
 	/**
 	 *  Queued task are not scheduled against the wrapped thread-pool until resumed
@@ -204,9 +204,11 @@ protected:
 
 	// Can be overriden to allocate a more specialized version if needed.
 	virtual FScheduledWork* AllocateScheduledWork() { return new FScheduledWork(); }
+
+	// Can be overridden to ignore specific priority levels if required.
+	CORE_API virtual bool CanSchedule(EQueuedWorkPriority Priority) const;
 private:
 	CORE_API FScheduledWork* AllocateWork(IQueuedWork* InnerWork, EQueuedWorkPriority Priority);
-	CORE_API bool CanSchedule(EQueuedWorkPriority Priority) const;
 	CORE_API bool Create(uint32 InNumQueuedThreads, uint32 StackSize, EThreadPriority ThreadPriority, const TCHAR* Name) override;
 	CORE_API void Destroy() override;
 	CORE_API void Schedule(FScheduledWork* Work = nullptr);
diff --git a/Engine/Source/Runtime/RHI/Private/PipelineStateCache.cpp b/Engine/Source/Runtime/RHI/Private/PipelineStateCache.cpp
--- a/Engine/Source/Runtime/RHI/Private/PipelineStateCache.cpp
+++ b/Engine/Source/Runtime/RHI/Private/PipelineStateCache.cpp
@@ -20,6 +20,7 @@ PipelineStateCache.cpp: Pipeline state cache implementation.
 #include "RHIBreadcrumbs.h"
 #include "Stats/StatsTrace.h"
 #include "Templates/TypeHash.h"
+#include "Misc/QueuedThreadPoolWrapper.h"
 
 // 5.4.2 local change to avoid modifying public headers
 namespace PipelineStateCache
@@ -223,6 +224,20 @@ static FAutoConsoleCommand DumpPipelineCmd(
 	FConsoleCommandDelegate::CreateStatic(DumpPipelineCacheStats)
 );
 
+int32 GPSOPrecacheProcessPriorityThreshold = 0;
+static FAutoConsoleVariableRef CVarPSOOnlyPrecacheMaxPriority(
+	TEXT("r.PSOPrecaching.TaskPriorityThreshold"),
+	GPSOPrecacheProcessPriorityThreshold,
+	TEXT("The precache task priority threshold below which tasks are not processed.\n")
+	TEXT("0: All compile tasks are processed. (default)\n")
+	TEXT("1: Minimum priority compile tasks and above are processed.\n")
+	TEXT("2: Normal priority compile tasks and above are processed.\n")
+	TEXT("3: Only maximum priority compile tasks are processed.")
+	,
+	ECVF_Default
+);
+static_assert((int)FGraphicsPipelineStateInitializer::EPSOPrecacheCompileType::NumTypes == 4, "EPSOPrecacheCompileType has changed, 'r.PSOPrecaching.TaskPriorityThreshold' must be updated.");
+
 int32 GPSOPrecacheUnhealthyCacheHitchThresholdMs = 80;
 static FAutoConsoleVariableRef CVarPSOPrecacheUnhealthyCacheHitchThresholdMs(
 	TEXT("r.PSOPrecache.UnhealthyCacheHitchThresholdMs"),
@@ -427,6 +442,45 @@ const TCHAR* LexToString(EPSOPrecacheResult Result)
 	}
 }
 
+static FGraphicsPipelineStateInitializer::EPSOPrecacheCompileType GetPSOCompileTypeFromQueuePri(EQueuedWorkPriority QueuePri)
+{
+	switch (QueuePri)
+	{
+	case EQueuedWorkPriority::Blocking:
+	case EQueuedWorkPriority::Highest:
+		return FGraphicsPipelineStateInitializer::EPSOPrecacheCompileType::MaxPri;
+	case EQueuedWorkPriority::High:
+	case EQueuedWorkPriority::Normal:
+		return FGraphicsPipelineStateInitializer::EPSOPrecacheCompileType::NormalPri;
+	default: checkNoEntry(); [[fallthrough]];
+	case EQueuedWorkPriority::Low:
+	case EQueuedWorkPriority::Lowest:
+		return FGraphicsPipelineStateInitializer::EPSOPrecacheCompileType::MinPri;
+	}
+}
+
+// PSO thread pool wrapper allows us to control scheduling based on the priority of tasks.
+// when r.PSOPrecaching.PrecacheMaximumPriorityRequestsOnly is true only EQueuedWorkPriority::Highest will be processed.
+// note, the remaining tasks are required to stick around as they are found and boosted when called for.
+class FPSOQueuedThreadPoolWrapper : public FQueuedThreadPoolWrapper
+{
+public:
+	FPSOQueuedThreadPoolWrapper(FQueuedThreadPool* InWrappedQueuedThreadPool)
+		: FQueuedThreadPoolWrapper(InWrappedQueuedThreadPool)
+	{}
+protected:
+
+	virtual bool CanSchedule(EQueuedWorkPriority Priority) const override
+	{
+		if( FQueuedThreadPoolWrapper::CanSchedule(Priority) )
+		{
+			const FGraphicsPipelineStateInitializer::EPSOPrecacheCompileType TaskPri = GetPSOCompileTypeFromQueuePri(Priority);
+			return ((int32)TaskPri >= GPSOPrecacheProcessPriorityThreshold);
+		}
+		return false;
+	}
+};
+
 class FPSOPrecacheThreadPool
 {
 public:
@@ -445,9 +499,9 @@ public:
 			{
 				check(UsePool());
 
-				FQueuedThreadPool* PSOPrecompileCompileThreadPoolLocal = FQueuedThreadPool::Allocate();
+				FQueuedThreadPool* PSOPrecompileCompileThreadPoolLocal = FPSOQueuedThreadPoolWrapper::Allocate();
 				PSOPrecompileCompileThreadPoolLocal->Create(GetDesiredPoolSize(), 512 * 1024, (EThreadPriority)GPSOPrecompileThreadPoolThreadPriority, TEXT("PSOPrecompilePool"));
-				PSOPrecompileCompileThreadPool = PSOPrecompileCompileThreadPoolLocal;
+				PSOPrecompileCompileThreadPool = new FPSOQueuedThreadPoolWrapper(PSOPrecompileCompileThreadPoolLocal);
 			}
 		}
 		return *PSOPrecompileCompileThreadPool;
@@ -458,8 +512,6 @@ public:
 		FQueuedThreadPool* LocalPool = PSOPrecompileCompileThreadPool.exchange(nullptr);
 		if (LocalPool != nullptr)
 		{
-
-			LocalPool->Destroy();
 			delete LocalPool;
 		}
 	}
@@ -774,23 +826,6 @@ namespace PipelineStateCache
 	}
 }
 
-static FGraphicsPipelineStateInitializer::EPSOPrecacheCompileType GetPSOCompileTypeFromQueuePri(EQueuedWorkPriority QueuePri)
-{
-	switch (QueuePri)
-	{
-		case EQueuedWorkPriority::Blocking:
-		case EQueuedWorkPriority::Highest:
-			return FGraphicsPipelineStateInitializer::EPSOPrecacheCompileType::MaxPri;
-		case EQueuedWorkPriority::High:
-		case EQueuedWorkPriority::Normal:
-			return FGraphicsPipelineStateInitializer::EPSOPrecacheCompileType::NormalPri;
-		default: checkNoEntry(); [[fallthrough]];
-		case EQueuedWorkPriority::Low:
-		case EQueuedWorkPriority::Lowest:
-			return FGraphicsPipelineStateInitializer::EPSOPrecacheCompileType::MinPri;
-	}
-}
-
 #if CSV_PROFILER_STATS
 class FPSOCSVStatTracker
 {
@@ -3491,7 +3526,9 @@ public:
 
 		uint64 StartTime = FPlatformTime::Cycles64();
 		using EPSOPrecacheCompileType = FGraphicsPipelineStateInitializer::EPSOPrecacheCompileType;
-		EPSOPrecacheCompileType PSOCompileType = EPSOPrecacheCompileType::NormalPri;
+
+		// when r.PSOPrecaching.PrecacheMaximumPriorityRequestsOnly is active use pri threshold for precache compute PSOs to ensure they are processed.
+		EPSOPrecacheCompileType PSOCompileType = GPSOPrecacheProcessPriorityThreshold ? (EPSOPrecacheCompileType)GPSOPrecacheProcessPriorityThreshold : EPSOPrecacheCompileType::NormalPri;
 		if (OptionalPriorityOverride)
 		{
 			PSOCompileType = (EPSOPrecacheCompileType)FMath::Clamp((uint32)*OptionalPriorityOverride, (uint32)EPSOPrecacheCompileType::MinPri, (uint32)EPSOPrecacheCompileType::MaxPri);
