Minor comment correction for r.PSOPrecaching.TaskPriorityThreshold usage.

diff --git a/Engine/Source/Runtime/RHI/Private/PipelineStateCache.cpp b/Engine/Source/Runtime/RHI/Private/PipelineStateCache.cpp
--- a/Engine/Source/Runtime/RHI/Private/PipelineStateCache.cpp
+++ b/Engine/Source/Runtime/RHI/Private/PipelineStateCache.cpp
@@ -460,7 +460,7 @@ static FGraphicsPipelineStateInitializer::EPSOPrecacheCompileType GetPSOCompileT
 }
 
 // PSO thread pool wrapper allows us to control scheduling based on the priority of tasks.
-// when r.PSOPrecaching.PrecacheMaximumPriorityRequestsOnly is true only EQueuedWorkPriority::Highest will be processed.
+// when r.PSOPrecaching.TaskPriorityThreshold is set only jobs >= priority will be processed.
 // note, the remaining tasks are required to stick around as they are found and boosted when called for.
 class FPSOQueuedThreadPoolWrapper : public FQueuedThreadPoolWrapper
 {
@@ -3527,7 +3527,7 @@ public:
 		uint64 StartTime = FPlatformTime::Cycles64();
 		using EPSOPrecacheCompileType = FGraphicsPipelineStateInitializer::EPSOPrecacheCompileType;
 
-		// when r.PSOPrecaching.PrecacheMaximumPriorityRequestsOnly is active use pri threshold for precache compute PSOs to ensure they are processed.
+		// when r.PSOPrecaching.TaskPriorityThreshold is active use pri threshold for precache compute PSOs to ensure they are processed.
 		EPSOPrecacheCompileType PSOCompileType = GPSOPrecacheProcessPriorityThreshold ? (EPSOPrecacheCompileType)GPSOPrecacheProcessPriorityThreshold : EPSOPrecacheCompileType::NormalPri;
 		if (OptionalPriorityOverride)
 		{
