Remove unused variable

diff --git a/Engine/Source/Runtime/RHI/Public/RHIUtilities.h b/Engine/Source/Runtime/RHI/Public/RHIUtilities.h
--- a/Engine/Source/Runtime/RHI/Public/RHIUtilities.h
+++ b/Engine/Source/Runtime/RHI/Public/RHIUtilities.h
@@ -639,7 +639,6 @@ inline void TransitionRenderPassTargets(FRHICommandList& RHICmdList, const FRHIR
 inline void ClearRenderTarget(FRHICommandList& RHICmdList, FRHITexture* Texture, uint32 MipIndex = 0, uint32 ArraySlice = 0)
 {
 	check(Texture);
-	const FIntPoint Extent = Texture->GetSizeXY();
 	FRHIRenderPassInfo Info(Texture, ERenderTargetActions::Clear_Store);
 	Info.ColorRenderTargets[0].MipIndex = (uint8)MipIndex;
 	Info.ColorRenderTargets[0].ArraySlice = (int32)ArraySlice;
