Increase priority of PSO precache tasks when waiting for shutdown.

diff --git a/Engine/Source/Runtime/RHI/Private/PipelineStateCache.cpp b/Engine/Source/Runtime/RHI/Private/PipelineStateCache.cpp
--- a/Engine/Source/Runtime/RHI/Private/PipelineStateCache.cpp
+++ b/Engine/Source/Runtime/RHI/Private/PipelineStateCache.cpp
@@ -913,9 +913,11 @@ public:
 
 		if (bPrecache)
 		{
-			check(InPri == Pri);
+			// This is called during the task execution where the priority is final.
+			// We set both the priority and the inflight state.
 			check(InPri != FGraphicsPipelineStateInitializer::EPSOPrecacheCompileType::NotSet);
-			GetPriStat(Pri, true)->Add(1);
+			SetPriority(InPri);
+			GetPriStat(InPri, true)->Add(1);
 		}
 	}
 
@@ -1081,6 +1083,14 @@ public:
 	bool WaitCompletion()
 	{
 		bool bNeedsToWait = false;
+
+		FPSOPrecacheAsyncTask* PrecompileTask = GetPrecompileTask();
+		if (PrecompileTask)
+		{
+			bNeedsToWait = bNeedsToWait || !PrecompileTask->IsDone();
+			PrecompileTask->EnsureCompletion();
+		}
+
 		FGraphEvent* CompletionEvent = GetCompletionEvent();
 		if (CompletionEvent && !CompletionEvent->IsComplete())
 		{
@@ -1092,13 +1102,6 @@ public:
 			FTaskGraphInterface::Get().WaitUntilTaskCompletes(CompletionEvent);
 		}
 
-		FPSOPrecacheAsyncTask* PrecompileTask = GetPrecompileTask();
-		if (PrecompileTask)
-		{
-			bNeedsToWait = bNeedsToWait || !PrecompileTask->IsDone();
-			PrecompileTask->EnsureCompletion();
-		}
-
 		ClearCompletionState();
 		return bNeedsToWait;
 	}
@@ -1132,6 +1135,14 @@ public:
 	bool WaitCompletion()
 	{
 		bool bNeedsToWait = false;
+
+		if (PrecompileTask.IsValid())
+		{
+			bNeedsToWait = bNeedsToWait || !PrecompileTask->IsDone();
+			PrecompileTask->EnsureCompletion();
+			PrecompileTask = nullptr;
+		}
+
 		if(CompletionEvent.IsValid() && !CompletionEvent->IsComplete())
 		{
 			QUICK_SCOPE_CYCLE_COUNTER(STAT_FPipelineState_WaitCompletion);
@@ -1143,13 +1154,6 @@ public:
 		}
 		CompletionEvent = nullptr;
 
-		if (PrecompileTask.IsValid())
-		{
-			bNeedsToWait = bNeedsToWait || !PrecompileTask->IsDone();
-			PrecompileTask->EnsureCompletion();
-			PrecompileTask = nullptr;
-		}
-
 		return bNeedsToWait;
 	}
 };
@@ -2502,6 +2506,7 @@ public:
 		// precache tasks may also attempt to lock PrecachePSOsRWLock (TPrecachePipelineCacheBase::PrecacheFinished).
 		// TODO: Replace all of this spin wait.
 		bool bTasksWaiting = true;
+		bool bFirstPass = true;
 		while(bTasksWaiting)
 		{
 			bTasksWaiting = false;
@@ -2512,8 +2517,23 @@ public:
 					FPrecacheTask& PrecacheTask = Iterator->Value;
 					if (PrecacheTask.PipelineState && !PrecacheTask.PipelineState->IsComplete())
 					{
-						bTasksWaiting = true;						
-						break; // release PrecachePSOsRWLock so's to avoid any further blocking of in-progress tasks.
+						bTasksWaiting = true;
+						if (bFirstPass)
+						{
+							// boost the priority of anything we're waiting for to improve throughput.
+							// ensurecomplete is guaranteed to complete, but We're not explicitly waiting for precache tasks here.
+							// So lower pri tasks may never complete when r.PSOPrecaching.TaskPriorityThreshold is in use.
+							if (!EnumHasAnyFlags(PrecacheTask.ReadPSOPrecacheState(), EPSOPrecacheStateMask::HighestPri))
+							{
+								BoostPriorityInternal(PrecacheTask, EPSOPrecachePriority::Highest);
+							}
+						}
+						else
+						{
+							// everything should be boosted on the second pass.
+							check(EnumHasAnyFlags(PrecacheTask.ReadPSOPrecacheState(), EPSOPrecacheStateMask::HighestPri));
+							break; // release PrecachePSOsRWLock so as to avoid any further blocking of in-progress tasks.
+						}
 					}
 					else if (PrecacheTask.PipelineState)
 					{
@@ -2522,6 +2542,7 @@ public:
 						PrecacheTask.PipelineState = nullptr;
 					}
 				}
+				bFirstPass = false;
 				if (!bTasksWaiting)
 				{
 					PrecachedPSOsToCleanup.Empty();
@@ -2581,17 +2602,7 @@ public:
 		FPrecacheTask* FindResult = PrecachedPSOInitializerData.Find(InitializerHash);
 		check(FindResult);
 
-		EPSOPrecacheStateMask NewMask = EPSOPrecacheStateMask::Boosted | (PSOPrecachePriority == EPSOPrecachePriority::Highest ? EPSOPrecacheStateMask::HighestPri : (EPSOPrecacheStateMask)0);
-		EPSOPrecacheStateMask PreviousStateMask = FindResult->AddPSOPrecacheState(NewMask);
-		// It's possible to get a boost request while the task has not been started yet. In this case, TryAddNewState will take care of it
-		// if TryAddNewState is done, then we can proceed to boost it, if the task is not done yet
-		if (!IsCompilationDone(PreviousStateMask) && EnumHasAnyFlags(PreviousStateMask, EPSOPrecacheStateMask::Compiling) && !EnumHasAnyFlags(PreviousStateMask, EPSOPrecacheStateMask::HighestPri) )
-		{
-			if (!EnumHasAnyFlags(PreviousStateMask, EPSOPrecacheStateMask::Boosted) || EnumHasAnyFlags(NewMask, EPSOPrecacheStateMask::HighestPri))
-			{
-				RescheduleTaskToHighPriority(NewMask, PreviousStateMask, FindResult->PipelineState);
-			}
-		}
+		BoostPriorityInternal(*FindResult, PSOPrecachePriority);
 	}
 
 	uint32 NumActivePrecacheRequests()
@@ -2842,6 +2853,22 @@ protected:
 	// Circular buffer of indices of PSOs that are kept in memory if MaxInMemoryPSOs is non-zero.
 	TDeque<uint32> InMemoryPSOIndices;
 	uint32 MaxInMemoryPSOs = 0;
+
+	void BoostPriorityInternal(FPrecacheTask& FindResult, EPSOPrecachePriority PSOPrecachePriority)
+	{
+		// Assume PrecachePSOsRWLock is locked and FindResult will not change.
+		EPSOPrecacheStateMask NewMask = EPSOPrecacheStateMask::Boosted | (PSOPrecachePriority == EPSOPrecachePriority::Highest ? EPSOPrecacheStateMask::HighestPri : (EPSOPrecacheStateMask)0);
+		EPSOPrecacheStateMask PreviousStateMask = FindResult.AddPSOPrecacheState(NewMask);
+		// It's possible to get a boost request while the task has not been started yet. In this case, TryAddNewState will take care of it
+		// if TryAddNewState is done, then we can proceed to boost it, if the task is not done yet
+		if (!IsCompilationDone(PreviousStateMask) && EnumHasAnyFlags(PreviousStateMask, EPSOPrecacheStateMask::Compiling) && !EnumHasAnyFlags(PreviousStateMask, EPSOPrecacheStateMask::HighestPri))
+		{
+			if (!EnumHasAnyFlags(PreviousStateMask, EPSOPrecacheStateMask::Boosted) || EnumHasAnyFlags(NewMask, EPSOPrecacheStateMask::HighestPri))
+			{
+				RescheduleTaskToHighPriority(NewMask, PreviousStateMask, FindResult.PipelineState);
+			}
+		}
+	}
 };
 
 class FPrecacheComputePipelineCache : public TPrecachePipelineCacheBase<FPrecacheComputePipelineCache, FComputePipelineStateInitializer, FComputePipelineState>
