Deprecate RHIGetNativeCommandBuffer
- Replace existing uses with platform-specific interface calls.

diff --git a/Engine/Source/Runtime/RHI/Public/DynamicRHI.h b/Engine/Source/Runtime/RHI/Public/DynamicRHI.h
--- a/Engine/Source/Runtime/RHI/Public/DynamicRHI.h
+++ b/Engine/Source/Runtime/RHI/Public/DynamicRHI.h
@@ -842,15 +842,6 @@ public:
 	// FlushType: Flush RHI Thread
 	virtual void* RHIGetNativeInstance() = 0;
 
-	/**
-	* Provides access to the native command buffer. Generally this should be avoided but is useful for third party plugins.
-	*/
-	// FlushType: Not Thread Safe!
-	virtual void* RHIGetNativeCommandBuffer() 
-	{
-		return nullptr;
-	}
-
 
 	// FlushType: Thread safe
 	virtual IRHICommandContext* RHIGetDefaultContext() = 0;
diff --git a/Engine/Source/Runtime/RHI/Public/RHICommandList.h b/Engine/Source/Runtime/RHI/Public/RHICommandList.h
--- a/Engine/Source/Runtime/RHI/Public/RHICommandList.h
+++ b/Engine/Source/Runtime/RHI/Public/RHICommandList.h
@@ -5025,9 +5025,20 @@ public:
 		return GDynamicRHI->RHIGetNativeInstance();
 	}
 	
+	UE_DEPRECATED(5.8,
+		"GetNativeCommandBuffer is dangerous, and has been deprecated. "
+		"Command buffers are owned by RHI contexts, and can therefore only be retrieved by bottom-of-pipe threads (the RHI thread, or parallel translate tasks). "
+		"Use RHI platform-specific interfaces to retrieve the native command buffer from a bottom-of-pipe RHI command list in a typesafe way.")
 	inline void* GetNativeCommandBuffer()
 	{
-		return GDynamicRHI->RHIGetNativeCommandBuffer();
+		if (ensure(IsBottomOfPipe()))
+		{
+		PRAGMA_DISABLE_DEPRECATION_WARNINGS
+			return GetContext().RHIGetNativeCommandBuffer();
+		PRAGMA_ENABLE_DEPRECATION_WARNINGS
+		}
+
+		return nullptr;
 	}
 
 	//UE_DEPRECATED(5.1, "SubmitCommandsHint is deprecated. Consider calling ImmediateFlush(EImmediateFlushType::DispatchToRHIThread) instead.")
@@ -5697,9 +5708,15 @@ inline void* RHIGetNativeInstance()
 	return FRHICommandListExecutor::GetImmediateCommandList().GetNativeInstance();
 }
 
+UE_DEPRECATED(5.8,
+	"GetNativeCommandBuffer is dangerous, and has been deprecated. "
+	"Command buffers are owned by RHI contexts, and can therefore only be retrieved by bottom-of-pipe threads (the RHI thread, or parallel translate tasks). "
+	"Use RHI platform-specific interfaces to retrieve the native command buffer from a bottom-of-pipe RHI command list in a typesafe way.")
 inline void* RHIGetNativeCommandBuffer()
 {
+PRAGMA_DISABLE_DEPRECATION_WARNINGS
 	return FRHICommandListExecutor::GetImmediateCommandList().GetNativeCommandBuffer();
+PRAGMA_ENABLE_DEPRECATION_WARNINGS
 }
 
 inline FRHIShaderLibraryRef RHICreateShaderLibrary(EShaderPlatform Platform, FString const& FilePath, FString const& Name)
diff --git a/Engine/Source/Runtime/RHI/Public/RHIContext.h b/Engine/Source/Runtime/RHI/Public/RHIContext.h
--- a/Engine/Source/Runtime/RHI/Public/RHIContext.h
+++ b/Engine/Source/Runtime/RHI/Public/RHIContext.h
@@ -531,7 +531,9 @@ public:
 		return Resource->TrackedAccess.Pipelines != ERHIPipeline::None ? Resource->TrackedAccess.Pipelines : ERHIPipeline::Graphics;
 	}
 
+	UE_DEPRECATED(5.8, "RHIGetNativeCommandBuffer has been deprecated. Use RHI platform-specific interfaces to retrieve the native command buffer from a bottom-of-pipe RHI command list in a typesafe way.")
 	virtual void* RHIGetNativeCommandBuffer() { return nullptr; }
+
 	virtual void RHIPostExternalCommandsReset() { }
 
 private:
diff --git a/Engine/Source/Runtime/RHI/Public/RHIValidation.h b/Engine/Source/Runtime/RHI/Public/RHIValidation.h
--- a/Engine/Source/Runtime/RHI/Public/RHIValidation.h
+++ b/Engine/Source/Runtime/RHI/Public/RHIValidation.h
@@ -966,15 +966,6 @@ public:
 		return RHI->RHIGetNativeInstance();
 	}
 
-	/**
-	* Provides access to the native device's command buffer. Generally this should be avoided but is useful for third party plugins.
-	*/
-	// FlushType: Not Thread Safe!!
-	virtual void* RHIGetNativeCommandBuffer() override final
-	{
-		return RHI->RHIGetNativeCommandBuffer();
-	}
-
 	virtual IRHICommandContext* RHIGetDefaultContext() override final;
 	virtual IRHIComputeContext* RHIGetCommandContext(ERHIPipeline Pipeline, FRHIGPUMask GPUMask) override final;
 	virtual IRHIUploadContext* RHIGetUploadContext() final override
diff --git a/Engine/Source/Runtime/RHI/Public/RHIValidationContext.h b/Engine/Source/Runtime/RHI/Public/RHIValidationContext.h
--- a/Engine/Source/Runtime/RHI/Public/RHIValidationContext.h
+++ b/Engine/Source/Runtime/RHI/Public/RHIValidationContext.h
@@ -677,7 +677,9 @@ public:
 
 	virtual void* RHIGetNativeCommandBuffer() override final
 	{
+	PRAGMA_DISABLE_DEPRECATION_WARNINGS
 		return RHIContext->RHIGetNativeCommandBuffer();
+	PRAGMA_ENABLE_DEPRECATION_WARNINGS
 	}
 
 	virtual void RHIBeginTransitions(TArrayView<const FRHITransition*> Transitions) override final
