Fix crash due to missing RHIWriteResourceGPUAddresses implementation in RHI validation layer.

diff --git a/Engine/Source/Runtime/RHI/Public/RHIValidation.h b/Engine/Source/Runtime/RHI/Public/RHIValidation.h
--- a/Engine/Source/Runtime/RHI/Public/RHIValidation.h
+++ b/Engine/Source/Runtime/RHI/Public/RHIValidation.h
@@ -1116,6 +1116,11 @@ public:
 		return RHI->RHICreateShaderBundle(CreateInfo);
 	}
 
+	virtual void RHIWriteResourceGPUAddresses(FRHICommandListBase& RHICmdList, uint32 GPUIndex, TConstArrayView<FRHIResource*> Resources, uint64* AddressesPtr)
+	{
+		return RHI->RHIWriteResourceGPUAddresses(RHICmdList, GPUIndex, Resources, AddressesPtr);
+	}
+
 //protected:
 	static void ReportValidationFailure(const TCHAR* InMessage);
 
