Deprecate / remove RHIGetNativeInstance
- The only useful implementation was in VulkanRHI, which returned the VkInstance.
- Use GetIVulkanDynamicRHI()->RHIGetVkInstance() instead.

diff --git a/Engine/Source/Runtime/Apple/MetalRHI/Private/MetalDynamicRHI.h b/Engine/Source/Runtime/Apple/MetalRHI/Private/MetalDynamicRHI.h
--- a/Engine/Source/Runtime/Apple/MetalRHI/Private/MetalDynamicRHI.h
+++ b/Engine/Source/Runtime/Apple/MetalRHI/Private/MetalDynamicRHI.h
@@ -189,7 +189,6 @@ public:
 	virtual void* METALRHI_API RHIGetNativeDevice() final override;
 	virtual void* METALRHI_API RHIGetNativeGraphicsQueue() final override;
 	virtual void* METALRHI_API RHIGetNativeComputeQueue() final override;
-	virtual void* METALRHI_API RHIGetNativeInstance() final override;
 	
 	virtual class IRHICommandContext* METALRHI_API RHIGetDefaultContext() final override;
 	virtual IRHIUploadContext* RHIGetUploadContext() final override;
diff --git a/Engine/Source/Runtime/Apple/MetalRHI/Private/MetalRHI.cpp b/Engine/Source/Runtime/Apple/MetalRHI/Private/MetalRHI.cpp
--- a/Engine/Source/Runtime/Apple/MetalRHI/Private/MetalRHI.cpp
+++ b/Engine/Source/Runtime/Apple/MetalRHI/Private/MetalRHI.cpp
@@ -1562,11 +1562,6 @@ void* FMetalDynamicRHI::RHIGetNativeComputeQueue()
 	return ImmediateContext.GetCommandQueue().GetQueue();
 }
 
-void* FMetalDynamicRHI::RHIGetNativeInstance()
-{
-	return (void*)Device;
-}
-
 uint16 FMetalDynamicRHI::RHIGetPlatformTextureMaxSampleCount()
 {
 	TArray<ECompositingSampleCount::Type> SamplesArray{ ECompositingSampleCount::Type::One, ECompositingSampleCount::Type::Two, ECompositingSampleCount::Type::Four, ECompositingSampleCount::Type::Eight };
diff --git a/Engine/Source/Runtime/D3D12RHI/Private/D3D12RHI.cpp b/Engine/Source/Runtime/D3D12RHI/Private/D3D12RHI.cpp
--- a/Engine/Source/Runtime/D3D12RHI/Private/D3D12RHI.cpp
+++ b/Engine/Source/Runtime/D3D12RHI/Private/D3D12RHI.cpp
@@ -822,11 +822,6 @@ void* FD3D12DynamicRHI::RHIGetNativeComputeQueue()
 	return (void*)RHIGetCommandQueue();
 }
 
-void* FD3D12DynamicRHI::RHIGetNativeInstance()
-{
-	return nullptr;
-}
-
 /**
 * Returns a supported screen resolution that most closely matches the input.
 * @param Width - Input: Desired resolution width in pixels. Output: A width that the platform supports.
diff --git a/Engine/Source/Runtime/D3D12RHI/Private/D3D12RHIPrivate.h b/Engine/Source/Runtime/D3D12RHI/Private/D3D12RHIPrivate.h
--- a/Engine/Source/Runtime/D3D12RHI/Private/D3D12RHIPrivate.h
+++ b/Engine/Source/Runtime/D3D12RHI/Private/D3D12RHIPrivate.h
@@ -428,7 +428,6 @@ public:
 	virtual void* RHIGetNativeDevice() final override;
 	virtual void* RHIGetNativeGraphicsQueue() final override;
 	virtual void* RHIGetNativeComputeQueue() final override;
-	virtual void* RHIGetNativeInstance() final override;
 	virtual class IRHICommandContext* RHIGetDefaultContext() final override;
 	virtual IRHIComputeContext* RHIGetCommandContext(ERHIPipeline Pipeline, FRHIGPUMask GPUMask) final override;
 	virtual void RHIFinalizeContext(FRHIFinalizeContextArgs&& Args, TRHIPipelineArray<IRHIPlatformCommandList*>& Output) final override;
diff --git a/Engine/Source/Runtime/NullDrv/Public/NullRHI.h b/Engine/Source/Runtime/NullDrv/Public/NullRHI.h
--- a/Engine/Source/Runtime/NullDrv/Public/NullRHI.h
+++ b/Engine/Source/Runtime/NullDrv/Public/NullRHI.h
@@ -559,10 +559,6 @@ public:
 	{ 
 		return 0; 
 	}
-	virtual void* RHIGetNativeInstance() final override
-	{
-		return 0;
-	}
 
 	virtual void RHISetStaticUniformBuffers(const FUniformBufferStaticBindings& InUniformBuffers) final override
 	{
diff --git a/Engine/Source/Runtime/OpenGLDrv/Private/OpenGLDevice.cpp b/Engine/Source/Runtime/OpenGLDrv/Private/OpenGLDevice.cpp
--- a/Engine/Source/Runtime/OpenGLDrv/Private/OpenGLDevice.cpp
+++ b/Engine/Source/Runtime/OpenGLDrv/Private/OpenGLDevice.cpp
@@ -1434,11 +1434,6 @@ void* FOpenGLDynamicRHI::RHIGetNativeDevice()
 	return PlatformDevice;
 }
 
-void* FOpenGLDynamicRHI::RHIGetNativeInstance()
-{
-	return nullptr;
-}
-
 void FOpenGLDynamicRHI::SetCustomPresent(FRHICustomPresent* InCustomPresent)
 {
 	FScopeLock lock(&CustomPresentSection);
diff --git a/Engine/Source/Runtime/OpenGLDrv/Private/OpenGLDrv.h b/Engine/Source/Runtime/OpenGLDrv/Private/OpenGLDrv.h
--- a/Engine/Source/Runtime/OpenGLDrv/Private/OpenGLDrv.h
+++ b/Engine/Source/Runtime/OpenGLDrv/Private/OpenGLDrv.h
@@ -385,7 +385,6 @@ public:
 	UE_API virtual bool RHIGetAvailableResolutions(FScreenResolutionArray& Resolutions, bool bIgnoreRefreshRate) final override;
 	UE_API virtual void RHIGetSupportedResolution(uint32& Width, uint32& Height) final override;
 	UE_API virtual void* RHIGetNativeDevice() final override;
-	UE_API virtual void* RHIGetNativeInstance() final override;
 	UE_API virtual class IRHICommandContext* RHIGetDefaultContext() final override;
 	UE_API virtual IRHIComputeContext* RHIGetCommandContext(ERHIPipeline Pipeline, FRHIGPUMask GPUMask) final override;
 	UE_API virtual void RHIFinalizeContext(FRHIFinalizeContextArgs&& Args, TRHIPipelineArray<IRHIPlatformCommandList*>& Output) final override;
diff --git a/Engine/Source/Runtime/RHI/Public/DynamicRHI.h b/Engine/Source/Runtime/RHI/Public/DynamicRHI.h
--- a/Engine/Source/Runtime/RHI/Public/DynamicRHI.h
+++ b/Engine/Source/Runtime/RHI/Public/DynamicRHI.h
@@ -836,12 +836,6 @@ public:
 		return nullptr;
 	}
 
-	/**
-	* Provides access to the native instance. Generally this should be avoided but is useful for third party plugins.
-	*/
-	// FlushType: Flush RHI Thread
-	virtual void* RHIGetNativeInstance() = 0;
-
 
 	// FlushType: Thread safe
 	virtual IRHICommandContext* RHIGetDefaultContext() = 0;
diff --git a/Engine/Source/Runtime/RHI/Public/RHICommandList.h b/Engine/Source/Runtime/RHI/Public/RHICommandList.h
--- a/Engine/Source/Runtime/RHI/Public/RHICommandList.h
+++ b/Engine/Source/Runtime/RHI/Public/RHICommandList.h
@@ -5017,12 +5017,10 @@ public:
 		return GDynamicRHI->RHIGetNativeComputeQueue();
 	}
 	
+	UE_DEPRECATED(5.8, "GetNativeInstance is deprecated. The only useful implementation was in Vulkan RHI. Use GetIVulkanDynamicRHI()->RHIGetVkInstance() instead.")
 	inline void* GetNativeInstance()
 	{
-		QUICK_SCOPE_CYCLE_COUNTER(STAT_RHIMETHOD_GetNativeInstance_Flush);
-		ImmediateFlush(EImmediateFlushType::FlushRHIThread);
-
-		return GDynamicRHI->RHIGetNativeInstance();
+		return nullptr;
 	}
 	
 	UE_DEPRECATED(5.8,
@@ -5703,9 +5701,10 @@ inline void* RHIGetNativeComputeQueue()
 	return FRHICommandListExecutor::GetImmediateCommandList().GetNativeComputeQueue();
 }
 
+UE_DEPRECATED(5.8, "RHIGetNativeInstance is deprecated. The only useful implementation was in Vulkan RHI. Use GetIVulkanDynamicRHI()->RHIGetVkInstance() instead.")
 inline void* RHIGetNativeInstance()
 {
-	return FRHICommandListExecutor::GetImmediateCommandList().GetNativeInstance();
+	return nullptr;
 }
 
 UE_DEPRECATED(5.8,
diff --git a/Engine/Source/Runtime/RHI/Public/RHIValidation.h b/Engine/Source/Runtime/RHI/Public/RHIValidation.h
--- a/Engine/Source/Runtime/RHI/Public/RHIValidation.h
+++ b/Engine/Source/Runtime/RHI/Public/RHIValidation.h
@@ -957,15 +957,6 @@ public:
 		return RHI->RHIGetNativeComputeQueue();
 	}
 
-	/**
-	* Provides access to the native instance. Generally this should be avoided but is useful for third party plugins.
-	*/
-	// FlushType: Flush RHI Thread
-	virtual void* RHIGetNativeInstance() override final
-	{
-		return RHI->RHIGetNativeInstance();
-	}
-
 	virtual IRHICommandContext* RHIGetDefaultContext() override final;
 	virtual IRHIComputeContext* RHIGetCommandContext(ERHIPipeline Pipeline, FRHIGPUMask GPUMask) override final;
 	virtual IRHIUploadContext* RHIGetUploadContext() final override
diff --git a/Engine/Source/Runtime/VulkanRHI/Private/VulkanDynamicRHI.h b/Engine/Source/Runtime/VulkanRHI/Private/VulkanDynamicRHI.h
--- a/Engine/Source/Runtime/VulkanRHI/Private/VulkanDynamicRHI.h
+++ b/Engine/Source/Runtime/VulkanRHI/Private/VulkanDynamicRHI.h
@@ -232,7 +232,6 @@ public:
 	virtual void* RHIGetNativePhysicalDevice() final override;
 	virtual void* RHIGetNativeGraphicsQueue() final override;
 	virtual void* RHIGetNativeComputeQueue() final override;
-	virtual void* RHIGetNativeInstance() final override;
 	virtual class IRHICommandContext* RHIGetDefaultContext() final override;
 	virtual IRHIComputeContext* RHIGetCommandContext(ERHIPipeline Pipeline, FRHIGPUMask GPUMask) final override;
 	virtual IRHIUploadContext* RHIGetUploadContext() final override;
diff --git a/Engine/Source/Runtime/VulkanRHI/Private/VulkanRHI.cpp b/Engine/Source/Runtime/VulkanRHI/Private/VulkanRHI.cpp
--- a/Engine/Source/Runtime/VulkanRHI/Private/VulkanRHI.cpp
+++ b/Engine/Source/Runtime/VulkanRHI/Private/VulkanRHI.cpp
@@ -1714,11 +1714,6 @@ void* FVulkanDynamicRHI::RHIGetNativeComputeQueue()
 		(void*)Device->GetGraphicsQueue()->GetHandle();
 }
 
-void* FVulkanDynamicRHI::RHIGetNativeInstance()
-{
-	return (void*)GetInstance();
-}
-
 IRHICommandContext* FVulkanDynamicRHI::RHIGetDefaultContext()
 {
 	return &Device->GetImmediateContext();
diff --git a/Engine/Source/Runtime/Windows/D3D11RHI/Private/D3D11Device.cpp b/Engine/Source/Runtime/Windows/D3D11RHI/Private/D3D11Device.cpp
--- a/Engine/Source/Runtime/Windows/D3D11RHI/Private/D3D11Device.cpp
+++ b/Engine/Source/Runtime/Windows/D3D11RHI/Private/D3D11Device.cpp
@@ -786,11 +786,6 @@ void* FD3D11DynamicRHI::RHIGetNativeDevice()
 	return (void*)Direct3DDevice.GetReference();
 }
 
-void* FD3D11DynamicRHI::RHIGetNativeInstance()
-{
-	return nullptr;
-}
-
 void* FD3D11DynamicRHI::RHIGetNativeCommandBuffer()
 {
 	return (void*)Direct3DDeviceIMContext.GetReference();
diff --git a/Engine/Source/Runtime/Windows/D3D11RHI/Private/D3D11RHIPrivate.h b/Engine/Source/Runtime/Windows/D3D11RHI/Private/D3D11RHIPrivate.h
--- a/Engine/Source/Runtime/Windows/D3D11RHI/Private/D3D11RHIPrivate.h
+++ b/Engine/Source/Runtime/Windows/D3D11RHI/Private/D3D11RHIPrivate.h
@@ -507,7 +507,6 @@ public:
 	D3D11RHI_API virtual bool RHIGetAvailableResolutions(FScreenResolutionArray& Resolutions, bool bIgnoreRefreshRate) final override;
 	D3D11RHI_API virtual void RHIGetSupportedResolution(uint32& Width, uint32& Height) final override;
 	D3D11RHI_API virtual void* RHIGetNativeDevice() final override;
-	D3D11RHI_API virtual void* RHIGetNativeInstance() final override;
 	D3D11RHI_API virtual void* RHIGetNativeCommandBuffer() final override;
 	D3D11RHI_API virtual class IRHICommandContext* RHIGetDefaultContext() final override;
 	D3D11RHI_API virtual IRHIComputeContext* RHIGetCommandContext(ERHIPipeline Pipeline, FRHIGPUMask GPUMask) final override;
