Deprecate:
- RHIGetNativePhysicalDevice
- RHIGetNativeGraphicsQueue
- RHIGetNativeComputeQueue

Callers should make use of platform-specific RHI interfaces rather than these generic functions that return void*.

diff --git a/Engine/Source/Runtime/RHI/Public/RHICommandList.h b/Engine/Source/Runtime/RHI/Public/RHICommandList.h
--- a/Engine/Source/Runtime/RHI/Public/RHICommandList.h
+++ b/Engine/Source/Runtime/RHI/Public/RHICommandList.h
@@ -4993,6 +4993,7 @@ public:
 		return GDynamicRHI->RHIGetNativeDevice();
 	}
 	
+	UE_DEPRECATED(5.8, "GetNativePhysicalDevice is deprecated. Use GetIVulkanDynamicRHI()->RHIGetVkPhysicalDevice().")
 	inline void* GetNativePhysicalDevice()
 	{
 		QUICK_SCOPE_CYCLE_COUNTER(STAT_RHIMETHOD_GetNativePhysicalDevice_Flush);
@@ -5001,6 +5002,7 @@ public:
 		return GDynamicRHI->RHIGetNativePhysicalDevice();
 	}
 	
+	UE_DEPRECATED(5.8, "GetNativeGraphicsQueue is deprecated. Use platform-specific RHI interfaces instead.")
 	inline void* GetNativeGraphicsQueue()
 	{
 		QUICK_SCOPE_CYCLE_COUNTER(STAT_RHIMETHOD_GetNativeGraphicsQueue_Flush);
@@ -5009,6 +5011,7 @@ public:
 		return GDynamicRHI->RHIGetNativeGraphicsQueue();
 	}
 	
+	UE_DEPRECATED(5.8, "GetNativeComputeQueue is deprecated. Use platform-specific RHI interfaces instead.")
 	inline void* GetNativeComputeQueue()
 	{
 		QUICK_SCOPE_CYCLE_COUNTER(STAT_RHIMETHOD_GetNativeComputeQueue_Flush);
@@ -5686,19 +5689,28 @@ inline void* RHIGetNativeDevice()
 	return FRHICommandListExecutor::GetImmediateCommandList().GetNativeDevice();
 }
 
+UE_DEPRECATED(5.8, "RHIGetNativePhysicalDevice is deprecated. Use GetIVulkanDynamicRHI()->RHIGetVkPhysicalDevice().")
 inline void* RHIGetNativePhysicalDevice()
 {
+PRAGMA_DISABLE_DEPRECATION_WARNINGS
 	return FRHICommandListExecutor::GetImmediateCommandList().GetNativePhysicalDevice();
+PRAGMA_ENABLE_DEPRECATION_WARNINGS
 }
 
+UE_DEPRECATED(5.8, "RHIGetNativeGraphicsQueue is deprecated. Use platform-specific RHI interfaces instead.")
 inline void* RHIGetNativeGraphicsQueue()
 {
+PRAGMA_DISABLE_DEPRECATION_WARNINGS
 	return FRHICommandListExecutor::GetImmediateCommandList().GetNativeGraphicsQueue();
+PRAGMA_ENABLE_DEPRECATION_WARNINGS
 }
 
+UE_DEPRECATED(5.8, "RHIGetNativeComputeQueue is deprecated. Use platform-specific RHI interfaces instead.")
 inline void* RHIGetNativeComputeQueue()
 {
+PRAGMA_DISABLE_DEPRECATION_WARNINGS
 	return FRHICommandListExecutor::GetImmediateCommandList().GetNativeComputeQueue();
+PRAGMA_ENABLE_DEPRECATION_WARNINGS
 }
 
 UE_DEPRECATED(5.8, "RHIGetNativeInstance is deprecated. The only useful implementation was in Vulkan RHI. Use GetIVulkanDynamicRHI()->RHIGetVkInstance() instead.")
