Remove GPUMask argument from RHIGetCommandContext and RHIGetParallelCommandContext
- RHI contexts should always be capable of recording work to any active GPU node, so the argument is redundant. We always called this function with FRHIGPUMask::All().
- The value is also only used by D3D12, which is the only RHI currently implementing explicit multi-GPU support.

diff --git a/Engine/Source/Runtime/Apple/MetalRHI/Private/MetalDynamicRHI.h b/Engine/Source/Runtime/Apple/MetalRHI/Private/MetalDynamicRHI.h
--- a/Engine/Source/Runtime/Apple/MetalRHI/Private/MetalDynamicRHI.h
+++ b/Engine/Source/Runtime/Apple/MetalRHI/Private/MetalDynamicRHI.h
@@ -193,8 +193,8 @@ public:
 	virtual class IRHICommandContext* METALRHI_API RHIGetDefaultContext() final override;
 	virtual IRHIUploadContext* RHIGetUploadContext() final override;
 	
-	virtual IRHIComputeContext* RHIGetCommandContext(ERHIPipeline Pipeline, FRHIGPUMask GPUMask) final override;
-	virtual IRHIComputeContext* RHIGetParallelCommandContext(FRHIParallelRenderPassInfo const& ParallelRenderPass, FRHIGPUMask GPUMask) final override;
+	virtual IRHIComputeContext* RHIGetCommandContext(ERHIPipeline Pipeline) final override;
+	virtual IRHIComputeContext* RHIGetParallelCommandContext(FRHIParallelRenderPassInfo const& ParallelRenderPass) final override;
 	
 	virtual void RHIProcessDeleteQueue() final override;
 	virtual void RHIFinalizeContext(FRHIFinalizeContextArgs&& Args, TRHIPipelineArray<IRHIPlatformCommandList*>& Output) final override;
diff --git a/Engine/Source/Runtime/Apple/MetalRHI/Private/MetalRHI.cpp b/Engine/Source/Runtime/Apple/MetalRHI/Private/MetalRHI.cpp
--- a/Engine/Source/Runtime/Apple/MetalRHI/Private/MetalRHI.cpp
+++ b/Engine/Source/Runtime/Apple/MetalRHI/Private/MetalRHI.cpp
@@ -1616,7 +1616,7 @@ IRHICommandContext* FMetalDynamicRHI::RHIGetDefaultContext()
 	return &ImmediateContext;
 }
 
-IRHIComputeContext* FMetalDynamicRHI::RHIGetCommandContext(ERHIPipeline Pipeline, FRHIGPUMask GPUMask)
+IRHIComputeContext* FMetalDynamicRHI::RHIGetCommandContext(ERHIPipeline Pipeline)
 {
 	check(GRHISupportsParallelRHIExecute);
 	
@@ -1631,7 +1631,7 @@ IRHIComputeContext* FMetalDynamicRHI::RHIGetCommandContext(ERHIPipeline Pipeline
 	return static_cast<IRHIComputeContext*>(Context);
 }
 
-IRHIComputeContext* FMetalDynamicRHI::RHIGetParallelCommandContext(FRHIParallelRenderPassInfo const& ParallelRenderPass, FRHIGPUMask GPUMask)
+IRHIComputeContext* FMetalDynamicRHI::RHIGetParallelCommandContext(FRHIParallelRenderPassInfo const& ParallelRenderPass)
 {
 	check(GRHISupportsParallelRHIExecute);
 	
diff --git a/Engine/Source/Runtime/D3D12RHI/Private/D3D12CommandContext.cpp b/Engine/Source/Runtime/D3D12RHI/Private/D3D12CommandContext.cpp
--- a/Engine/Source/Runtime/D3D12RHI/Private/D3D12CommandContext.cpp
+++ b/Engine/Source/Runtime/D3D12RHI/Private/D3D12CommandContext.cpp
@@ -743,8 +743,10 @@ void FD3D12DynamicRHI::UpdateMemoryStats()
 #endif // PLATFORM_WINDOWS && (STATS || CSV_PROFILER_STATS)
 }
 
-IRHIComputeContext* FD3D12DynamicRHI::RHIGetCommandContext(ERHIPipeline Pipeline, FRHIGPUMask GPUMask)
+IRHIComputeContext* FD3D12DynamicRHI::RHIGetCommandContext(ERHIPipeline Pipeline)
 {
+	FRHIGPUMask GPUMask = FRHIGPUMask::All();
+
 	if (GPUMask.HasSingleIndex())
 	{
 		FD3D12Device* Device = GetAdapter().GetDevice(GPUMask.ToIndex());
diff --git a/Engine/Source/Runtime/D3D12RHI/Private/D3D12RHIPrivate.h b/Engine/Source/Runtime/D3D12RHI/Private/D3D12RHIPrivate.h
--- a/Engine/Source/Runtime/D3D12RHI/Private/D3D12RHIPrivate.h
+++ b/Engine/Source/Runtime/D3D12RHI/Private/D3D12RHIPrivate.h
@@ -429,7 +429,7 @@ public:
 	virtual void* RHIGetNativeGraphicsQueue() final override;
 	virtual void* RHIGetNativeComputeQueue() final override;
 	virtual class IRHICommandContext* RHIGetDefaultContext() final override;
-	virtual IRHIComputeContext* RHIGetCommandContext(ERHIPipeline Pipeline, FRHIGPUMask GPUMask) final override;
+	virtual IRHIComputeContext* RHIGetCommandContext(ERHIPipeline Pipeline) final override;
 	virtual void RHIFinalizeContext(FRHIFinalizeContextArgs&& Args, TRHIPipelineArray<IRHIPlatformCommandList*>& Output) final override;
 	virtual void RHISubmitCommandLists(FRHISubmitCommandListsArgs&& Args) final override;
 	virtual void RHIProcessDeleteQueue() final override;
diff --git a/Engine/Source/Runtime/NullDrv/Public/NullRHI.h b/Engine/Source/Runtime/NullDrv/Public/NullRHI.h
--- a/Engine/Source/Runtime/NullDrv/Public/NullRHI.h
+++ b/Engine/Source/Runtime/NullDrv/Public/NullRHI.h
@@ -582,12 +582,12 @@ public:
 		return this; 
 	}
 
-	virtual IRHIComputeContext* RHIGetCommandContext(ERHIPipeline Pipeline, FRHIGPUMask GPUMask) final override
+	virtual IRHIComputeContext* RHIGetCommandContext(ERHIPipeline Pipeline) final override
 	{
 		return nullptr;
 	}
 	
-	virtual IRHIComputeContext* RHIGetParallelCommandContext(FRHIParallelRenderPassInfo const& ParallelRenderPass, FRHIGPUMask GPUMask) final override
+	virtual IRHIComputeContext* RHIGetParallelCommandContext(FRHIParallelRenderPassInfo const& ParallelRenderPass) final override
 	{
 		return nullptr;
 	}
diff --git a/Engine/Source/Runtime/OpenGLDrv/Private/OpenGLCommands.cpp b/Engine/Source/Runtime/OpenGLDrv/Private/OpenGLCommands.cpp
--- a/Engine/Source/Runtime/OpenGLDrv/Private/OpenGLCommands.cpp
+++ b/Engine/Source/Runtime/OpenGLDrv/Private/OpenGLCommands.cpp
@@ -2748,7 +2748,7 @@ IRHICommandContext* FOpenGLDynamicRHI::RHIGetDefaultContext()
 	return this;
 }
 
-IRHIComputeContext* FOpenGLDynamicRHI::RHIGetCommandContext(ERHIPipeline Pipeline, FRHIGPUMask GPUMask)
+IRHIComputeContext* FOpenGLDynamicRHI::RHIGetCommandContext(ERHIPipeline Pipeline)
 {
 	UE_LOG(LogRHI, Fatal, TEXT("FOpenGLDynamicRHI::RHIGetCommandContext should never be called. OpenGL RHI does not implement parallel command list execution."));
 	return nullptr;
diff --git a/Engine/Source/Runtime/OpenGLDrv/Private/OpenGLDrv.h b/Engine/Source/Runtime/OpenGLDrv/Private/OpenGLDrv.h
--- a/Engine/Source/Runtime/OpenGLDrv/Private/OpenGLDrv.h
+++ b/Engine/Source/Runtime/OpenGLDrv/Private/OpenGLDrv.h
@@ -386,7 +386,7 @@ public:
 	UE_API virtual void RHIGetSupportedResolution(uint32& Width, uint32& Height) final override;
 	UE_API virtual void* RHIGetNativeDevice() final override;
 	UE_API virtual class IRHICommandContext* RHIGetDefaultContext() final override;
-	UE_API virtual IRHIComputeContext* RHIGetCommandContext(ERHIPipeline Pipeline, FRHIGPUMask GPUMask) final override;
+	UE_API virtual IRHIComputeContext* RHIGetCommandContext(ERHIPipeline Pipeline) final override;
 	UE_API virtual void RHIFinalizeContext(FRHIFinalizeContextArgs&& Args, TRHIPipelineArray<IRHIPlatformCommandList*>& Output) final override;
 	UE_API virtual void RHISubmitCommandLists(FRHISubmitCommandListsArgs&& Args) final override;
 
diff --git a/Engine/Source/Runtime/RHI/Private/RHICommandList.cpp b/Engine/Source/Runtime/RHI/Private/RHICommandList.cpp
--- a/Engine/Source/Runtime/RHI/Private/RHICommandList.cpp
+++ b/Engine/Source/Runtime/RHI/Private/RHICommandList.cpp
@@ -417,13 +417,13 @@ void FRHICommandListBase::ActivatePipelines(ERHIPipeline Pipelines)
 				{
 					if (ExecutingCmdList.IsSubCommandList())
 					{
-						Context = GDynamicRHI->RHIGetParallelCommandContext(*ExecutingCmdList.SubRenderPassInfo, FRHIGPUMask::All());
+						Context = GDynamicRHI->RHIGetParallelCommandContext(*ExecutingCmdList.SubRenderPassInfo);
 					}
 					else
 					{
 						// Need to handle the "immediate" context separately.
 						Context = ExecutingCmdList.AllowParallelTranslate()
-							? GDynamicRHI->RHIGetCommandContext(Pipeline, FRHIGPUMask::All()) // This mask argument specifies which contexts are included in an mGPU redirector (we always want all of them).
+							? GDynamicRHI->RHIGetCommandContext(Pipeline)
 							: ::RHIGetDefaultContext();
 					}
 				}
@@ -442,7 +442,7 @@ void FRHICommandListBase::ActivatePipelines(ERHIPipeline Pipelines)
 
 				if (!Context)
 				{
-					Context = GDynamicRHI->RHIGetCommandContext(Pipeline, FRHIGPUMask::All()); // This mask argument specifies which contexts are included in an mGPU redirector (we always want all of them).
+					Context = GDynamicRHI->RHIGetCommandContext(Pipeline);
 					check(Context);
 				}
 
diff --git a/Engine/Source/Runtime/RHI/Private/RHIValidation.cpp b/Engine/Source/Runtime/RHI/Private/RHIValidation.cpp
--- a/Engine/Source/Runtime/RHI/Private/RHIValidation.cpp
+++ b/Engine/Source/Runtime/RHI/Private/RHIValidation.cpp
@@ -619,9 +619,9 @@ struct FValidationCommandList : public IRHIPlatformCommandList
 	TArray<RHIValidation::FOperation> CompletedOpList;
 };
 
-IRHIComputeContext* FValidationRHI::RHIGetCommandContext(ERHIPipeline Pipeline, FRHIGPUMask GPUMask)
+IRHIComputeContext* FValidationRHI::RHIGetCommandContext(ERHIPipeline Pipeline)
 {
-	IRHIComputeContext* InnerContext = RHI->RHIGetCommandContext(Pipeline, GPUMask);
+	IRHIComputeContext* InnerContext = RHI->RHIGetCommandContext(Pipeline);
 	check(InnerContext);
 
 	switch (Pipeline)
@@ -696,7 +696,7 @@ void FValidationRHI::RHIFinalizeContext(FRHIFinalizeContextArgs&& Args, TRHIPipe
 	}
 }
 
-IRHIComputeContext* FValidationRHI::RHIGetParallelCommandContext(FRHIParallelRenderPassInfo const& ParallelRenderPass, FRHIGPUMask GPUMask)
+IRHIComputeContext* FValidationRHI::RHIGetParallelCommandContext(FRHIParallelRenderPassInfo const& ParallelRenderPass)
 {
 	// If a platform has a ChildWait or ParentWait it is expected that they will override RHIGetParallelCommandContext
 	// otherwise, we need to manually call RHIGetCommandContext and RHIBeginRenderPass seperately, because the default implementation
@@ -704,7 +704,7 @@ IRHIComputeContext* FValidationRHI::RHIGetParallelCommandContext(FRHIParallelRen
 	// TODO: Remove this after implementing RHISetupParallelPass
 	if (GRHIParallelRHIExecuteChildWait || GRHIParallelRHIExecuteParentWait)
 	{
-		IRHIComputeContext* InnerContext = RHI->RHIGetParallelCommandContext(ParallelRenderPass, GPUMask);
+		IRHIComputeContext* InnerContext = RHI->RHIGetParallelCommandContext(ParallelRenderPass);
 		check(InnerContext);
 
 		FValidationContext* OuterContext = new FValidationContext(FValidationContext::EType::Parallel);
@@ -722,7 +722,7 @@ IRHIComputeContext* FValidationRHI::RHIGetParallelCommandContext(FRHIParallelRen
 	}
 	else
 	{
-		IRHICommandContext* Context = static_cast<IRHICommandContext*>(RHIGetCommandContext(ERHIPipeline::Graphics, GPUMask));
+		IRHICommandContext* Context = static_cast<IRHICommandContext*>(RHIGetCommandContext(ERHIPipeline::Graphics));
 		Context->RHIBeginRenderPass(ParallelRenderPass, ParallelRenderPass.PassName);
 		return Context;
 	}
diff --git a/Engine/Source/Runtime/RHI/Public/DynamicRHI.h b/Engine/Source/Runtime/RHI/Public/DynamicRHI.h
--- a/Engine/Source/Runtime/RHI/Public/DynamicRHI.h
+++ b/Engine/Source/Runtime/RHI/Public/DynamicRHI.h
@@ -847,7 +847,7 @@ public:
 	//
 	// Called by parallel worker threads, and the render thread. Platform implementations must be thread safe.
 	//
-	virtual IRHIComputeContext* RHIGetCommandContext(ERHIPipeline Pipeline, FRHIGPUMask GPUMask) = 0;
+	virtual IRHIComputeContext* RHIGetCommandContext(ERHIPipeline Pipeline) = 0;
 	
 	//
 	// Retrieves a new command context to begin the recording of a new platform command list.
@@ -856,9 +856,9 @@ public:
 	//
 	// Called by parallel worker threads, and the render thread. Platform implementations must be thread safe.
 	//
-	virtual IRHIComputeContext* RHIGetParallelCommandContext(FRHIParallelRenderPassInfo const& ParallelRenderPass, FRHIGPUMask GPUMask)
+	virtual IRHIComputeContext* RHIGetParallelCommandContext(FRHIParallelRenderPassInfo const& ParallelRenderPass)
 	{
-		IRHICommandContext* Context = static_cast<IRHICommandContext*>(RHIGetCommandContext(ERHIPipeline::Graphics, GPUMask));
+		IRHICommandContext* Context = static_cast<IRHICommandContext*>(RHIGetCommandContext(ERHIPipeline::Graphics));
 		Context->RHIBeginRenderPass(ParallelRenderPass, ParallelRenderPass.PassName);
 		
 		return Context;
diff --git a/Engine/Source/Runtime/RHI/Public/RHIValidation.h b/Engine/Source/Runtime/RHI/Public/RHIValidation.h
--- a/Engine/Source/Runtime/RHI/Public/RHIValidation.h
+++ b/Engine/Source/Runtime/RHI/Public/RHIValidation.h
@@ -958,7 +958,7 @@ public:
 	}
 
 	virtual IRHICommandContext* RHIGetDefaultContext() override final;
-	virtual IRHIComputeContext* RHIGetCommandContext(ERHIPipeline Pipeline, FRHIGPUMask GPUMask) override final;
+	virtual IRHIComputeContext* RHIGetCommandContext(ERHIPipeline Pipeline) override final;
 	virtual IRHIUploadContext* RHIGetUploadContext() final override
 	{
 		return RHI->RHIGetUploadContext();
@@ -969,7 +969,7 @@ public:
 	virtual void RHIFinalizeContext(FRHIFinalizeContextArgs&& Args, TRHIPipelineArray<IRHIPlatformCommandList*>& Output) override final;
 	virtual void RHISubmitCommandLists(FRHISubmitCommandListsArgs&& Args) override final;
 	
-	virtual IRHIComputeContext* RHIGetParallelCommandContext(FRHIParallelRenderPassInfo const& ParallelRenderPass, FRHIGPUMask GPUMask) override final;
+	virtual IRHIComputeContext* RHIGetParallelCommandContext(FRHIParallelRenderPassInfo const& ParallelRenderPass) override final;
 	virtual IRHIPlatformCommandList* RHIFinalizeParallelContext(IRHIComputeContext* Context) override final;
 
 	virtual void RHIProcessDeleteQueue() override final
diff --git a/Engine/Source/Runtime/VulkanRHI/Private/VulkanCommands.cpp b/Engine/Source/Runtime/VulkanRHI/Private/VulkanCommands.cpp
--- a/Engine/Source/Runtime/VulkanRHI/Private/VulkanCommands.cpp
+++ b/Engine/Source/Runtime/VulkanRHI/Private/VulkanCommands.cpp
@@ -901,7 +901,7 @@ void FVulkanCommandListContext::RHIWriteGPUFence(FRHIGPUFence* FenceRHI)
 	checkNoEntry(); // Should never be called
 }
 
-IRHIComputeContext* FVulkanDynamicRHI::RHIGetCommandContext(ERHIPipeline Pipeline, FRHIGPUMask GPUMask)
+IRHIComputeContext* FVulkanDynamicRHI::RHIGetCommandContext(ERHIPipeline Pipeline)
 {
 	check(IsSingleRHIPipeline(Pipeline));
 	return new FVulkanCommandListContext(*Device, Pipeline, &Device->GetImmediateContext());
diff --git a/Engine/Source/Runtime/VulkanRHI/Private/VulkanDynamicRHI.h b/Engine/Source/Runtime/VulkanRHI/Private/VulkanDynamicRHI.h
--- a/Engine/Source/Runtime/VulkanRHI/Private/VulkanDynamicRHI.h
+++ b/Engine/Source/Runtime/VulkanRHI/Private/VulkanDynamicRHI.h
@@ -233,13 +233,13 @@ public:
 	virtual void* RHIGetNativeGraphicsQueue() final override;
 	virtual void* RHIGetNativeComputeQueue() final override;
 	virtual class IRHICommandContext* RHIGetDefaultContext() final override;
-	virtual IRHIComputeContext* RHIGetCommandContext(ERHIPipeline Pipeline, FRHIGPUMask GPUMask) final override;
+	virtual IRHIComputeContext* RHIGetCommandContext(ERHIPipeline Pipeline) final override;
 	virtual IRHIUploadContext* RHIGetUploadContext() final override;
 	virtual void RHIFinalizeContext(FRHIFinalizeContextArgs&& Args, TRHIPipelineArray<IRHIPlatformCommandList*>& Output) final override;
 	virtual void RHISubmitCommandLists(FRHISubmitCommandListsArgs&& Args) final override;
 	virtual uint64 RHIGetMinimumAlignmentForBufferBackedSRV(EPixelFormat Format) final override;
 
-	virtual IRHIComputeContext* RHIGetParallelCommandContext(FRHIParallelRenderPassInfo const& ParallelRenderPass, FRHIGPUMask GPUMask) final override;
+	virtual IRHIComputeContext* RHIGetParallelCommandContext(FRHIParallelRenderPassInfo const& ParallelRenderPass) final override;
 	virtual IRHIPlatformCommandList* RHIFinalizeParallelContext(IRHIComputeContext* Context) final override;
 
 	// Transient Resource Functions and Helpers
diff --git a/Engine/Source/Runtime/VulkanRHI/Private/VulkanSubmission.cpp b/Engine/Source/Runtime/VulkanRHI/Private/VulkanSubmission.cpp
--- a/Engine/Source/Runtime/VulkanRHI/Private/VulkanSubmission.cpp
+++ b/Engine/Source/Runtime/VulkanRHI/Private/VulkanSubmission.cpp
@@ -283,7 +283,7 @@ void FVulkanDynamicRHI::RHIFinalizeContext(FRHIFinalizeContextArgs&& Args, TRHIP
 	}
 }
 
-IRHIComputeContext* FVulkanDynamicRHI::RHIGetParallelCommandContext(FRHIParallelRenderPassInfo const& ParallelRenderPass, FRHIGPUMask GPUMask)
+IRHIComputeContext* FVulkanDynamicRHI::RHIGetParallelCommandContext(FRHIParallelRenderPassInfo const& ParallelRenderPass)
 {
 	FVulkanParallelRenderPassInfo* VulkanParallelRenderPassInfo = (FVulkanParallelRenderPassInfo*)ParallelRenderPass.RHIPlatformData;
 	checkf(VulkanParallelRenderPassInfo, TEXT("Must begin parallel render pass on parent context before creating parallel contexts!"));
diff --git a/Engine/Source/Runtime/Windows/D3D11RHI/Private/D3D11Commands.cpp b/Engine/Source/Runtime/Windows/D3D11RHI/Private/D3D11Commands.cpp
--- a/Engine/Source/Runtime/Windows/D3D11RHI/Private/D3D11Commands.cpp
+++ b/Engine/Source/Runtime/Windows/D3D11RHI/Private/D3D11Commands.cpp
@@ -1383,7 +1383,7 @@ IRHICommandContext* FD3D11DynamicRHI::RHIGetDefaultContext()
 	return this;
 }
 
-IRHIComputeContext* FD3D11DynamicRHI::RHIGetCommandContext(ERHIPipeline Pipeline, FRHIGPUMask GPUMask)
+IRHIComputeContext* FD3D11DynamicRHI::RHIGetCommandContext(ERHIPipeline Pipeline)
 {
 	UE_LOG(LogRHI, Fatal, TEXT("FD3D11DynamicRHI::RHIGetCommandContext should never be called. D3D11 RHI does not implement parallel command list execution."));
 	return nullptr;
diff --git a/Engine/Source/Runtime/Windows/D3D11RHI/Private/D3D11RHIPrivate.h b/Engine/Source/Runtime/Windows/D3D11RHI/Private/D3D11RHIPrivate.h
--- a/Engine/Source/Runtime/Windows/D3D11RHI/Private/D3D11RHIPrivate.h
+++ b/Engine/Source/Runtime/Windows/D3D11RHI/Private/D3D11RHIPrivate.h
@@ -509,7 +509,7 @@ public:
 	D3D11RHI_API virtual void* RHIGetNativeDevice() final override;
 	D3D11RHI_API virtual void* RHIGetNativeCommandBuffer() final override;
 	D3D11RHI_API virtual class IRHICommandContext* RHIGetDefaultContext() final override;
-	D3D11RHI_API virtual IRHIComputeContext* RHIGetCommandContext(ERHIPipeline Pipeline, FRHIGPUMask GPUMask) final override;
+	D3D11RHI_API virtual IRHIComputeContext* RHIGetCommandContext(ERHIPipeline Pipeline) final override;
 	D3D11RHI_API virtual void RHIFinalizeContext(FRHIFinalizeContextArgs&& Args, TRHIPipelineArray<IRHIPlatformCommandList*>& Output) final override;
 	D3D11RHI_API virtual void RHISubmitCommandLists(FRHISubmitCommandListsArgs&& Args) final override;
 
