[RHI] Add a feature to maintain a cpu copy of FUniformBuffer (disabled by default)

Problem: DumpGPU can only dump pass parameters but not parameters inside referenced structs. Thus, important parameters like the view uniforms containing all matrices are not dumped.

Solution: RHI is extended to maintain a cpu copy of uniform buffers. This feature must be enabled with r.RHI.UniformBufferContentMap.Enable. The functionality in DumpGPU is part of a separate CL.

#rnx

diff --git a/Engine/Source/Runtime/RHI/Private/RHIResources.cpp b/Engine/Source/Runtime/RHI/Private/RHIResources.cpp
--- a/Engine/Source/Runtime/RHI/Private/RHIResources.cpp
+++ b/Engine/Source/Runtime/RHI/Private/RHIResources.cpp
@@ -464,6 +464,67 @@ FRHIUniformBufferLayout::FRHIUniformBufferLayout(const FRHIUniformBufferLayoutIn
 {
 }
 
+FRHIUniformBuffer::~FRHIUniformBuffer()
+{
+#if !(UE_BUILD_SHIPPING || UE_BUILD_TEST)
+	UE::RHI::UniformBufferContentMap::Remove(this);
+#endif
+}
+
+static bool GRHIUniformBufferContentMapEnable = false;
+static FAutoConsoleVariableRef CVarRHIUniformBufferContentMapEnable(
+	TEXT("r.RHI.UniformBufferContentMap.Enable"),
+	GRHIUniformBufferContentMapEnable,
+	TEXT("When true, a copy of each FUniformBuffer content is maintained and updated on CPU inside FRHIUniformBufferContentMap."),
+	ECVF_ReadOnly);
+
+#if !(UE_BUILD_SHIPPING || UE_BUILD_TEST)
+namespace UE::RHI::UniformBufferContentMap
+{
+	static FCriticalSection CriticalSection;
+	static TMap<const FRHIUniformBuffer*, TArray<uint8>> UniformBufferContentMap;
+
+	void Update(const FRHIUniformBuffer* InUniformBuffer, const void* Contents)
+	{
+		if (GRHIUniformBufferContentMapEnable)
+		{
+			check(InUniformBuffer);
+			FScopeLock ScopeLock(&CriticalSection);
+			if (Contents)
+			{
+				UniformBufferContentMap.Emplace(InUniformBuffer, TArray<uint8>((const uint8*)Contents, InUniformBuffer->GetSize()));
+			}
+			else
+			{
+				UniformBufferContentMap.Emplace(InUniformBuffer, TArray<uint8>());
+			}
+		}
+	}
+
+	void Remove(const FRHIUniformBuffer* InUniformBuffer)
+	{
+		if (GRHIUniformBufferContentMapEnable)
+		{
+			check(InUniformBuffer);
+			FScopeLock ScopeLock(&CriticalSection);
+			UniformBufferContentMap.Remove(InUniformBuffer);
+		}
+	}
+
+	const void* Get(const FRHIUniformBuffer* InUniformBuffer)
+	{
+		if (GRHIUniformBufferContentMapEnable)
+		{
+			check(InUniformBuffer);
+			FScopeLock ScopeLock(&CriticalSection);
+			TArray<uint8>* Result = UniformBufferContentMap.Find(InUniformBuffer);
+			return (Result != nullptr) ? (void*)Result->GetData() : nullptr;
+		}
+		return nullptr;
+	}
+}
+#endif
+
 uint32 FRayTracingPipelineStateInitializer::GetMaxLocalBindingDataSize() const
 {
 	uint32 MaxLocalBindingDataSize = 0;
diff --git a/Engine/Source/Runtime/RHI/Public/DynamicRHI.h b/Engine/Source/Runtime/RHI/Public/DynamicRHI.h
--- a/Engine/Source/Runtime/RHI/Public/DynamicRHI.h
+++ b/Engine/Source/Runtime/RHI/Public/DynamicRHI.h
@@ -1221,7 +1221,11 @@ inline FUniformBufferLayoutRHIRef RHICreateUniformBufferLayout(const FRHIUniform
 inline FUniformBufferRHIRef RHICreateUniformBuffer(const void* Contents, const FRHIUniformBufferLayout* Layout, EUniformBufferUsage Usage, EUniformBufferValidation Validation = EUniformBufferValidation::ValidateResources)
 {
 	LLM_SCOPE_BYNAME(TEXT("RHIMisc/CreateUniformBuffer"));
-	return GDynamicRHI->RHICreateUniformBuffer(Contents, Layout, Usage, Validation);
+	FUniformBufferRHIRef Result = GDynamicRHI->RHICreateUniformBuffer(Contents, Layout, Usage, Validation);
+#if !(UE_BUILD_SHIPPING || UE_BUILD_TEST)
+	UE::RHI::UniformBufferContentMap::Update(Result.GetReference(), Contents);
+#endif
+	return Result;
 }
 
 inline FRHICalcTextureSizeResult RHICalcTexturePlatformSize(FRHITextureDesc const& Desc, uint32 FirstMipIndex = 0)
diff --git a/Engine/Source/Runtime/RHI/Public/RHICommandList.h b/Engine/Source/Runtime/RHI/Public/RHICommandList.h
--- a/Engine/Source/Runtime/RHI/Public/RHICommandList.h
+++ b/Engine/Source/Runtime/RHI/Public/RHICommandList.h
@@ -892,6 +892,9 @@ PRAGMA_ENABLE_DEPRECATION_WARNINGS
 	{
 		FRHICommandListScopedPipelineGuard ScopedPipeline(*this);
 		GDynamicRHI->RHIUpdateUniformBuffer(*this, UniformBufferRHI, Contents);
+#if !(UE_BUILD_SHIPPING || UE_BUILD_TEST)
+		UE::RHI::UniformBufferContentMap::Update(UniformBufferRHI, Contents);
+#endif
 	}
 
 	inline void UpdateStreamSourceSlot(FRHIStreamSourceSlot* StreamSourceSlotRHI, FRHIBuffer* BufferRHI)
diff --git a/Engine/Source/Runtime/RHI/Public/RHIResources.h b/Engine/Source/Runtime/RHI/Public/RHIResources.h
--- a/Engine/Source/Runtime/RHI/Public/RHIResources.h
+++ b/Engine/Source/Runtime/RHI/Public/RHIResources.h
@@ -1254,6 +1254,8 @@ public:
 protected:
 	TArray<TRefCountPtr<FRHIResource>> ResourceTable;
 
+	RHI_API virtual ~FRHIUniformBuffer();
+
 private:
 	/** Layout of the uniform buffer. */
 	TRefCountPtr<const FRHIUniformBufferLayout> Layout;
@@ -1261,6 +1263,15 @@ private:
 	uint32 LayoutConstantBufferSize;
 };
 
+#if !(UE_BUILD_SHIPPING || UE_BUILD_TEST)
+namespace UE::RHI::UniformBufferContentMap
+{
+	RHI_API void Update(const FRHIUniformBuffer* InUniformBuffer, const void* Contents);
+	RHI_API void Remove(const FRHIUniformBuffer* InUniformBuffer);
+	RHI_API const void* Get(const FRHIUniformBuffer* InUniformBuffer);
+}
+#endif
+
 class FRHIViewableResource : public FRHIResource
 {
 public:
