Breadcrumbs were missing from RHI Validation when RHI_VALIDATION_CHECK was used within the ValidationContext.

Added a new RHI_VALIDATION_CONTEXT_CHECK define which takes a context or tracker and uses the Executing Command Lists breadcrumbs instead of the ones stored in TLS.

diff --git a/Engine/Source/Runtime/RHI/Private/RHIValidation.cpp b/Engine/Source/Runtime/RHI/Private/RHIValidation.cpp
--- a/Engine/Source/Runtime/RHI/Private/RHIValidation.cpp
+++ b/Engine/Source/Runtime/RHI/Private/RHIValidation.cpp
@@ -332,7 +332,12 @@ namespace RHIValidation
 
 				if (!bFound)
 				{
-					check(UniqueIdentities.Num() < GRHIGlobals.MaxSimultaneousUAVs);
+					RHI_VALIDATION_CONTEXT_CHECK(UniqueIdentities.Num() < GRHIGlobals.MaxSimultaneousUAVs, 
+										BarrierTracker,
+										*FString::Printf(TEXT("Using too many unique UAV's, Max UAV count is %d, current index %d"),
+										GRHIGlobals.MaxSimultaneousUAVs,
+										UniqueIdentities.Num() - 1));
+					
 					UniqueIdentities.Add(Identity);
 
 					// Assert unique resources have the required state.
@@ -1228,7 +1233,22 @@ void FValidationRHI::RHIBindDebugLabelName(FRHICommandListBase& RHICmdList, FRHI
 	RHI->RHIBindDebugLabelName(RHICmdList, UnorderedAccessViewRHI, Name);
 }
 
-void FValidationRHI::ReportValidationFailure(const TCHAR* InMessage)
+void FValidationRHI::ReportValidationFailure(const TCHAR* InMessage, RHIValidation::FTracker* Tracker)
+{
+	ReportValidationFailure(InMessage, Tracker->GetExecutingCmdList());
+}
+
+void FValidationRHI::ReportValidationFailure(const TCHAR* InMessage, FValidationComputeContext* Context)
+{
+	ReportValidationFailure(InMessage, &Context->RHIContext->GetExecutingCommandList());
+}
+
+void FValidationRHI::ReportValidationFailure(const TCHAR* InMessage, FValidationContext* Context)
+{
+	ReportValidationFailure(InMessage, &Context->RHIContext->GetExecutingCommandList());
+}
+
+void FValidationRHI::ReportValidationFailure(const TCHAR* InMessage, FRHICommandListBase* CmdList)
 {
 	// Report failures only once per session, since many of them will happen repeatedly. This is similar to what ensure() does, but
 	// ensure() looks at the source location to determine if it's seen the error before. We want to look at the actual message, since
@@ -1247,16 +1267,28 @@ void FValidationRHI::ReportValidationFailure(const TCHAR* InMessage)
 	}
 
 	FString Message;
-	FString BreadcrumbPath = GetBreadcrumbPath();
-	if (!BreadcrumbPath.IsEmpty())
+	FString BreadcrumbPath;
+	
+	if (CmdList)
+	{
+		FRHIBreadcrumbNode* BreadcrumbNode = CmdList->GetCurrentBreadcrumbRef();
+		BreadcrumbPath = BreadcrumbNode ? BreadcrumbNode->GetFullPath() : TEXT("NoBreadcrumb");
+	}
+	else 
+	{
+		BreadcrumbPath = GetBreadcrumbPath();
+	}
+	
+	if(!BreadcrumbPath.IsEmpty())
 	{
 		Message = FString::Printf(
-			TEXT("%s")
-			TEXT("Breadcrumbs: %s\n")
-			TEXT("--------------------------------------------------------------------\n"),
-			InMessage, *BreadcrumbPath);
+					  TEXT("%s\n")
+					  TEXT("Breadcrumbs: %s\n")
+					  TEXT("--------------------------------------------------------------------\n"),
+					  InMessage, *BreadcrumbPath);
 	}
-	else
+	
+	if(Message.IsEmpty())
 	{
 		Message = InMessage;
 	}
@@ -1273,7 +1305,7 @@ void FValidationRHI::ReportValidationFailure(const TCHAR* InMessage)
 	}
 }
 
-static void ValidateBoundUniformBuffers(FRHIShader* Shader, const RHIValidation::FStaticUniformBuffers& StaticUniformBuffers, const RHIValidation::FStageBoundUniformBuffers& BoundUniformBuffers)
+static void ValidateBoundUniformBuffers(FRHIShader* Shader, const RHIValidation::FStaticUniformBuffers& StaticUniformBuffers, const RHIValidation::FStageBoundUniformBuffers& BoundUniformBuffers, FRHICommandListBase* InCmdList)
 {
 	const TCHAR* FreqName = GetShaderFrequencyString(Shader->GetFrequency(), false);
 	const TArray<uint32>& LayoutHashes = Shader->GetShaderResourceTable().ResourceTableLayoutHashes;
@@ -1281,7 +1313,7 @@ static void ValidateBoundUniformBuffers(FRHIShader* Shader, const RHIValidation:
 	const TArray<FUniformBufferStaticSlot>& StaticSlots = Shader->GetStaticSlots();
 	if (LayoutHashes.Num() != StaticSlots.Num())
 	{
-		RHI_VALIDATION_CHECK(false, *FString::Printf(TEXT("Shader %s(%s): The number of layout hashes (%d) is different from the number of static slots (%d)."), Shader->GetShaderName(), FreqName, LayoutHashes.Num(), StaticSlots.Num()));
+		RHI_VALIDATION_CONTEXT_CHECK(false, InCmdList, *FString::Printf(TEXT("Shader %s(%s): The number of layout hashes (%d) is different from the number of static slots (%d)."), Shader->GetShaderName(), FreqName, LayoutHashes.Num(), StaticSlots.Num()));
 		return;
 	}
 
@@ -1315,12 +1347,12 @@ static void ValidateBoundUniformBuffers(FRHIShader* Shader, const RHIValidation:
 		{
 			const FRHIUniformBufferLayout& Layout = BoundBuffer->GetLayout();
 			uint32 UniformBufferHash = Layout.GetHash();
-			RHI_VALIDATION_CHECK(UniformBufferHash == ExpectedLayoutHash, *FString::Printf(TEXT("Shader %s(%s): Invalid layout hash %u for uniform buffer \"%s\" at bind index %d (static: %s). Expecting a buffer called \"%s\", hash %u.)"),
+			RHI_VALIDATION_CONTEXT_CHECK(UniformBufferHash == ExpectedLayoutHash, InCmdList, *FString::Printf(TEXT("Shader %s(%s): Invalid layout hash %u for uniform buffer \"%s\" at bind index %d (static: %s). Expecting a buffer called \"%s\", hash %u.)"),
 				Shader->GetShaderName(), FreqName, UniformBufferHash, *Layout.GetDebugName(), BindIndex, bIsStatic ? TEXT("yes") : TEXT("no"), *Shader->GetUniformBufferName(BindIndex), ExpectedLayoutHash));
 		}
 		else
 		{
-			RHI_VALIDATION_CHECK(false, *FString::Printf(TEXT("Shader %s(%s): missing uniform buffer \"%s\" at index %d."),
+			RHI_VALIDATION_CONTEXT_CHECK(false, InCmdList, *FString::Printf(TEXT("Shader %s(%s): missing uniform buffer \"%s\" at index %d."),
 				Shader->GetShaderName(), FreqName , *Shader->GetUniformBufferName(BindIndex), BindIndex));
 		}
 	}
@@ -1337,11 +1369,11 @@ void FValidationComputeContext::ValidateDispatch()
 {
 	if (State.BoundShader == nullptr)
 	{
-		RHI_VALIDATION_CHECK(false, TEXT("A compute PSO has to be set before dispatching a compute shader."));
+		RHI_VALIDATION_CONTEXT_CHECK(false, this, TEXT("A compute PSO has to be set before dispatching a compute shader."));
 		return;
 	}
 
-	ValidateBoundUniformBuffers(State.BoundShader, State.StaticUniformBuffers, State.BoundUniformBuffers);
+	ValidateBoundUniformBuffers(State.BoundShader, State.StaticUniformBuffers, State.BoundUniformBuffers, &GetExecutingCommandList());
 }
 
 void FValidationComputeContext::FState::Reset()
@@ -1387,7 +1419,7 @@ void FValidationContext::RHICopyToStagingBuffer(FRHIBuffer* SourceBufferRHI, FRH
 	Tracker->Assert(SourceBufferRHI->GetWholeResourceIdentity(), ERHIAccess::CopySrc);
 	if (GRHIValidateBufferSourceCopy)
 	{
-		RHI_VALIDATION_CHECK(EnumHasAnyFlags(SourceBufferRHI->GetUsage(), BUF_SourceCopy), *GetReasonString_SourceCopyFlagMissing(SourceBufferRHI));
+		RHI_VALIDATION_CONTEXT_CHECK(EnumHasAnyFlags(SourceBufferRHI->GetUsage(), BUF_SourceCopy), this, *GetReasonString_SourceCopyFlagMissing(SourceBufferRHI));
 	}
 	RHIContext->RHICopyToStagingBuffer(SourceBufferRHI, DestinationStagingBufferRHI, InOffset, InNumBytes);
 }
@@ -1398,7 +1430,7 @@ void FValidationComputeContext::RHICopyToStagingBuffer(FRHIBuffer* SourceBufferR
 	Tracker->Assert(SourceBufferRHI->GetWholeResourceIdentity(), ERHIAccess::CopySrc);
 	if (GRHIValidateBufferSourceCopy)
 	{
-		RHI_VALIDATION_CHECK(EnumHasAnyFlags(SourceBufferRHI->GetUsage(), BUF_SourceCopy), *GetReasonString_SourceCopyFlagMissing(SourceBufferRHI));
+		RHI_VALIDATION_CONTEXT_CHECK(EnumHasAnyFlags(SourceBufferRHI->GetUsage(), BUF_SourceCopy), this, *GetReasonString_SourceCopyFlagMissing(SourceBufferRHI));
 	}
 	RHIContext->RHICopyToStagingBuffer(SourceBufferRHI, DestinationStagingBufferRHI, InOffset, InNumBytes);
 }
@@ -1407,18 +1439,18 @@ void FValidationContext::ValidateDispatch()
 {
 	if (State.BoundShaders[SF_Compute] == nullptr)
 	{
-		RHI_VALIDATION_CHECK(false, TEXT("A compute PSO has to be set before dispatching a compute shader."));
+		RHI_VALIDATION_CONTEXT_CHECK(false, this, TEXT("A compute PSO has to be set before dispatching a compute shader."));
 		return;
 	}
 
-	ValidateBoundUniformBuffers(State.BoundShaders[SF_Compute], State.StaticUniformBuffers, State.BoundUniformBuffers.Get(SF_Compute));
+	ValidateBoundUniformBuffers(State.BoundShaders[SF_Compute], State.StaticUniformBuffers, State.BoundUniformBuffers.Get(SF_Compute), &RHIContext->GetExecutingCommandList());
 }
 
 void FValidationContext::ValidateDrawing()
 {
 	if (!State.bGfxPSOSet)
 	{
-		RHI_VALIDATION_CHECK(false, TEXT("A graphics PSO has to be set in order to be able to draw!"));
+		RHI_VALIDATION_CONTEXT_CHECK(false, this, TEXT("A graphics PSO has to be set in order to be able to draw!"));
 		return;
 	}
 
@@ -1427,7 +1459,7 @@ void FValidationContext::ValidateDrawing()
 		EShaderFrequency Frequency = (EShaderFrequency)FrequencyIndex;
 		if (IsValidGraphicsFrequency(Frequency) && State.BoundShaders[Frequency])
 		{
-			ValidateBoundUniformBuffers(State.BoundShaders[Frequency], State.StaticUniformBuffers, State.BoundUniformBuffers.Get(Frequency));
+			ValidateBoundUniformBuffers(State.BoundShaders[Frequency], State.StaticUniformBuffers, State.BoundUniformBuffers.Get(Frequency), & GetExecutingCommandList());
 		}
 	}
 }
diff --git a/Engine/Source/Runtime/RHI/Public/RHIValidation.h b/Engine/Source/Runtime/RHI/Public/RHIValidation.h
--- a/Engine/Source/Runtime/RHI/Public/RHIValidation.h
+++ b/Engine/Source/Runtime/RHI/Public/RHIValidation.h
@@ -27,6 +27,18 @@ extern RHI_API bool GRHIValidateBufferSourceCopy;
 		} \
 	}while(0)
 
+#define RHI_VALIDATION_CONTEXT_CHECK(InExpression, Context, InMessage) \
+	do \
+	{ \
+		if(UNLIKELY(!(InExpression))) \
+		{ \
+			FValidationRHI::ReportValidationFailure(InMessage, Context); \
+		} \
+	}while(0)
+
+class FValidationComputeContext;
+class FValidationContext;
+
 class FValidationRHI : public FDynamicRHI
 {
 public:
@@ -1113,7 +1125,10 @@ public:
 	}
 
 //protected:
-	static void ReportValidationFailure(const TCHAR* InMessage);
+	static void ReportValidationFailure(const TCHAR* InMessage, RHIValidation::FTracker* Tracker);
+	static void ReportValidationFailure(const TCHAR* InMessage, FValidationComputeContext* Context);
+	static void ReportValidationFailure(const TCHAR* InMessage, FValidationContext* Context);
+	static void ReportValidationFailure(const TCHAR* InMessage, FRHICommandListBase* CmdList = nullptr);
 
 	FDynamicRHI*				RHI;
 	TMap<FRHIDepthStencilState*, FDepthStencilStateInitializerRHI> DepthStencilStates;
diff --git a/Engine/Source/Runtime/RHI/Public/RHIValidationCommon.h b/Engine/Source/Runtime/RHI/Public/RHIValidationCommon.h
--- a/Engine/Source/Runtime/RHI/Public/RHIValidationCommon.h
+++ b/Engine/Source/Runtime/RHI/Public/RHIValidationCommon.h
@@ -1076,6 +1076,16 @@ namespace RHIValidation
 				ResetUAVState(EUAVMode(Index));
 			}
 		}
+		
+		void SetExecutingCmdList(FRHICommandListBase* InCmdList)
+		{
+			ExecutingCmdList = InCmdList;
+		}
+		
+		FRHICommandListBase* GetExecutingCmdList()
+		{
+			return ExecutingCmdList;
+		}
 
 		static FOpQueueState& GetQueue(ERHIPipeline Pipeline);
 
@@ -1085,7 +1095,8 @@ namespace RHIValidation
 		const ERHIPipeline Pipeline;
 		TArray<FOperation> CurrentList;
 		FUAVTracker UAVTrackers[int32(EUAVMode::Num)];
-
+		FRHICommandListBase* ExecutingCmdList = nullptr;
+		
 		friend FOperation;
 		static RHI_API FOpQueueState OpQueues[int32(ERHIPipeline::Num)];
 	};
diff --git a/Engine/Source/Runtime/RHI/Public/RHIValidationContext.h b/Engine/Source/Runtime/RHI/Public/RHIValidationContext.h
--- a/Engine/Source/Runtime/RHI/Public/RHIValidationContext.h
+++ b/Engine/Source/Runtime/RHI/Public/RHIValidationContext.h
@@ -43,6 +43,8 @@ public:
 	{
 		IRHIComputeContext::SetExecutingCommandList(InCmdList);
 		RHIContext->SetExecutingCommandList(InCmdList);
+		
+		Tracker->SetExecutingCmdList(InCmdList);
 	}
 
 	virtual void RHISetComputePipelineState(FRHIComputePipelineState* ComputePipelineState) override final
@@ -177,10 +179,11 @@ public:
 	{
 		if (!GRHIGlobals.ShaderBundles.RequiresSharedBindlessParameters)
 		{
-			RHI_VALIDATION_CHECK(SharedBindlessParameters.Num() == 0, TEXT("SharedBindlessParameters should not be set on this platform and configuration"));
+			RHI_VALIDATION_CONTEXT_CHECK(SharedBindlessParameters.Num() == 0, this, TEXT("SharedBindlessParameters should not be set on this platform and configuration"));
 		}
-
-		RHI_VALIDATION_CHECK(Dispatches.Num() > 0, TEXT("A shader bundle must be dispatched with at least one record."));
+		
+		RHI_VALIDATION_CONTEXT_CHECK(Dispatches.Num() > 0, this, TEXT("A shader bundle must be dispatched with at least one record."));
+		
 		for (const FRHIShaderBundleComputeDispatch& Dispatch : Dispatches)
 		{
 			if (!Dispatch.IsValid())
@@ -218,12 +221,13 @@ public:
 	{
 		if (!GRHIGlobals.ShaderBundles.RequiresSharedBindlessParameters)
 		{
-			RHI_VALIDATION_CHECK(SharedBindlessParameters.Num() == 0, TEXT("SharedBindlessParameters should not be set on this platform and configuration"));
+			RHI_VALIDATION_CONTEXT_CHECK(SharedBindlessParameters.Num() == 0, this, TEXT("SharedBindlessParameters should not be set on this platform and configuration"));
 		}
 
 		// TODO:
 #if 0
-		RHI_VALIDATION_CHECK(Dispatches.Num() > 0, TEXT("A shader bundle must be dispatched with at least one record."));
+		RHI_VALIDATION_CONTEXT_CHECK(Dispatches.Num() > 0, this, TEXT("A shader bundle must be dispatched with at least one record."));
+		
 		for (const FRHIShaderBundleGraphicsDispatch& Dispatch : Dispatches)
 		{
 			if (!Dispatch.IsValid())
@@ -284,13 +288,13 @@ public:
 	{
 		if (State.BoundShader == nullptr)
 		{
-			RHI_VALIDATION_CHECK(false, TEXT("A compute PSO has to be set to set resources into a shader!"));
+			RHI_VALIDATION_CONTEXT_CHECK(false, this, TEXT("A compute PSO has to be set to set resources into a shader!"));
 			return;
 		}
 
 		if (Shader != State.BoundShader)
 		{
-			RHI_VALIDATION_CHECK(false, *FString::Printf(TEXT("Invalid attempt to set parameters for compute shader '%s' while the currently bound shader is '%s'"), Shader->GetShaderName(), State.BoundShader->GetShaderName()));
+			RHI_VALIDATION_CONTEXT_CHECK(false, this, *FString::Printf(TEXT("Invalid attempt to set parameters for compute shader '%s' while the currently bound shader is '%s'"), Shader->GetShaderName(), State.BoundShader->GetShaderName()));
 			return;
 		}
 
@@ -304,7 +308,7 @@ public:
 	{
 		if (State.BoundShader == nullptr)
 		{
-			RHI_VALIDATION_CHECK(false, TEXT("A compute PSO has to be set to set resources into a shader!"));
+			RHI_VALIDATION_CONTEXT_CHECK(false, this, TEXT("A compute PSO has to be set to set resources into a shader!"));
 			return;
 		}
 
@@ -486,6 +490,8 @@ public:
 	{
 		IRHICommandContext::SetExecutingCommandList(InCmdList);
 		RHIContext->SetExecutingCommandList(InCmdList);
+		
+		Tracker->SetExecutingCmdList(InCmdList);
 	}
 
 	virtual void RHISetComputePipelineState(FRHIComputePipelineState* ComputePipelineState) override final
@@ -561,10 +567,11 @@ public:
 	{
 		if (!GRHIGlobals.ShaderBundles.RequiresSharedBindlessParameters)
 		{
-			RHI_VALIDATION_CHECK(SharedBindlessParameters.Num() == 0, TEXT("SharedBindlessParameters should not be set on this platform and configuration"));
+			RHI_VALIDATION_CONTEXT_CHECK(SharedBindlessParameters.Num() == 0, this, TEXT("SharedBindlessParameters should not be set on this platform and configuration"));
 		}
 
-		RHI_VALIDATION_CHECK(Dispatches.Num() > 0, TEXT("A shader bundle must be dispatched with at least one record."));
+		RHI_VALIDATION_CONTEXT_CHECK(Dispatches.Num() > 0, this, TEXT("A shader bundle must be dispatched with at least one record."));
+		
 		RHIValidation::FStageBoundUniformBuffers& BoundUniformBuffers = State.BoundUniformBuffers.Get(SF_Compute);
 
 		for (const FRHIShaderBundleComputeDispatch& Dispatch : Dispatches)
@@ -604,12 +611,13 @@ public:
 	{
 		if (!GRHIGlobals.ShaderBundles.RequiresSharedBindlessParameters)
 		{
-			RHI_VALIDATION_CHECK(SharedBindlessParameters.Num() == 0, TEXT("SharedBindlessParameters should not be set on this platform and configuration"));
+			RHI_VALIDATION_CONTEXT_CHECK(SharedBindlessParameters.Num() == 0, this, TEXT("SharedBindlessParameters should not be set on this platform and configuration"));
 		}
 
 		// TODO
 #if 0
-		RHI_VALIDATION_CHECK(Dispatches.Num() > 0, TEXT("A shader bundle must be dispatched with at least one record."));
+		RHI_VALIDATION_CONTEXT_CHECK(Dispatches.Num() > 0, this, TEXT("SharedBindlessParameters should not be set on this platform and configuration"));
+
 		for (const FRHIShaderBundleGraphicsDispatch& Dispatch : Dispatches)
 		{
 			if (!Dispatch.IsValid())
@@ -850,7 +858,7 @@ public:
 	{
 		if (!State.bGfxPSOSet)
 		{
-			RHI_VALIDATION_CHECK(false, TEXT("A graphics PSO has to be set to set resources into a shader!"));
+			RHI_VALIDATION_CONTEXT_CHECK(false, this, TEXT("A graphics PSO has to be set to set resources into a shader!"));
 			return;
 		}
 
@@ -868,13 +876,13 @@ public:
 	{
 		if (State.BoundShaders[SF_Compute] == nullptr)
 		{
-			RHI_VALIDATION_CHECK(false, TEXT("A compute PSO has to be set to set resources into a shader!"));
+			RHI_VALIDATION_CONTEXT_CHECK(false, this, TEXT("A compute PSO has to be set to set resources into a shader!"));
 			return;
 		}
 
 		if (Shader != State.BoundShaders[SF_Compute])
 		{
-			RHI_VALIDATION_CHECK(false, *FString::Printf(TEXT("Invalid attempt to set parameters for compute shader '%s' while the currently bound shader is '%s'"), Shader->GetShaderName(), State.BoundShaders[SF_Compute]->GetShaderName()));
+			RHI_VALIDATION_CONTEXT_CHECK(false, this, *FString::Printf(TEXT("Invalid attempt to set parameters for compute shader '%s' while the currently bound shader is '%s'"), Shader->GetShaderName(), State.BoundShaders[SF_Compute]->GetShaderName()));
 			return;
 		}
 
@@ -890,7 +898,7 @@ public:
 	{
 		if (!State.bGfxPSOSet)
 		{
-			RHI_VALIDATION_CHECK(false, TEXT("A graphics PSO has to be set to set resources into a shader!"));
+			RHI_VALIDATION_CONTEXT_CHECK(false, this, TEXT("A graphics PSO has to be set to set resources into a shader!"));
 			return;
 		}
 
@@ -901,7 +909,7 @@ public:
 	{
 		if (State.BoundShaders[SF_Compute] == nullptr)
 		{
-			RHI_VALIDATION_CHECK(false, TEXT("A compute PSO has to be set to set resources into a shader!"));
+			RHI_VALIDATION_CONTEXT_CHECK(false, this, TEXT("A compute PSO has to be set to set resources into a shader!"));
 			return;
 		}
 
