PR #13754: Fix PSO precaching crash when referencing Breadcrumb data from RHI command list after memory is freed

diff --git a/Engine/Source/Runtime/RHI/Private/PipelineStateCache.cpp b/Engine/Source/Runtime/RHI/Private/PipelineStateCache.cpp
--- a/Engine/Source/Runtime/RHI/Private/PipelineStateCache.cpp
+++ b/Engine/Source/Runtime/RHI/Private/PipelineStateCache.cpp
@@ -3820,7 +3820,7 @@ FComputePipelineState* PipelineStateCache::GetAndOrCreateComputePipelineState(FR
 
 		FPSOCompilationDebugData PSOCompilationDebugData;
 #if WITH_RHI_BREADCRUMBS
-		if (DoAsyncCompile)
+		if (DoAsyncCompile && !bFromFileCache)
 		{
 			PSOCompilationDebugData.BreadcrumbRoot = FRHIBreadcrumbNode::GetNonNullRoot(RHICmdList.GetCurrentBreadcrumbRef());
 			PSOCompilationDebugData.BreadcrumbNode = PSOCompilationDebugData.BreadcrumbRoot
@@ -4352,7 +4352,7 @@ FGraphicsPipelineState* PipelineStateCache::GetAndOrCreateGraphicsPipelineState(
 
 		FPSOCompilationDebugData PSOCompilationDebugData;
 #if WITH_RHI_BREADCRUMBS
-		if (DoAsyncCompile)
+		if (DoAsyncCompile && !Initializer.bFromPSOFileCache)
 		{
 			PSOCompilationDebugData.BreadcrumbRoot = FRHIBreadcrumbNode::GetNonNullRoot(RHICmdList.GetCurrentBreadcrumbRef());
 			PSOCompilationDebugData.BreadcrumbNode = PSOCompilationDebugData.BreadcrumbRoot
