Remove RHIClearShaderBindingTable
- The only use was in FSceneRenderState::DestroyRayTracingScene() (LightmapRenderer.cpp), where the entire SBT resource was also getting released.
- The platform RHI implementations of the SBT type all release the same objects / memory in their destructors as they did in RHIClearShaderBindingTable, so the RHI command was unnecessary.
- The function also didn't belong on the RHI context, since its not an operation that's specific to a particular GPU queue.

diff --git a/Engine/Plugins/Experimental/GPULightmass/Source/GPULightmass/Private/LightmapRenderer.cpp b/Engine/Plugins/Experimental/GPULightmass/Source/GPULightmass/Private/LightmapRenderer.cpp
--- a/Engine/Plugins/Experimental/GPULightmass/Source/GPULightmass/Private/LightmapRenderer.cpp
+++ b/Engine/Plugins/Experimental/GPULightmass/Source/GPULightmass/Private/LightmapRenderer.cpp
@@ -1418,18 +1418,7 @@ bool FSceneRenderState::SetupRayTracingScene(FRDGBuilder& GraphBuilder, FSceneUn
 void FSceneRenderState::DestroyRayTracingScene()
 {
 	ReferenceView.Reset();
-
-#if RHI_RAYTRACING
-	// Is this needed?
-	if (IsRayTracingEnabled() && SBT.IsValid())
-	{
-		FRHICommandListImmediate& RHICmdList = FRHICommandListExecutor::GetImmediateCommandList();
-		SCOPED_GPU_MASK(RHICmdList, FRHIGPUMask::All());
-		RHICmdList.ClearShaderBindingTable(SBT);
-
-		SBT.SafeRelease();
-	}
-#endif
+	SBT.SafeRelease();
 }
 
 void FSceneRenderState::CalculateDistributionPrefixSumForAllLightmaps()
diff --git a/Engine/Source/Runtime/Apple/MetalRHI/Private/MetalRayTracing.cpp b/Engine/Source/Runtime/Apple/MetalRHI/Private/MetalRayTracing.cpp
--- a/Engine/Source/Runtime/Apple/MetalRHI/Private/MetalRayTracing.cpp
+++ b/Engine/Source/Runtime/Apple/MetalRHI/Private/MetalRayTracing.cpp
@@ -1177,11 +1177,6 @@ void FMetalRHICommandContext::RHIBindAccelerationStructureMemory(FRHIRayTracingS
 	MetalScene->BindBuffer(InBuffer, InBufferOffset);
 }
 
-void FMetalRHICommandContext::RHIClearShaderBindingTable(FRHIShaderBindingTable* SBT)
-{
-	// TODO:
-}
-
 void FMetalRHICommandContext::RHIRayTraceDispatch(FRHIRayTracingPipelineState* RayTracingPipelineState, FRHIRayTracingShader* RayGenShader,
 												  FRHIShaderBindingTable* SBT, const FRayTracingShaderBindings& GlobalResourceBindings,
 												  uint32 Width, uint32 Height)
diff --git a/Engine/Source/Runtime/Apple/MetalRHI/Public/MetalRHIContext.h b/Engine/Source/Runtime/Apple/MetalRHI/Public/MetalRHIContext.h
--- a/Engine/Source/Runtime/Apple/MetalRHI/Public/MetalRHIContext.h
+++ b/Engine/Source/Runtime/Apple/MetalRHI/Public/MetalRHIContext.h
@@ -160,7 +160,6 @@ public:
 	virtual void RHIBindAccelerationStructureMemory(FRHIRayTracingScene* Scene, FRHIBuffer* Buffer, uint32 BufferOffset) final override;
 	virtual void RHIBuildAccelerationStructures(TConstArrayView<FRayTracingGeometryBuildParams> Params, const FRHIBufferRange& ScratchBufferRange) final override;
 	virtual void RHIBuildAccelerationStructures(TConstArrayView<FRayTracingSceneBuildParams> Params) final override;
-	virtual void RHIClearShaderBindingTable(FRHIShaderBindingTable* SBT) final override;
 	virtual void RHIRayTraceDispatch(FRHIRayTracingPipelineState* RayTracingPipelineState, FRHIRayTracingShader* RayGenShader,
 									 FRHIShaderBindingTable* SBT, const FRayTracingShaderBindings& GlobalResourceBindings,
 									 uint32 Width, uint32 Height) final override;
diff --git a/Engine/Source/Runtime/D3D12RHI/Private/D3D12CommandContext.h b/Engine/Source/Runtime/D3D12RHI/Private/D3D12CommandContext.h
--- a/Engine/Source/Runtime/D3D12RHI/Private/D3D12CommandContext.h
+++ b/Engine/Source/Runtime/D3D12RHI/Private/D3D12CommandContext.h
@@ -738,7 +738,6 @@ public:
 	virtual void RHIExecuteMultiIndirectClusterOperation(const FRayTracingClusterOperationParams& Params) final override;
 	virtual void RHIBuildAccelerationStructures(TConstArrayView<FRayTracingGeometryBuildParams> Params, const FRHIBufferRange& ScratchBufferRange) final override;
 	virtual void RHIBuildAccelerationStructures(TConstArrayView<FRayTracingSceneBuildParams> Params) final override;
-	virtual void RHIClearShaderBindingTable(FRHIShaderBindingTable* SBT) final override;
 	virtual void RHICommitShaderBindingTable(FRHIShaderBindingTable* SBT, FRHIBuffer* InlineBindingDataBuffer) final override;
 
 	virtual void RHIRayTraceDispatch(FRHIRayTracingPipelineState* RayTracingPipelineState, FRHIRayTracingShader* RayGenShader,
@@ -1115,11 +1114,6 @@ public:
 		ContextRedirect(RHISetBindingsOnShaderBindingTable(SBT, Pipeline, NumBindings, Bindings, BindingType));
 	}
 
-	virtual void RHIClearShaderBindingTable(FRHIShaderBindingTable* SBT) final override
-	{
-		ContextRedirect(RHIClearShaderBindingTable(SBT));
-	}
-
 	virtual void RHIBindAccelerationStructureMemory(FRHIRayTracingScene* Scene, FRHIBuffer* Buffer, uint32 BufferOffset) final override
 	{
 		ContextRedirect(RHIBindAccelerationStructureMemory(Scene, Buffer, BufferOffset));
diff --git a/Engine/Source/Runtime/D3D12RHI/Private/D3D12RayTracing.cpp b/Engine/Source/Runtime/D3D12RHI/Private/D3D12RayTracing.cpp
--- a/Engine/Source/Runtime/D3D12RHI/Private/D3D12RayTracing.cpp
+++ b/Engine/Source/Runtime/D3D12RHI/Private/D3D12RayTracing.cpp
@@ -3135,14 +3135,6 @@ public:
 		const uint32 GPUIndex = Device->GetGPUIndex();
 		return ShaderTablesPerGPU[GPUIndex];
 	}
-
-	void ReleaseForDevice(FD3D12Device* Device)
-	{
-		const uint32 GPUIndex = Device->GetGPUIndex();
-
-		delete ShaderTablesPerGPU[GPUIndex];
-		ShaderTablesPerGPU[GPUIndex] = nullptr;
-	}
 	   
 	virtual FRHISizeAndStride GetInlineBindingDataSizeAndStride() const override final
 	{
@@ -4916,14 +4908,6 @@ void FD3D12CommandContext::RHICommitShaderBindingTable(FRHIShaderBindingTable* I
 	}
 }
 
-void FD3D12CommandContext::RHIClearShaderBindingTable(FRHIShaderBindingTable* InSBT)
-{
-	FD3D12RayTracingShaderBindingTable* SBT = FD3D12DynamicRHI::ResourceCast(InSBT);
-	check(SBT);
-
-	SBT->ReleaseForDevice(GetParentDevice());
-}
-
 static constexpr uint32 LooseParameterCBVIndex = 0; // Global uniform buffer is always assumed to be in slot 0
 
 struct FD3D12RayTracingGlobalResourceBinder
diff --git a/Engine/Source/Runtime/RHI/Private/RHIValidation.cpp b/Engine/Source/Runtime/RHI/Private/RHIValidation.cpp
--- a/Engine/Source/Runtime/RHI/Private/RHIValidation.cpp
+++ b/Engine/Source/Runtime/RHI/Private/RHIValidation.cpp
@@ -375,14 +375,6 @@ namespace RHIValidation
 		ShaderBindingMode = InInitializer.ShaderBindingMode;
 		HitGroupIndexingMode = InInitializer.HitGroupIndexingMode;
 	}
-
-	void FShaderBindingTable::Clear()
-	{
-		WorkerData[0].SRVs.Empty();
-		WorkerData[0].UAVs.Empty();
-
-		bIsDirty = true;
-	}
 	
 	static void CollectShaderBindingTableResources(FRayTracingPipelineState* RayTracingPipelineState, FShaderBindingTable* ShaderTable, const FRayTracingLocalShaderBindings& LocalShaderBinding, ERayTracingBindingType BindingType, uint32 WorkerIndex)
 	{
@@ -704,7 +696,7 @@ void FValidationRHI::RHIFinalizeContext(FRHIFinalizeContextArgs&& Args, TRHIPipe
 IRHIComputeContext* FValidationRHI::RHIGetParallelCommandContext(FRHIParallelRenderPassInfo const& ParallelRenderPass)
 {
 	// If a platform has a ChildWait or ParentWait it is expected that they will override RHIGetParallelCommandContext
-	// otherwise, we need to manually call RHIGetCommandContext and RHIBeginRenderPass seperately, because the default implementation
+	// otherwise, we need to manually call RHIGetCommandContext and RHIBeginRenderPass separately, because the default implementation
 	// calls one after another potentially crasing in GetHightestLevelContext (FRHICommandList_RecursiveHazardous)
 	// TODO: Remove this after implementing RHISetupParallelPass
 	if (GRHIParallelRHIExecuteChildWait || GRHIParallelRHIExecuteParentWait)
diff --git a/Engine/Source/Runtime/RHI/Public/RHICommandList.h b/Engine/Source/Runtime/RHI/Public/RHICommandList.h
--- a/Engine/Source/Runtime/RHI/Public/RHICommandList.h
+++ b/Engine/Source/Runtime/RHI/Public/RHICommandList.h
@@ -2592,17 +2592,6 @@ FRHICOMMAND_MACRO(FRHICommandCommitShaderBindingTable)
 	RHI_EXECUTE_API void Execute(FRHICommandListBase& CmdList);
 };
 
-FRHICOMMAND_MACRO(FRHICommandClearShaderBindingTable)
-{
-	FRHIShaderBindingTable* SBT;
-
-	explicit FRHICommandClearShaderBindingTable(FRHIShaderBindingTable* InSBT)
-		: SBT(InSBT)
-	{}
-
-	RHI_EXECUTE_API void Execute(FRHICommandListBase & CmdList);
-};
-
 FRHICOMMAND_UNNAMED(FRHICommandBuildAccelerationStructures)
 {
 	TConstArrayView<FRayTracingGeometryBuildParams> Params;
@@ -4383,20 +4372,9 @@ public:
 		}
 	}
 
+	UE_DEPRECATED(5.8, "ClearShaderBindingTable is deprecated and does nothing. Remove any remaining calls to this function.")
 	inline void ClearShaderBindingTable(FRHIShaderBindingTable* SBT)
 	{
-		if (Bypass())
-		{
-			GetContext().RHIClearShaderBindingTable(SBT);
-		}
-		else
-		{
-			ALLOC_COMMAND(FRHICommandClearShaderBindingTable)(SBT);
-
-			// This RHI command modifies members of the FRHIShaderBindingTable inside platform RHI implementations.
-			// It therefore needs the RHI lock fence to prevent races on those members.
-			RHIThreadFence(true);
-		}
 	}
 	
 	inline void SetBindingsOnShaderBindingTable(
diff --git a/Engine/Source/Runtime/RHI/Public/RHICommandListCommandExecutes.inl b/Engine/Source/Runtime/RHI/Public/RHICommandListCommandExecutes.inl
--- a/Engine/Source/Runtime/RHI/Public/RHICommandListCommandExecutes.inl
+++ b/Engine/Source/Runtime/RHI/Public/RHICommandListCommandExecutes.inl
@@ -507,12 +507,6 @@ void FRHICommandCommitShaderBindingTable::Execute(FRHICommandListBase& CmdList)
 	INTERNAL_DECORATOR(RHICommitShaderBindingTable)(SBT, InlineBindingDataBuffer);
 }
 
-void FRHICommandClearShaderBindingTable::Execute(FRHICommandListBase& CmdList)
-{
-	RHISTAT(ClearShaderBindingTable);
-	INTERNAL_DECORATOR(RHIClearShaderBindingTable)(SBT);
-}
-
 void FRHICommandBuildAccelerationStructures::Execute(FRHICommandListBase& CmdList)
 {
 	RHISTAT(BuildAccelerationStructure);
diff --git a/Engine/Source/Runtime/RHI/Public/RHIContext.h b/Engine/Source/Runtime/RHI/Public/RHIContext.h
--- a/Engine/Source/Runtime/RHI/Public/RHIContext.h
+++ b/Engine/Source/Runtime/RHI/Public/RHIContext.h
@@ -859,11 +859,6 @@ public:
 
 	virtual void RHICopyBufferRegion(FRHIBuffer* DestBuffer, uint64 DstOffset, FRHIBuffer* SourceBuffer, uint64 SrcOffset, uint64 NumBytes) = 0;
 
-	virtual void RHIClearShaderBindingTable(FRHIShaderBindingTable* SBT)
-	{
-		checkNoEntry();
-	}
-	
 	virtual void RHISetBindingsOnShaderBindingTable(FRHIShaderBindingTable* SBT, FRHIRayTracingPipelineState* Pipeline, uint32 NumBindings, const FRayTracingLocalShaderBindings* Bindings, ERayTracingBindingType BindingType)
 	{
 		checkNoEntry();
diff --git a/Engine/Source/Runtime/RHI/Public/RHIValidationCommon.h b/Engine/Source/Runtime/RHI/Public/RHIValidationCommon.h
--- a/Engine/Source/Runtime/RHI/Public/RHIValidationCommon.h
+++ b/Engine/Source/Runtime/RHI/Public/RHIValidationCommon.h
@@ -529,7 +529,6 @@ namespace RHIValidation
 
 		RHI_API FShaderBindingTable(const FRayTracingShaderBindingTableInitializer& InInitializer);
 
-		RHI_API void Clear();
 		RHI_API void SetBindingsOnShaderBindingTable(FRayTracingPipelineState* RayTracingPipelineState, uint32 NumBindings, const FRayTracingLocalShaderBindings* Bindings, ERayTracingBindingType BindingType);
 		RHI_API void Commit();
 
diff --git a/Engine/Source/Runtime/RHI/Public/RHIValidationContext.h b/Engine/Source/Runtime/RHI/Public/RHIValidationContext.h
--- a/Engine/Source/Runtime/RHI/Public/RHIValidationContext.h
+++ b/Engine/Source/Runtime/RHI/Public/RHIValidationContext.h
@@ -1221,12 +1221,6 @@ public:
 		RHIContext->RHICopyBufferRegion(DestBuffer, DstOffset, SourceBuffer, SrcOffset, NumBytes);
 	}
 
-	void RHIClearShaderBindingTable(FRHIShaderBindingTable* SBT)
-	{
-		SBT->Clear();
-		RHIContext->RHIClearShaderBindingTable(SBT);
-	}
-
 	void RHICommitShaderBindingTable(FRHIShaderBindingTable* SBT, FRHIBuffer* InlineBindingDataBuffer)
 	{
 		SBT->Commit();
diff --git a/Engine/Source/Runtime/VulkanRHI/Private/VulkanContext.h b/Engine/Source/Runtime/VulkanRHI/Private/VulkanContext.h
--- a/Engine/Source/Runtime/VulkanRHI/Private/VulkanContext.h
+++ b/Engine/Source/Runtime/VulkanRHI/Private/VulkanContext.h
@@ -341,7 +341,6 @@ public:
 	virtual void RHIBeginParallelRenderPass(TSharedPtr<FRHIParallelRenderPassInfo> InInfo, const TCHAR* InName) final override;
 	virtual void RHIEndParallelRenderPass() final override;
 
-	virtual void RHIClearShaderBindingTable(FRHIShaderBindingTable* SBT) final override;
 	virtual void RHICommitShaderBindingTable(FRHIShaderBindingTable* SBT, FRHIBuffer* InlineBindingDataBuffer) final override;
 	virtual void RHIBindAccelerationStructureMemory(FRHIRayTracingScene* Scene, FRHIBuffer* Buffer, uint32 BufferOffset) final override;
 	virtual void RHIBuildAccelerationStructures(TConstArrayView<FRayTracingGeometryBuildParams> Params, const FRHIBufferRange& ScratchBufferRange) final override;
diff --git a/Engine/Source/Runtime/VulkanRHI/Private/VulkanRayTracing.cpp b/Engine/Source/Runtime/VulkanRHI/Private/VulkanRayTracing.cpp
--- a/Engine/Source/Runtime/VulkanRHI/Private/VulkanRayTracing.cpp
+++ b/Engine/Source/Runtime/VulkanRHI/Private/VulkanRayTracing.cpp
@@ -1143,12 +1143,6 @@ FShaderBindingTableRHIRef FVulkanDynamicRHI::RHICreateShaderBindingTable(FRHICom
 	return new FVulkanRayTracingShaderTable(RHICmdList, *Device, Initializer);
 }
 
-void FVulkanCommandListContext::RHIClearShaderBindingTable(FRHIShaderBindingTable* InSBT)
-{
-	FVulkanRayTracingShaderTable* SBT = ResourceCast(InSBT);
-	SBT->ReleaseLocalBuffers();
-}
-
 void FVulkanCommandListContext::RHICommitShaderBindingTable(FRHIShaderBindingTable* InSBT, FRHIBuffer* InlineBindingDataBuffer)
 {
 	FVulkanRayTracingShaderTable* SBT = ResourceCast(InSBT);
