Deprecating Lock/Unlock on FRHIGPUMemoryReadback in favor of calling the derived type methods directly.
FRHIGPUMemoryReadback types should get the immediate command list manually (this will eventually be an argument)

diff --git a/Engine/Source/Runtime/RHI/Private/RHIGPUReadback.cpp b/Engine/Source/Runtime/RHI/Private/RHIGPUReadback.cpp
--- a/Engine/Source/Runtime/RHI/Private/RHIGPUReadback.cpp
+++ b/Engine/Source/Runtime/RHI/Private/RHIGPUReadback.cpp
@@ -63,8 +63,9 @@ void* FRHIGPUBufferReadback::Lock(uint32 NumBytes)
 
 	if (DestinationStagingBuffers[GPUIndex])
 	{
+		FRHICommandListImmediate& RHICmdList = FRHICommandListImmediate::Get();
 		LastLockGPUIndex = GPUIndex;
-		return RHILockStagingBuffer(DestinationStagingBuffers[GPUIndex], Fence.GetReference(), 0, NumBytes);
+		return RHICmdList.LockStagingBuffer(DestinationStagingBuffers[GPUIndex], Fence.GetReference(), 0, NumBytes);
 	}
 	else
 	{
@@ -75,7 +76,8 @@ void* FRHIGPUBufferReadback::Lock(uint32 NumBytes)
 void FRHIGPUBufferReadback::Unlock()
 {
 	ensure(DestinationStagingBuffers[LastLockGPUIndex]);
-	RHIUnlockStagingBuffer(DestinationStagingBuffers[LastLockGPUIndex]);
+	FRHICommandListImmediate& RHICmdList = FRHICommandListImmediate::Get();
+	RHICmdList.UnlockStagingBuffer(DestinationStagingBuffers[LastLockGPUIndex]);
 }
 
 uint64 FRHIGPUBufferReadback::GetGPUSizeBytes() const
@@ -206,38 +208,18 @@ void FRHIGPUTextureReadback::EnqueueCopy(FRHICommandList& RHICmdList, FRHITextur
 	}
 }
 
-void* FRHIGPUTextureReadback::Lock(uint32 NumBytes)
-{
-	uint32 GPUIndex = LastCopyGPUMask.GetFirstIndex();
-
-	if (DestinationStagingTextures[GPUIndex])
-	{
-		LastLockGPUIndex = GPUIndex;
-
-		void* ResultsBuffer = nullptr;
-		int32 BufferWidth = 0, BufferHeight = 0;
-		FRHICommandListImmediate& RHICmdList = FRHICommandListExecutor::GetImmediateCommandList();
-		RHICmdList.MapStagingSurface(DestinationStagingTextures[GPUIndex], Fence.GetReference(), ResultsBuffer, BufferWidth, BufferHeight, LastLockGPUIndex);
-
-		return ResultsBuffer;
-	}
-	else
-	{
-		return nullptr;
-	}
-}
-
 void* FRHIGPUTextureReadback::Lock(int32& OutRowPitchInPixels, int32 *OutBufferHeight)
 {
 	uint32 GPUIndex = LastCopyGPUMask.GetFirstIndex();
 
 	if (DestinationStagingTextures[GPUIndex])
 	{
+		FRHICommandListImmediate& RHICmdList = FRHICommandListImmediate::Get();
+
 		LastLockGPUIndex = GPUIndex;
 
 		void* ResultsBuffer = nullptr;
 		int32 BufferWidth = 0, BufferHeight = 0;
-		FRHICommandListImmediate& RHICmdList = FRHICommandListExecutor::GetImmediateCommandList();
 		RHICmdList.MapStagingSurface(DestinationStagingTextures[GPUIndex], Fence.GetReference(), ResultsBuffer, BufferWidth, BufferHeight, LastLockGPUIndex);
 
 		if (OutBufferHeight)
@@ -258,16 +240,11 @@ void* FRHIGPUTextureReadback::Lock(int32& OutRowPitchInPixels, int32 *OutBufferH
 	return nullptr;
 }
 
-void FRHIGPUTextureReadback::LockTexture(FRHICommandListImmediate& RHICmdList, void*& OutBufferPtr, int32& OutRowPitchInPixels)
-{
-	OutBufferPtr = Lock(OutRowPitchInPixels);
-}
-
 void FRHIGPUTextureReadback::Unlock()
 {
 	ensure(DestinationStagingTextures[LastLockGPUIndex]);
 
-	FRHICommandListImmediate& RHICmdList = FRHICommandListExecutor::GetImmediateCommandList();
+	FRHICommandListImmediate& RHICmdList = FRHICommandListImmediate::Get();
 	RHICmdList.UnmapStagingSurface(DestinationStagingTextures[LastLockGPUIndex], LastLockGPUIndex);
 }
 
diff --git a/Engine/Source/Runtime/RHI/Public/RHIGPUReadback.h b/Engine/Source/Runtime/RHI/Public/RHIGPUReadback.h
--- a/Engine/Source/Runtime/RHI/Public/RHIGPUReadback.h
+++ b/Engine/Source/Runtime/RHI/Public/RHIGPUReadback.h
@@ -11,6 +11,7 @@
 #include "RHI.h"
 #include "DynamicRHI.h"
 #include "RHICommandList.h"
+#include "MultiGPU.h"
 
 /**
  * FRHIGPUMemoryReadback: Represents a memory readback request scheduled with CopyToStagingBuffer
@@ -96,12 +97,19 @@ public:
 	 * @param NumBytes The maximum number of bytes the host will read from this pointer.
 	 * @returns A CPU accessible pointer to the backing buffer.
 	 */
-	virtual void* Lock(uint32 NumBytes) = 0;
+	UE_DEPRECATED(5.8, "FRHIGPUMemoryReadback::Lock is deprecated. Please use FRHIGPUBufferReadback::Lock or FRHIGPUTextureReadback::Lock directly.")
+	void* Lock(uint32 NumBytes)
+	{
+		return nullptr;
+	}
 
 	/**
 	 * Signals that the host is finished reading from the backing buffer.
 	 */
-	virtual void Unlock() = 0;
+	UE_DEPRECATED(5.8, "FRHIGPUMemoryReadback::Unlock is deprecated. Please use FRHIGPUBufferReadback::Unlock or FRHIGPUTextureReadback::Unlock directly.")
+	void Unlock()
+	{
+	}
 
 	inline const FRHIGPUMask& GetLastCopyGPUMask() const { return LastCopyGPUMask; }
 
@@ -128,18 +136,14 @@ public:
 	RHI_API FRHIGPUBufferReadback(FName RequestName);
 	 
 	RHI_API void EnqueueCopy(FRHICommandList& RHICmdList, FRHIBuffer* SourceBuffer, uint32 NumBytes = 0) override;
-	RHI_API void* Lock(uint32 NumBytes) override;
-	RHI_API void Unlock() override;
+	RHI_API void* Lock(uint32 NumBytes);
+	RHI_API void Unlock();
 	RHI_API uint64 GetGPUSizeBytes() const;
 
 private:
 
 	// RHI staging buffers are single GPU -- need to be branched when using multiple GPUs
-#if WITH_MGPU
 	FStagingBufferRHIRef DestinationStagingBuffers[MAX_NUM_GPUS];
-#else
-	FStagingBufferRHIRef DestinationStagingBuffers[1];
-#endif
 };
 
 
@@ -153,20 +157,10 @@ public:
 
 	RHI_API void EnqueueCopy(FRHICommandList& RHICmdList, FRHITexture* SourceTexture, const FIntVector& SourcePosition, uint32 SourceSlice, const FIntVector& Size) override;
 
-	UE_DEPRECATED(5.0, "Use FRHIGPUTextureReadback::Lock( int32& OutRowPitchInPixels) instead.")
-	RHI_API void* Lock(uint32 NumBytes) override;
-
 	RHI_API void* Lock(int32& OutRowPitchInPixels, int32* OutBufferHeight = nullptr);
-	RHI_API void Unlock() override;
-
-	UE_DEPRECATED(5.0, "Use FRHIGPUTextureReadback::Lock( int32& OutRowPitchInPixels) instead.")
-	RHI_API void LockTexture(FRHICommandListImmediate& RHICmdList, void*& OutBufferPtr, int32& OutRowPitchInPixels);
+	RHI_API void Unlock();
 
 	RHI_API uint64 GetGPUSizeBytes() const;
 
-#if WITH_MGPU
 	FTextureRHIRef DestinationStagingTextures[MAX_NUM_GPUS];
-#else
-	FTextureRHIRef DestinationStagingTextures[1];
-#endif
 };
