Deprecate the global scope RHIGetDefaultContext() function
- High level code should not retrieve RHI contexts directly. They are a private implementation detail of the platform RHI API
- The intended user of the GDynamicRHI->RHIGetDefaultContext() function is the RHI command list management code responsible for recording and replaying RHI commands on parallel translate chains.

diff --git a/Engine/Source/Runtime/RHI/Private/RHICommandList.cpp b/Engine/Source/Runtime/RHI/Private/RHICommandList.cpp
--- a/Engine/Source/Runtime/RHI/Private/RHICommandList.cpp
+++ b/Engine/Source/Runtime/RHI/Private/RHICommandList.cpp
@@ -424,7 +424,7 @@ void FRHICommandListBase::ActivatePipelines(ERHIPipeline Pipelines)
 						// Need to handle the "immediate" context separately.
 						Context = ExecutingCmdList.AllowParallelTranslate()
 							? GDynamicRHI->RHIGetCommandContext(Pipeline)
-							: ::RHIGetDefaultContext();
+							: GDynamicRHI->RHIGetDefaultContext();
 					}
 				}
 
diff --git a/Engine/Source/Runtime/RHI/Public/DynamicRHI.h b/Engine/Source/Runtime/RHI/Public/DynamicRHI.h
--- a/Engine/Source/Runtime/RHI/Public/DynamicRHI.h
+++ b/Engine/Source/Runtime/RHI/Public/DynamicRHI.h
@@ -1354,6 +1354,7 @@ inline void RHIGetSupportedResolution(uint32& Width, uint32& Height)
 	GDynamicRHI->RHIGetSupportedResolution(Width, Height);
 }
 
+UE_DEPRECATED(5.8, "RHIGetDefaultContext is deprecated. Do not call this function. RHI contexts are a private implementation detail of the platform RHI API, and should not be retrieved by high-level code directly via this function.")
 inline class IRHICommandContext* RHIGetDefaultContext()
 {
 	return GDynamicRHI->RHIGetDefaultContext();
