Fix for Texture LRU breaking with recent tset changes.,

#tests

diff --git a/Engine/Source/Runtime/RHI/Public/PsoLruCache.h b/Engine/Source/Runtime/RHI/Public/PsoLruCache.h
--- a/Engine/Source/Runtime/RHI/Public/PsoLruCache.h
+++ b/Engine/Source/Runtime/RHI/Public/PsoLruCache.h
@@ -630,7 +630,7 @@ private:
 private:
 
 	/** Set of entries for fast lookup. */
-	TSet<FCacheEntry*, FKeyFuncs> LookupSet;
+	TSparseSet<FCacheEntry*, FKeyFuncs> LookupSet;
 
 	/** Least recent item in the cache. */
 	FCacheEntry* LeastRecent;
