Re-implement "stat DrawCount" in the new RHI GPU profiler
- Removed code associated with tracking the current draw stat category across RHI command lists, by wrapping it in #if RHI_NEW_GPU_PROFILER == 0
- The categories are now attached to GPU stats, via RHI breadcrumbs, meaning we don't need any complex logic for tracking the current category. Instead, the profiler can track the active category as a stack while processing the stream of GPU profiler events.
- Begin/EndBreadcrumb events push and pop the category. Stats events accumulate the draw, dispatch, primitive, and vertex counts into the category at the top of the stack.
- Totals are computed when the profiler completes a frame, and emitted to the STATs system / CSV profiler.
- Counts will agree with "profileGPU" and Insights, since they are computed from the same event stream.

diff --git a/Engine/Source/Runtime/RHI/Private/GPUProfiler.cpp b/Engine/Source/Runtime/RHI/Private/GPUProfiler.cpp
--- a/Engine/Source/Runtime/RHI/Private/GPUProfiler.cpp
+++ b/Engine/Source/Runtime/RHI/Private/GPUProfiler.cpp
@@ -755,6 +755,14 @@ bool FGPUTiming::GAreGlobalsInitialized = false;
 
 #else
 
+CSV_DEFINE_CATEGORY(RHI, true);
+
+#if UE_BUILD_SHIPPING
+	CSV_DEFINE_CATEGORY(DrawCall, false);
+#else
+	CSV_DEFINE_CATEGORY(DrawCall, true);
+#endif
+
 namespace UE::RHI::GPUProfiler
 {
 	TLockFreePointerListUnordered<void, PLATFORM_CACHE_LINE_SIZE> FEventStream::FChunk::MemoryPool;
@@ -1424,7 +1432,7 @@ namespace UE::RHI::GPUProfiler
 	// Per queue GPU stats
 
 	// Total busy time on the current queue. StatName == "Unaccounted" is used by the Csv profiler
-	static FGPUStat GPUStat_Total(TEXT("Unaccounted"), TEXT("Queue Total"));
+	static FGPUStat GPUStat_Total(TEXT("Unaccounted"), TEXT("Queue Total"), nullptr);
 #endif
 
 #if STATS
@@ -1696,9 +1704,45 @@ namespace UE::RHI::GPUProfiler
 		#endif
 		};
 
-		struct FQueueTimestamps
+		struct FQueueDrawStats
+		{
+		#if HAS_GPU_STATS
+			// The +1 is for "uncategorised"
+			static constexpr int32 NumCategories = FRHIDrawStatsCategory::MAX_DRAWCALL_CATEGORY + 1;
+		#else
+			static constexpr int32 NumCategories = 1;
+		#endif
+
+			static constexpr int32 NoCategory = NumCategories - 1;
+			FEvent::FStats Categories[NumCategories];
+
+			void AddStats(FEvent::FStats const& Event, FRHIDrawStatsCategory const* Category)
+			{
+				uint32 CategoryIndex = Category ? Category->Index : NoCategory;
+				Categories[CategoryIndex] += Event;
+			}
+
+			FQueueDrawStats() = default;
+			FQueueDrawStats(FQueueDrawStats&& Other)
+			{
+				FMemory::Memcpy(Categories, Other.Categories, sizeof(Categories));
+				FMemory::Memzero(Other.Categories, sizeof(Categories));
+			}
+
+			FQueueDrawStats& operator += (FQueueDrawStats const& RHS)
+			{
+				for (int32 Index = 0; Index < NumCategories; ++Index)
+				{
+					Categories[Index] += RHS.Categories[Index];
+				}
+
+				return *this;
+			}
+		};
+
+		struct FQueueFrameState
 		{
-			FTimestampStream Queue;
+			FTimestampStream Timestamps;
 			FStatState WholeQueueStat;
 
 			uint64 CPUFrameBoundary = 0;
@@ -1706,6 +1750,8 @@ namespace UE::RHI::GPUProfiler
 			// Used to override the GPU time calculation for this queue, if an FFrameTime event is in the stream
 			TOptional<uint64> TotalBusyCycles;
 
+			FQueueDrawStats DrawStats;
+
 		#if WITH_RHI_BREADCRUMBS
 			TMap<FRHIBreadcrumbData_Stats, FStatState> Stats;
 		#endif
@@ -1723,7 +1769,7 @@ namespace UE::RHI::GPUProfiler
 			uint64 Value = 0;
 		};
 
-		struct FFrameState : TMap<FQueue, FQueueTimestamps>
+		struct FFrameState : TMap<FQueue, FQueueFrameState>
 		{
 		#if STATS
 			TOptional<int64> StatsFrame;
@@ -1745,18 +1791,18 @@ namespace UE::RHI::GPUProfiler
 			// The GPU timestamp of the last event processed.
 			uint64 LastGPUCycles = 0;
 
-			FQueueTimestamps Timestamps;
+			FQueueFrameState QueueFrameState;
 
 			bool bBusy = false;
-
 			bool bWasTraced = false;
 
 		#if WITH_RHI_BREADCRUMBS
 			TMap<FRHIBreadcrumbData_Stats, int32> ActiveStats;
 			TArray<FRHIBreadcrumbData_Stats> ActiveStatsStack;
-			FRHIBreadcrumbNodeRef Breadcrumb {};
 		#endif
 
+			TArray<FRHIDrawStatsCategory const*> DrawStatsStack { nullptr };
+
 		#if WITH_PROFILEGPU
 			struct
 			{
@@ -2059,7 +2105,7 @@ namespace UE::RHI::GPUProfiler
 				// Apply the timings to all active stats
 				for (auto const& [Stat, RefCount] : ActiveStats)
 				{
-					FStatState& State = Timestamps.Stats.FindChecked(Stat);
+					FStatState& State = QueueFrameState.Stats.FindChecked(Stat);
 					State.Inclusive.Accumulate(Busy, Wait, Idle);
 
 					if (ActiveStatsStack.Num() > 0 && ActiveStatsStack.Last() == Stat)
@@ -2071,10 +2117,10 @@ namespace UE::RHI::GPUProfiler
 				if (ActiveStatsStack.Num() == 0)
 			#endif
 				{
-					Timestamps.WholeQueueStat.Exclusive.Accumulate(Busy, Wait, Idle);
+					QueueFrameState.WholeQueueStat.Exclusive.Accumulate(Busy, Wait, Idle);
 				}
 
-				Timestamps.WholeQueueStat.Inclusive.Accumulate(Busy, Wait, Idle);
+				QueueFrameState.WholeQueueStat.Inclusive.Accumulate(Busy, Wait, Idle);
 
 			#if WITH_PROFILEGPU
 				for (FNode* Node = Profile.Current; Node; Node = Node->Parent)
@@ -2091,7 +2137,7 @@ namespace UE::RHI::GPUProfiler
 
 			void BeginWork(FEvent::FBeginWork const& Event)
 			{
-				Timestamps.Queue.AddTimestamp(Event.GPUTimestampTOP, true);
+				QueueFrameState.Timestamps.AddTimestamp(Event.GPUTimestampTOP, true);
 
 				uint64 Idle = Event.CPUTimestamp > LastGPUCycles
 					? Event.CPUTimestamp - LastGPUCycles
@@ -2106,7 +2152,7 @@ namespace UE::RHI::GPUProfiler
 
 			void EndWork(FEvent::FEndWork const& Event)
 			{
-				Timestamps.Queue.AddTimestamp(Event.GPUTimestampBOP, false);
+				QueueFrameState.Timestamps.AddTimestamp(Event.GPUTimestampBOP, false);
 
 				uint64 Busy = Event.GPUTimestampBOP > LastGPUCycles
 					? Event.GPUTimestampBOP - LastGPUCycles
@@ -2137,14 +2183,20 @@ namespace UE::RHI::GPUProfiler
 					int32 RefCount = ActiveStats.FindOrAdd(Stat)++;
 					if (RefCount == 0)
 					{
-						Timestamps.Stats.FindOrAdd(Stat);
+						QueueFrameState.Stats.FindOrAdd(Stat);
 					}
 
 					ActiveStatsStack.Add(Stat);
+
+				#if HAS_GPU_STATS
+					if (Stat.GPUStat->DrawStatsCategory)
+					{
+						DrawStatsStack.Add(Stat.GPUStat->DrawStatsCategory);
+					}
+				#endif
 				}
 
-				Breadcrumb = Event.Breadcrumb;
-				Breadcrumb->TraceBeginGPU(Queue.Value, Event.GPUTimestampTOP);
+				Event.Breadcrumb->TraceBeginGPU(Queue.Value, Event.GPUTimestampTOP);
 
 			#if WITH_PROFILEGPU
 				if (Profile.bProfileFrame)
@@ -2179,10 +2231,17 @@ namespace UE::RHI::GPUProfiler
 
 					check(ActiveStatsStack.Last() == Stat);
 					ActiveStatsStack.RemoveAt(ActiveStatsStack.Num() - 1, EAllowShrinking::No);
+
+				#if HAS_GPU_STATS
+					if (Stat.GPUStat->DrawStatsCategory)
+					{
+						check(DrawStatsStack.Last() == Stat.GPUStat->DrawStatsCategory);
+						DrawStatsStack.RemoveAt(DrawStatsStack.Num() - 1, EAllowShrinking::No);
+					}
+				#endif
 				}
-				
-				Breadcrumb->TraceEndGPU(Queue.Value, Event.GPUTimestampBOP);
-				Breadcrumb = Event.Breadcrumb->GetParent();
+
+				Event.Breadcrumb->TraceEndGPU(Queue.Value, Event.GPUTimestampBOP);
 
 			#if WITH_PROFILEGPU
 				if (Profile.bProfileFrame)
@@ -2206,6 +2265,9 @@ namespace UE::RHI::GPUProfiler
 					}
 				}
 			#endif
+				
+				QueueFrameState.DrawStats.AddStats(Event, DrawStatsStack.Last());
+
 				FGpuProfilerTrace::Stats(Queue.Value, Event.NumDraws, Event.NumPrimitives);
 			}
 
@@ -2246,13 +2308,13 @@ namespace UE::RHI::GPUProfiler
 
 			void FrameTime(uint64 TotalGPUTime)
 			{
-				Timestamps.TotalBusyCycles = TotalGPUTime;
+				QueueFrameState.TotalBusyCycles = TotalGPUTime;
 			}
 
 			void FrameBoundary(FEvent::FFrameBoundary const& Event, FFrameState& FrameState, uint32 InProfileFrameNumber)
 			{
 				check(!bBusy);
-				Timestamps.CPUFrameBoundary = Event.CPUTimestamp;
+				QueueFrameState.CPUFrameBoundary = Event.CPUTimestamp;
 
 				FGpuProfilerTrace::FrameBoundary(Queue.Value, Event.FrameNumber);
 
@@ -2264,14 +2326,14 @@ namespace UE::RHI::GPUProfiler
 				}
 			#endif
 
-				FrameState.Emplace(Queue, MoveTemp(Timestamps));
+				FrameState.Emplace(Queue, MoveTemp(QueueFrameState));
 
 			#if WITH_RHI_BREADCRUMBS
 				// Reinsert timestamp streams for the current active stats on 
 				// this queue, since these got moved into the frame state.
 				for (auto& [Stat, RefCount] : ActiveStats)
 				{
-					Timestamps.Stats.FindOrAdd(Stat);
+					QueueFrameState.Stats.FindOrAdd(Stat);
 				}
 			#endif
 
@@ -2482,9 +2544,9 @@ namespace UE::RHI::GPUProfiler
 							// Trim the Signals array in each queue, up to the lowest frame boundary CPU timestamp.
 							{
 								uint64 MinFrameBoundary = TNumericLimits<uint64>::Max();
-								for (auto& [Queue, QueueTimestamps] : FrameState)
+								for (auto& [Queue, QueueFrameState] : FrameState)
 								{
-									MinFrameBoundary = FMath::Min(MinFrameBoundary, QueueTimestamps.CPUFrameBoundary);
+									MinFrameBoundary = FMath::Min(MinFrameBoundary, QueueFrameState.CPUFrameBoundary);
 								}
 
 								for (auto& [Queue, LocalQueueState] : QueueStates)
@@ -2530,17 +2592,37 @@ namespace UE::RHI::GPUProfiler
 			const bool bCsvStatsEnabled = !!CVarGPUCsvStatsEnabled.GetValueOnAnyThread();
 			FCsvProfiler* CsvProfiler = FCsvProfiler::Get();
 			CsvProfiler->BeginFrameEOP();
-		#else
-			const bool bCsvStatsEnabled = false;
 		#endif
 
+		#if HAS_GPU_STATS
+			// Only copy the display counters every half second keep things more stable.
+			constexpr float TimeoutSeconds = 0.5;
+
+			static double LastTime = 0.0;
+			double CurrentTime = FPlatformTime::Seconds();
+
+			bool bCopyDisplayFrames = false;
+			if (CurrentTime - LastTime > TimeoutSeconds)
+			{
+				LastTime = CurrentTime;
+				bCopyDisplayFrames = true;
+			}
+
+			FRHIDrawStatsCategory::FManager& Manager = FRHIDrawStatsCategory::GetManager();
+		#endif
+
+			// Summed stats across all GPUs
+			FEvent::FStats Total = {};
+			TStaticArray<FEvent::FStats, FQueueDrawStats::NumCategories> TotalPerCategory { InPlace };
+			TStaticArray<FEvent::FStats, MAX_NUM_GPUS> TotalPerGPU { InPlace };
+
 			TOptional<uint64> MaxQueueBusyCycles;
 
-			for (auto const& [Queue, QueueTimestamps] : FrameState)
+			for (auto const& [Queue, QueueFrameState] : FrameState)
 			{
 			#if WITH_RHI_BREADCRUMBS && HAS_GPU_STATS
 				// Compute the individual GPU stats
-				for (auto const& [Stat, StatState] : QueueTimestamps.Stats)
+				for (auto const& [Stat, StatState] : QueueFrameState.Stats)
 				{
 					StatState.EmitResults(Queue, *Stat.GPUStat
 					#if STATS
@@ -2555,7 +2637,7 @@ namespace UE::RHI::GPUProfiler
 
 				// Set the whole-frame per queue stat
 			#if HAS_GPU_STATS
-				QueueTimestamps.WholeQueueStat.EmitResults(Queue, GPUStat_Total
+				QueueFrameState.WholeQueueStat.EmitResults(Queue, GPUStat_Total
 				#if STATS
 					, Stats
 				#endif
@@ -2565,13 +2647,58 @@ namespace UE::RHI::GPUProfiler
 				);
 			#endif
 
-				if (QueueTimestamps.TotalBusyCycles.IsSet())
+				if (QueueFrameState.TotalBusyCycles.IsSet())
 				{
 					uint64 CurrentMax = MaxQueueBusyCycles ? *MaxQueueBusyCycles : 0;
-					MaxQueueBusyCycles = FMath::Max(CurrentMax, *QueueTimestamps.TotalBusyCycles);
+					MaxQueueBusyCycles = FMath::Max(CurrentMax, *QueueFrameState.TotalBusyCycles);
 				}
+
+				for (int32 CategoryIndex = 0; CategoryIndex < FQueueDrawStats::NumCategories; ++CategoryIndex)
+				{
+					FEvent::FStats const& Category = QueueFrameState.DrawStats.Categories[CategoryIndex];
+					TotalPerCategory[CategoryIndex] += Category;
+					TotalPerGPU[Queue.GPU]          += Category;
+					Total                           += Category;
+
+				#if HAS_GPU_STATS
+					if (bCopyDisplayFrames && CategoryIndex < Manager.NumCategory)
+					{
+						Manager.DisplayCounts[CategoryIndex][Queue.GPU] = Category.NumDraws;
+					}
+				#endif // HAS_GPU_STATS
+				}				
 			}
 
+			for (uint32 GPUIndex = 0; GPUIndex < GNumExplicitGPUsForRendering; ++GPUIndex)
+			{
+				GNumDrawCallsRHI      [GPUIndex] = TotalPerGPU[GPUIndex].NumDraws;
+				GNumPrimitivesDrawnRHI[GPUIndex] = TotalPerGPU[GPUIndex].NumPrimitives;
+			}
+
+		#if CSV_PROFILER_STATS
+			// Multi-GPU support : CSV stats do not support MGPU yet. We're summing the totals across all GPUs here.
+			CsvProfiler->RecordEndOfPipeCustomStat("DrawCalls"      , CSV_CATEGORY_INDEX(RHI), Total.NumDraws     , ECsvCustomStatOp::Set);
+			CsvProfiler->RecordEndOfPipeCustomStat("PrimitivesDrawn", CSV_CATEGORY_INDEX(RHI), Total.NumPrimitives, ECsvCustomStatOp::Set);
+		#endif
+
+	#if HAS_GPU_STATS
+
+		#if STATS
+			Stats->AddMessage(GET_STATFNAME(STAT_RHIDraws     ), EStatOperation::Set, uint64(Total.NumDraws     ));
+			Stats->AddMessage(GET_STATFNAME(STAT_RHIDispatches), EStatOperation::Set, uint64(Total.NumDispatches));
+			Stats->AddMessage(GET_STATFNAME(STAT_RHIPrimitives), EStatOperation::Set, uint64(Total.NumPrimitives));
+			Stats->AddMessage(GET_STATFNAME(STAT_RHIVertices  ), EStatOperation::Set, uint64(Total.NumVertices  ));
+		#endif
+
+		#if CSV_PROFILER_STATS
+			for (int32 CategoryIndex = 0; CategoryIndex < Manager.NumCategory; ++CategoryIndex)
+			{
+				FCsvProfiler::RecordCustomStat(Manager.Array[CategoryIndex]->Name, CSV_CATEGORY_INDEX(DrawCall), int32(TotalPerCategory[CategoryIndex].NumDraws), ECsvCustomStatOp::Set);
+			}
+		#endif
+
+	#endif // HAS_GPU_STATS
+
 			if (MaxQueueBusyCycles.IsSet())
 			{
 				// Set the total GPU time stat according to the value directly provided by the platform RHI
@@ -2581,9 +2708,9 @@ namespace UE::RHI::GPUProfiler
 			{
 				// Compute the whole-frame total GPU time.
 				TArray<FTimestampStream::FState, TInlineAllocator<GetRHIPipelineCount() * MAX_NUM_GPUS>> StreamPointers;
-				for (auto const& [Queue, State] : FrameState)
+				for (auto const& [Queue, QueueFrameState] : FrameState)
 				{
-					StreamPointers.Emplace(State.Queue);
+					StreamPointers.Emplace(QueueFrameState.Timestamps);
 				}
 				uint64 WholeFrameUnion = FTimestampStream::ComputeUnion(StreamPointers);
 
@@ -2591,9 +2718,6 @@ namespace UE::RHI::GPUProfiler
 				GRHIGPUFrameTimeHistory.PushFrameCycles(1.0 / FPlatformTime::GetSecondsPerCycle64(), WholeFrameUnion);
 			}
 
-			// @todo set global csv GPU time
-			//RHISetGPUStatTotals(bCsvStatsEnabled, FPlatformTime::ToMilliseconds64(WholeFrameUnion));
-
 		#if STATS
 			Stats->Flush();
 		#endif
diff --git a/Engine/Source/Runtime/RHI/Private/RHI.cpp b/Engine/Source/Runtime/RHI/Private/RHI.cpp
--- a/Engine/Source/Runtime/RHI/Private/RHI.cpp
+++ b/Engine/Source/Runtime/RHI/Private/RHI.cpp
@@ -36,12 +36,17 @@ IMPLEMENT_MODULE(FDefaultModuleImpl, RHI);
 
 /** RHI Logging. */
 DEFINE_LOG_CATEGORY(LogRHI);
-CSV_DEFINE_CATEGORY(RHI, true);
 
-#if UE_BUILD_SHIPPING
-CSV_DEFINE_CATEGORY(DrawCall, false);
-#else
-CSV_DEFINE_CATEGORY(DrawCall, true);
+#if (RHI_NEW_GPU_PROFILER == 0)
+
+	CSV_DEFINE_CATEGORY(RHI, true);
+
+	#if UE_BUILD_SHIPPING
+	CSV_DEFINE_CATEGORY(DrawCall, false);
+	#else
+	CSV_DEFINE_CATEGORY(DrawCall, true);
+	#endif
+
 #endif
 
 IMPLEMENT_TYPE_LAYOUT(FRHIUniformBufferLayoutInitializer);
@@ -136,10 +141,12 @@ const FClearValueBinding FClearValueBinding::DefaultNormal8Bit(FLinearColor(128.
 
 #if HAS_GPU_STATS
 
+#if (RHI_NEW_GPU_PROFILER == 0)
 	FRHIDrawStatsCategory::FRHIDrawStatsCategory()
 		: Name(NAME_None)
 		, Index(-1)
 	{}
+#endif
 
 	FRHIDrawStatsCategory::FRHIDrawStatsCategory(FName InName)
 		: Name(InName)
@@ -1244,6 +1251,8 @@ ERHIAccess GRHIMultiPipelineMergeableAccessMask = ERHIAccess::ReadOnlyExclusiveM
 
 ERHIAccess GRHIMultiSubresourceDiscardIntermediateAccess = ERHIAccess::SRVCompute;
 
+#if (RHI_NEW_GPU_PROFILER == 0)
+
 void FRHIDrawStats::Accumulate(FRHIDrawStats& Other)
 {
 	for (uint32 GPUIndex = 0; GPUIndex < GNumExplicitGPUsForRendering; ++GPUIndex)
@@ -1328,6 +1337,8 @@ RHI_API void FRHIDrawStats::ProcessAsFrameStats()
 	Reset();
 }
 
+#endif // (RHI_NEW_GPU_PROFILER == 0)
+
 //
 // The current shader platform.
 //
diff --git a/Engine/Source/Runtime/RHI/Private/RHICommandList.cpp b/Engine/Source/Runtime/RHI/Private/RHICommandList.cpp
--- a/Engine/Source/Runtime/RHI/Private/RHICommandList.cpp
+++ b/Engine/Source/Runtime/RHI/Private/RHICommandList.cpp
@@ -287,7 +287,7 @@ void FRHICommandListBase::FinishRecording()
 	DispatchEvent->DispatchSubsequents();
 }
 
-#if HAS_GPU_STATS
+#if (RHI_NEW_GPU_PROFILER == 0) && HAS_GPU_STATS
 TOptional<FRHIDrawStatsCategory const*> FRHICommandListBase::SetDrawStatsCategory(TOptional<FRHIDrawStatsCategory const*> Category)
 {
 	check(!Category.IsSet() || ((*Category) == nullptr || (*Category)->ShouldCountDraws()));
@@ -978,7 +978,7 @@ void FRHICommandListExecutor::FSubmitState::Dispatch(FRHICommandListBase* CmdLis
 		}
 #endif // WITH_RHI_BREADCRUMBS
 
-#if HAS_GPU_STATS
+#if (RHI_NEW_GPU_PROFILER == 0) && HAS_GPU_STATS
 		CmdList->InitialDrawStatsCategory = GRHICommandList.CurrentDrawStatsCategory;
 
 		if (CmdList->PersistentState.CurrentDrawStatsCategory.IsSet())
@@ -1120,7 +1120,10 @@ void FRHICommandListExecutor::FTranslateState::Translate(FRHICommandListBase* Cm
 #if WITH_RHI_BREADCRUMBS
 	BreadcrumbAllocatorRefs.Append(MoveTemp(CmdList->BreadcrumbAllocatorRefs));
 #endif
+
+#if (RHI_NEW_GPU_PROFILER == 0)
 	DrawStats.Accumulate(CmdList->DrawStats);
+#endif
 
 	if (CmdList->ParallelRenderPassBegin && !bShouldFinalize)
 	{
@@ -1234,7 +1237,9 @@ void FRHICommandListExecutor::FSubmitState::Submit(const FSubmitArgs& Args)
 			}
 		}
 
+#if (RHI_NEW_GPU_PROFILER == 0)
 		GRHICommandList.FrameDrawStats.Accumulate(Job->DrawStats);
+#endif
 	}
 
 	if (GDynamicRHI)
@@ -1287,7 +1292,9 @@ void FRHICommandListExecutor::FSubmitState::Submit(const FSubmitArgs& Args)
 		};
 
 		GDynamicRHI->RHIEndFrame(EndFrameArgs);
+#if (RHI_NEW_GPU_PROFILER == 0)
 		GRHICommandList.FrameDrawStats.ProcessAsFrameStats();
+#endif
 	}
 
 	CompletionEvent->DispatchSubsequents();
@@ -1551,7 +1558,7 @@ void FRHICommandListImmediate::InitializeImmediateContexts()
 		}
 #endif
 
-#if HAS_GPU_STATS
+#if (RHI_NEW_GPU_PROFILER == 0) && HAS_GPU_STATS
 		InitialDrawStatsCategory = GRHICommandList.CurrentDrawStatsCategory;
 #endif
 	}
diff --git a/Engine/Source/Runtime/RHI/Private/RHIStats.cpp b/Engine/Source/Runtime/RHI/Private/RHIStats.cpp
--- a/Engine/Source/Runtime/RHI/Private/RHIStats.cpp
+++ b/Engine/Source/Runtime/RHI/Private/RHIStats.cpp
@@ -11,11 +11,20 @@ int32 GNumPrimitivesDrawnRHI[MAX_NUM_GPUS] = {};
 
 // Define counter stats.
 #if HAS_GPU_STATS
+
+#if RHI_NEW_GPU_PROFILER
+	DEFINE_STAT(STAT_RHIDraws);
+	DEFINE_STAT(STAT_RHIDispatches);
+	DEFINE_STAT(STAT_RHIPrimitives);
+	DEFINE_STAT(STAT_RHIVertices);
+#else
 	DEFINE_STAT(STAT_RHIDrawPrimitiveCalls);
 	DEFINE_STAT(STAT_RHITriangles);
 	DEFINE_STAT(STAT_RHILines);
 #endif
 
+#endif
+
 // Define memory stats.
 DEFINE_STAT(STAT_RenderTargetMemory2D);
 DEFINE_STAT(STAT_RenderTargetMemory3D);
diff --git a/Engine/Source/Runtime/RHI/Public/GPUProfiler.h b/Engine/Source/Runtime/RHI/Public/GPUProfiler.h
--- a/Engine/Source/Runtime/RHI/Public/GPUProfiler.h
+++ b/Engine/Source/Runtime/RHI/Public/GPUProfiler.h
@@ -194,10 +194,19 @@ namespace UE::RHI::GPUProfiler
 
 		struct FStats
 		{
-			uint32 NumDraws;
-			uint32 NumDispatches;
-			uint32 NumPrimitives;
-			uint32 NumVertices;
+			uint32 NumDraws      = 0;
+			uint32 NumDispatches = 0;
+			uint32 NumPrimitives = 0;
+			uint32 NumVertices   = 0;
+
+			FStats& operator += (FStats const& RHS)
+			{
+				NumDraws      += RHS.NumDraws;
+				NumDispatches += RHS.NumDispatches;
+				NumPrimitives += RHS.NumPrimitives;
+				NumVertices   += RHS.NumVertices;
+				return *this;
+			}
 
 			operator bool() const
 			{
@@ -531,6 +540,7 @@ namespace UE::RHI::GPUProfiler
 
 		TCHAR const* const StatName;
 		TCHAR const* const DisplayName;
+		FRHIDrawStatsCategory const* const DrawStatsCategory;
 
 	#if CSV_PROFILER_STATS
 		TOptional<FCsvDeclaredStat> CsvStat;
@@ -549,7 +559,6 @@ namespace UE::RHI::GPUProfiler
 			FStatCategory(FQueue Queue);
 
 			static TMap<FQueue, TUniquePtr<FStatCategory>> Categories;
-			static FStatCategory& GetCategory(FQueue Queue);
 		};
 
 		struct FStatInstance
@@ -571,9 +580,10 @@ namespace UE::RHI::GPUProfiler
 	#endif
 
 	public:
-		FGPUStat(TCHAR const* StatName, TCHAR const* DisplayName)
-			: StatName   (StatName)
-			, DisplayName(DisplayName)
+		FGPUStat(TCHAR const* StatName, TCHAR const* DisplayName, FRHIDrawStatsCategory const* DrawStatsCategory)
+			: StatName         (StatName)
+			, DisplayName      (DisplayName)
+			, DrawStatsCategory(DrawStatsCategory)
 		{}
 
 	#if STATS
@@ -584,8 +594,8 @@ namespace UE::RHI::GPUProfiler
 	template <typename TNameProvider>
 	struct TGPUStat : public FGPUStat
 	{
-		TGPUStat()
-			: FGPUStat(TNameProvider::GetStatName(), TNameProvider::GetDisplayName())
+		TGPUStat(FRHIDrawStatsCategory const* DrawStatsCategory = nullptr)
+			: FGPUStat(TNameProvider::GetStatName(), TNameProvider::GetDisplayName(), DrawStatsCategory)
 		{}
 	};
 
@@ -593,7 +603,12 @@ namespace UE::RHI::GPUProfiler
 	struct TGPUStatWithDrawcallCategory : public TGPUStat<TNameProvider>
 	{
 	#if HAS_GPU_STATS
-		FRHIDrawStatsCategory DrawcallCategory;
+		FRHIDrawStatsCategory DrawStatsCategory;
+
+		TGPUStatWithDrawcallCategory()
+			: TGPUStat<TNameProvider>(&DrawStatsCategory)
+			, DrawStatsCategory(FName(TNameProvider::GetStatName()))
+		{}
 	#endif
 	};
 
diff --git a/Engine/Source/Runtime/RHI/Public/RHICommandList.h b/Engine/Source/Runtime/RHI/Public/RHICommandList.h
--- a/Engine/Source/Runtime/RHI/Public/RHICommandList.h
+++ b/Engine/Source/Runtime/RHI/Public/RHICommandList.h
@@ -738,7 +738,7 @@ public:
 	bool IsInsideRenderPass    () const { return PersistentState.bInsideRenderPass;  }
 	bool IsInsideComputePass   () const { return PersistentState.bInsideComputePass; }
 
-#if HAS_GPU_STATS
+#if (RHI_NEW_GPU_PROFILER == 0) && HAS_GPU_STATS
 	RHI_API TOptional<FRHIDrawStatsCategory const*> SetDrawStatsCategory(TOptional<FRHIDrawStatsCategory const*> Category);
 #endif
 
@@ -1476,7 +1476,7 @@ protected:
 	} ActivatePipelineCommands {};
 #endif
 
-#if HAS_GPU_STATS
+#if (RHI_NEW_GPU_PROFILER == 0) && HAS_GPU_STATS
 	TOptional<FRHIDrawStatsCategory const*> InitialDrawStatsCategory {};
 #endif
 
@@ -1512,7 +1512,7 @@ protected:
 		FRHIBreadcrumbNode* LocalBreadcrumb = FRHIBreadcrumbNode::Sentinel;
 #endif
 
-#if HAS_GPU_STATS
+#if (RHI_NEW_GPU_PROFILER == 0) && HAS_GPU_STATS
 		TOptional<FRHIDrawStatsCategory const*> CurrentDrawStatsCategory {};
 #endif
 
@@ -1533,7 +1533,9 @@ protected:
 
 	} PersistentState;
 
+#if (RHI_NEW_GPU_PROFILER == 0)
 	FRHIDrawStats DrawStats {};
+#endif
 
 	TArray<FRHIBuffer*> PendingBufferUploads;
 	TArray<FRHITexture*> PendingTextureUploads;
@@ -1551,6 +1553,7 @@ public:
 	RHI_API void AttachBreadcrumbSubTree(FRHIBreadcrumbAllocator& Allocator, FRHIBreadcrumbList& Nodes);
 #endif
 
+#if (RHI_NEW_GPU_PROFILER == 0)
 	void Stats_AddDraw()
 	{
 #if HAS_GPU_STATS
@@ -1568,6 +1571,7 @@ public:
 		DrawStats.AddDrawAndPrimitives(PersistentState.CurrentGPUMask, nullptr, PrimitiveType, NumPrimitives);
 #endif
 	}
+#endif
 
 	TStaticArray<void*, MAX_NUM_GPUS>& GetQueryBatchData(ERenderQueryType QueryType)
 	{
@@ -4633,7 +4637,7 @@ class FRHICommandListImmediate : public FRHICommandList
 		PersistentState.LocalBreadcrumb = nullptr;
 #endif
 
-#if HAS_GPU_STATS
+#if (RHI_NEW_GPU_PROFILER == 0) && HAS_GPU_STATS
 		PersistentState.CurrentDrawStatsCategory = nullptr;
 #endif
 	}
@@ -5438,7 +5442,9 @@ private:
 		bool bUsingSubCmdLists = false;
 		bool bShouldFinalize = true;
 
+#if (RHI_NEW_GPU_PROFILER == 0)
 		FRHIDrawStats DrawStats{};
+#endif
 
 		FTaskPipe* GetTranslateTaskPipe(ENamedThreads::Type& NamedThread);
 		FTaskPipe* EnqueueTranslateTask(FGraphEventArray&& Prereqs, TFunction<void()>&& Lambda);
@@ -5464,7 +5470,9 @@ private:
 		bool bEmitBreadcrumbs = false;
 #endif
 
+#if (RHI_NEW_GPU_PROFILER == 0)
 		FRHIDrawStats DrawStats {};
+#endif
 
 		ERHISubmitFlags SubmitFlags = ERHISubmitFlags::None;
 		TArray<FRHIResource*> ResourcesToDelete {};
@@ -5515,10 +5523,12 @@ private:
 	
 #endif
 
+#if (RHI_NEW_GPU_PROFILER == 0)
 #if HAS_GPU_STATS
 	FRHIDrawStatsCategory const* CurrentDrawStatsCategory = nullptr;
 #endif
 	FRHIDrawStats FrameDrawStats;
+#endif
 
 	// Counts the number of calls to RHIEndFrame, and is used in GPU profiler frame boundary events.
 	uint32 FrameNumber = 0;
diff --git a/Engine/Source/Runtime/RHI/Public/RHIContext.h b/Engine/Source/Runtime/RHI/Public/RHIContext.h
--- a/Engine/Source/Runtime/RHI/Public/RHIContext.h
+++ b/Engine/Source/Runtime/RHI/Public/RHIContext.h
@@ -38,7 +38,9 @@ enum class EAsyncComputeBudget;
 
 struct FRHIBufferRange;
 struct FRHIPerCategoryDrawStats;
+#if (RHI_NEW_GPU_PROFILER == 0)
 struct FRHIDrawStats;
+#endif
 struct FRHICopyTextureInfo;
 
 
diff --git a/Engine/Source/Runtime/RHI/Public/RHIStats.h b/Engine/Source/Runtime/RHI/Public/RHIStats.h
--- a/Engine/Source/Runtime/RHI/Public/RHIStats.h
+++ b/Engine/Source/Runtime/RHI/Public/RHIStats.h
@@ -74,10 +74,12 @@ extern RHI_API int32 GNumPrimitivesDrawnRHI[MAX_NUM_GPUS];
 
 struct FRHIDrawStatsCategory
 {
-	RHI_API FRHIDrawStatsCategory();
 	RHI_API FRHIDrawStatsCategory(FName InName);
 
+#if (RHI_NEW_GPU_PROFILER == 0)
+	RHI_API FRHIDrawStatsCategory();
 	bool ShouldCountDraws() const { return Index != -1; }
+#endif
 
 	FName  const Name;
 	uint32 const Index;
@@ -99,10 +101,17 @@ struct FRHIDrawStatsCategory
 	RHI_API static FManager& GetManager();
 };
 
-// RHI counter stats.
-DECLARE_DWORD_COUNTER_STAT_EXTERN(TEXT("DrawPrimitive calls"), STAT_RHIDrawPrimitiveCalls, STATGROUP_RHI, RHI_API);
-DECLARE_DWORD_COUNTER_STAT_EXTERN(TEXT("Triangles drawn"), STAT_RHITriangles, STATGROUP_RHI, RHI_API);
-DECLARE_DWORD_COUNTER_STAT_EXTERN(TEXT("Lines drawn"), STAT_RHILines, STATGROUP_RHI, RHI_API);
+#if RHI_NEW_GPU_PROFILER
+	DECLARE_DWORD_COUNTER_STAT_EXTERN(TEXT("Draws"     ), STAT_RHIDraws     , STATGROUP_RHI, RHI_API);
+	DECLARE_DWORD_COUNTER_STAT_EXTERN(TEXT("Dispatches"), STAT_RHIDispatches, STATGROUP_RHI, RHI_API);
+	DECLARE_DWORD_COUNTER_STAT_EXTERN(TEXT("Primitives"), STAT_RHIPrimitives, STATGROUP_RHI, RHI_API);
+	DECLARE_DWORD_COUNTER_STAT_EXTERN(TEXT("Vertices"  ), STAT_RHIVertices  , STATGROUP_RHI, RHI_API);
+#else
+	// RHI counter stats.
+	DECLARE_DWORD_COUNTER_STAT_EXTERN(TEXT("DrawPrimitive calls"), STAT_RHIDrawPrimitiveCalls, STATGROUP_RHI, RHI_API);
+	DECLARE_DWORD_COUNTER_STAT_EXTERN(TEXT("Triangles drawn"), STAT_RHITriangles, STATGROUP_RHI, RHI_API);
+	DECLARE_DWORD_COUNTER_STAT_EXTERN(TEXT("Lines drawn"), STAT_RHILines, STATGROUP_RHI, RHI_API);
+#endif
 
 #else
 
@@ -122,116 +131,116 @@ struct FRHIDrawStatsCategory
 		StatEvent.NumDraws++;                                              \
 		StatEvent.NumPrimitives += Prims * FMath::Max(1u, Instances);      \
 		StatEvent.NumVertices   += Verts * FMath::Max(1u, Instances);      \
-		GetExecutingCommandList().Stats_AddDrawAndPrimitives(Type, Prims); \
 	} while (false)
 
-	#define RHI_DRAW_CALL_INC()                             do { StatEvent.NumDraws++; GetExecutingCommandList().Stats_AddDraw(); } while (false)
+	#define RHI_DRAW_CALL_INC()                             do { StatEvent.NumDraws++; } while (false)
 	#define RHI_DISPATCH_CALL_INC()			                do { StatEvent.NumDispatches++; } while (false)
 #else
 	#define RHI_DRAW_CALL_INC()                             do { GetExecutingCommandList().Stats_AddDraw(); } while (false)
 	#define RHI_DRAW_CALL_STATS(Type,Verts,Prims,Instances) do { GetExecutingCommandList().Stats_AddDrawAndPrimitives(Type, Prims); } while (false)
 	#define RHI_DISPATCH_CALL_INC()			                do { } while (false)
-#endif
-
-struct FRHIDrawStats
-{
-#if HAS_GPU_STATS
-	// The +1 is for "uncategorised"
-	static constexpr int32 NumCategories = FRHIDrawStatsCategory::MAX_DRAWCALL_CATEGORY + 1;
-#else
-	static constexpr int32 NumCategories = 1;
-#endif
 
-	static constexpr int32 NoCategory = NumCategories - 1;
-
-	struct FPerCategory
+	struct FRHIDrawStats
 	{
-		uint32 Draws;
-		uint32 Triangles;
-		uint32 Lines;
-		uint32 Quads;
-		uint32 Points;
-		uint32 Rectangles;
-
-		uint32 GetTotalPrimitives() const
-		{
-			return Triangles
-				+ Lines
-				+ Quads
-				+ Points
-				+ Rectangles;
-		}
+	#if HAS_GPU_STATS
+		// The +1 is for "uncategorised"
+		static constexpr int32 NumCategories = FRHIDrawStatsCategory::MAX_DRAWCALL_CATEGORY + 1;
+	#else
+		static constexpr int32 NumCategories = 1;
+	#endif
 
-		FPerCategory& operator += (FPerCategory const& RHS)
+		static constexpr int32 NoCategory = NumCategories - 1;
+
+		struct FPerCategory
 		{
-			Draws      += RHS.Draws;
-			Triangles  += RHS.Triangles;
-			Lines      += RHS.Lines;
-			Quads      += RHS.Quads;
-			Points     += RHS.Points;
-			Rectangles += RHS.Rectangles;
-			return *this;
-		}
-	};
+			uint32 Draws;
+			uint32 Triangles;
+			uint32 Lines;
+			uint32 Quads;
+			uint32 Points;
+			uint32 Rectangles;
+
+			uint32 GetTotalPrimitives() const
+			{
+				return Triangles
+					+ Lines
+					+ Quads
+					+ Points
+					+ Rectangles;
+			}
 
-	struct FPerGPU
-	{
-		friend struct FRHIDrawStats;
-		FPerCategory Categories[NumCategories];
-	};
+			FPerCategory& operator += (FPerCategory const& RHS)
+			{
+				Draws      += RHS.Draws;
+				Triangles  += RHS.Triangles;
+				Lines      += RHS.Lines;
+				Quads      += RHS.Quads;
+				Points     += RHS.Points;
+				Rectangles += RHS.Rectangles;
+				return *this;
+			}
+		};
 
-	FPerGPU& GetGPU(uint32 GPUIndex)
-	{
-		checkSlow(GPUIndex < UE_ARRAY_COUNT(GPUs));
-		return GPUs[GPUIndex];
-	}
+		struct FPerGPU
+		{
+			friend struct FRHIDrawStats;
+			FPerCategory Categories[NumCategories];
+		};
 
-	FRHIDrawStats()
-	{
-		Reset();
-	}
+		FPerGPU& GetGPU(uint32 GPUIndex)
+		{
+			checkSlow(GPUIndex < UE_ARRAY_COUNT(GPUs));
+			return GPUs[GPUIndex];
+		}
 
-	void Reset()
-	{
-		FMemory::Memzero(*this);
-	}
+		FRHIDrawStats()
+		{
+			Reset();
+		}
 
-	void AddDraw(FRHIGPUMask GPUMask, FRHIDrawStatsCategory const* Category)
-	{
-		uint32 CategoryIndex = Category ? Category->Index : NoCategory;
-		for (uint32 GPUIndex : GPUMask)
+		void Reset()
 		{
-			FPerCategory& Stats = GPUs[GPUIndex].Categories[CategoryIndex];
-			Stats.Draws++;
+			FMemory::Memzero(*this);
 		}
-	}
 
-	void AddDrawAndPrimitives(FRHIGPUMask GPUMask, FRHIDrawStatsCategory const* Category, EPrimitiveType PrimitiveType, uint32 NumPrimitives)
-	{
-		uint32 CategoryIndex = Category ? Category->Index : NoCategory;
-		for (uint32 GPUIndex : GPUMask)
+		void AddDraw(FRHIGPUMask GPUMask, FRHIDrawStatsCategory const* Category)
 		{
-			FPerCategory& Stats = GPUs[GPUIndex].Categories[CategoryIndex];
-			Stats.Draws++;
+			uint32 CategoryIndex = Category ? Category->Index : NoCategory;
+			for (uint32 GPUIndex : GPUMask)
+			{
+				FPerCategory& Stats = GPUs[GPUIndex].Categories[CategoryIndex];
+				Stats.Draws++;
+			}
+		}
 
-			switch (PrimitiveType)
+		void AddDrawAndPrimitives(FRHIGPUMask GPUMask, FRHIDrawStatsCategory const* Category, EPrimitiveType PrimitiveType, uint32 NumPrimitives)
+		{
+			uint32 CategoryIndex = Category ? Category->Index : NoCategory;
+			for (uint32 GPUIndex : GPUMask)
 			{
-			case PT_TriangleList : Stats.Triangles  += NumPrimitives; break;
-			case PT_TriangleStrip: Stats.Triangles  += NumPrimitives; break;
-			case PT_LineList     : Stats.Lines      += NumPrimitives; break;
-			case PT_QuadList     : Stats.Quads      += NumPrimitives; break;
-			case PT_PointList    : Stats.Points     += NumPrimitives; break;
-			case PT_RectList     : Stats.Rectangles += NumPrimitives; break;
+				FPerCategory& Stats = GPUs[GPUIndex].Categories[CategoryIndex];
+				Stats.Draws++;
+
+				switch (PrimitiveType)
+				{
+				case PT_TriangleList : Stats.Triangles  += NumPrimitives; break;
+				case PT_TriangleStrip: Stats.Triangles  += NumPrimitives; break;
+				case PT_LineList     : Stats.Lines      += NumPrimitives; break;
+				case PT_QuadList     : Stats.Quads      += NumPrimitives; break;
+				case PT_PointList    : Stats.Points     += NumPrimitives; break;
+				case PT_RectList     : Stats.Rectangles += NumPrimitives; break;
+				}
 			}
 		}
-	}
 
-	RHI_API void Accumulate(FRHIDrawStats& RHS);
-	RHI_API void ProcessAsFrameStats();
+		RHI_API void Accumulate(FRHIDrawStats& RHS);
+		RHI_API void ProcessAsFrameStats();
 
-private:
-	FPerGPU GPUs[MAX_NUM_GPUS];
-};
+	private:
+		FPerGPU GPUs[MAX_NUM_GPUS];
+	};
+
+#endif
 
 // RHI memory stats.
 DECLARE_MEMORY_STAT_POOL_EXTERN(TEXT("Render Target 2D Memory"), STAT_RenderTargetMemory2D, STATGROUP_RHI, FPlatformMemory::MCR_GPU, RHI_API);
