Simplify GPU profiler event stream management
- Removed FEventSink. The original GPU profiler design allowed for multiple consumers of RHI GPU profiler events. We've ended up with a single sink that processes everything, so there is no need to have multiple sinks.
- Removing the sinks also simplifies queue management, as we no longer need to hold TSharedRefs to the FEventStreams, or use a TSpscQueue.
- Removed the FEventSink::FIterator, and added Peek()/Pop() private functions directly on the FEventStream. These are for private use by the profiler, to consume events from the stream. Incoming streams from the platform RHIs are now appended into the existing per-queue stream.

diff --git a/Engine/Source/Runtime/RHI/Private/GPUProfiler.cpp b/Engine/Source/Runtime/RHI/Private/GPUProfiler.cpp
--- a/Engine/Source/Runtime/RHI/Private/GPUProfiler.cpp
+++ b/Engine/Source/Runtime/RHI/Private/GPUProfiler.cpp
@@ -767,51 +767,37 @@ namespace UE::RHI::GPUProfiler
 {
 	TLockFreePointerListUnordered<void, PLATFORM_CACHE_LINE_SIZE> FEventStream::FChunk::MemoryPool;
 
-	static TArray<FEventSink*>& GetSinks()
+	FEvent const* FEventStream::Peek() const
 	{
-		static TArray<FEventSink*> Sinks;
-		return Sinks;
-	}
+		if (!First)
+		{
+			return nullptr;
+		}
 
-	FEventSink::FEventSink()
-	{
-		GetSinks().Add(this);
+		return First->GetElement(First->Header.Iter);
 	}
 
-	FEventSink::~FEventSink()
+	void FEventStream::Pop()
 	{
-		GetSinks().RemoveSingle(this);
-	}
+		++First->Header.Iter;
 
-	void ProcessEvents(TArrayView<FEventStream> EventStreams)
-	{
-		TRACE_CPUPROFILER_EVENT_SCOPE(UE::RHI::GPUProfiler::ProcessEvents);
-
-		TArray<TSharedRef<FEventStream>> SharedStreams;
-		SharedStreams.Reserve(EventStreams.Num());
-
-		for (FEventStream& Stream : EventStreams)
+		if (First->Header.Iter == First->Header.Num)
 		{
-			if (!Stream.IsEmpty())
+			FChunk* Next = First->Header.Next;
+			if (!Next)
 			{
-				SharedStreams.Emplace(MakeShared<FEventStream>(MoveTemp(Stream)));
+				check(First == Current);
+				Current = nullptr;
 			}
-		}
 
-		if (SharedStreams.Num())
-		{
-			for (FEventSink* Sink : GetSinks())
-			{
-				Sink->ProcessStreams(SharedStreams);
-			}
-		}
-	}
+			delete First;
+			First = Next;
 
-	void InitializeQueues(TConstArrayView<FQueue> Queues)
-	{
-		for (FEventSink* Sink : GetSinks())
+			check(!First || (First->Header.Iter == 0 && First->Header.Num > 0));
+		}
+		else
 		{
-			Sink->InitializeQueues(Queues);
+			check(First->Header.Iter < First->Header.Num);
 		}
 	}
 
@@ -1522,8 +1508,7 @@ namespace UE::RHI::GPUProfiler
 
 #endif
 
-	// Handles computing the "stat unit" GPU time, "stat gpu" stats, and "profilegpu".
-	struct FGPUProfilerSink_StatSystem final : public FEventSink
+	struct FGPUProfiler
 	{
 		class FTimestampStream
 		{
@@ -1779,7 +1764,7 @@ namespace UE::RHI::GPUProfiler
 		struct FQueueState
 		{
 			FQueue const Queue;
-			TSpscQueue<FEventSink::FIterator> PendingStreams;
+			FEventStream EventStream;
 
 			// Array of fence signal history. Events are kept until all queues have processed events
 			// later than the CPU timestamps of these signals. The old events are then trimmed.
@@ -2047,6 +2032,7 @@ namespace UE::RHI::GPUProfiler
 
 			FQueueState(FQueue const& Queue)
 				: Queue(Queue)
+				, EventStream(Queue)
 			{}
 
 			void ResolveSignal(FEvent::FSignalFence const& Event)
@@ -2420,7 +2406,7 @@ namespace UE::RHI::GPUProfiler
 			}
 		}
 
-		void InitializeQueues(TConstArrayView<FQueue> Queues) override
+		void InitializeQueues(TConstArrayView<FQueue> Queues)
 		{
 			FGpuProfilerTrace::Initialize();
 
@@ -2434,15 +2420,17 @@ namespace UE::RHI::GPUProfiler
 			}
 		}
 
-		bool ProcessQueue(FQueueState& QueueState, FIterator& Iterator)
+		bool ProcessQueue(FQueueState& QueueState)
 		{
+			bool bProgress = false;
+
 			if (FGpuProfilerTrace::IsAvailable() && !QueueState.bWasTraced)
 			{
 				FGpuProfilerTrace::InitializeQueue(QueueState.Queue.Value, QueueState.Queue.GetTypeString());
 				QueueState.bWasTraced = true;
 			}
 
-			while (FEvent const* Event = Iterator.Peek())
+			while (FEvent const* Event = QueueState.EventStream.Peek())
 			{
 				switch (Event->GetType())
 				{
@@ -2503,7 +2491,7 @@ namespace UE::RHI::GPUProfiler
 						if (!ResolvedWait.IsSet())
 						{
 							// Unresolved fence, pause processing
-							return false;
+							return bProgress;
 						}
 
 						QueueState.Wait(*ResolvedWait, WaitFence);
@@ -2572,10 +2560,11 @@ namespace UE::RHI::GPUProfiler
 					break;
 				}
 
-				Iterator.Pop();
+				QueueState.EventStream.Pop();
+				bProgress = true;
 			}
 
-			return true;
+			return bProgress;
 		}
 
 		void ProcessFrame(FFrameState& FrameState)
@@ -2725,50 +2714,44 @@ namespace UE::RHI::GPUProfiler
 
 		void ProcessAllQueues()
 		{
-			// Process the queue as far as possible
+			// Process the queues as far as possible
 			bool bProgress;
 			do
 			{
 				bProgress = false;
 				for (auto& [Queue, QueueState] : QueueStates)
 				{
-					while (FIterator* Iterator = QueueState->PendingStreams.Peek())
-					{
-						FEvent const* Start = Iterator->Peek();
-
-						bool bPaused = !ProcessQueue(*QueueState.Get(), *Iterator);
-
-						FEvent const* End = Iterator->Peek();
-						bProgress |= End != Start;
-
-						if (bPaused)
-						{
-							// The queue was paused by a Wait event
-							check(End);
-							break;
-						}
-
-						if (!End)
-						{
-							// This stream has been fully processed.
-							QueueState->PendingStreams.Dequeue();
-						}
-					}
+					bProgress |= ProcessQueue(*QueueState.Get());
 				}
 			} while (bProgress);
 		}
 
-		void ProcessStreams(TConstArrayView<TSharedRef<FEventStream>> EventStreams) override
+		void ProcessEvents(TArrayView<FEventStream> EventStreams)
 		{
-			for (TSharedRef<FEventStream> const& Stream : EventStreams)
+			for (FEventStream& Stream : EventStreams)
 			{
-				FQueueState& QueueState = *QueueStates.FindChecked(Stream->Queue);
-				QueueState.PendingStreams.Enqueue(FIterator(Stream));
+				if (!Stream.IsEmpty())
+				{
+					FQueueState& QueueState = *QueueStates.FindChecked(Stream.Queue);
+					QueueState.EventStream.Append(MoveTemp(Stream));
+				}
 			}
 
 			ProcessAllQueues();
 		}
-	} GGPUProfilerSink_StatSystem;
+
+	} GGPUProfiler;
+
+	void ProcessEvents(TArrayView<FEventStream> EventStreams)
+	{
+		TRACE_CPUPROFILER_EVENT_SCOPE(UE::RHI::GPUProfiler::ProcessEvents);
+		GGPUProfiler.ProcessEvents(EventStreams);
+	}
+
+	void InitializeQueues(TConstArrayView<FQueue> Queues)
+	{
+		GGPUProfiler.InitializeQueues(Queues);
+	}
 
 #if WITH_PROFILEGPU
 	static FAutoConsoleCommand GCommand_ProfileGPU(
@@ -2776,15 +2759,15 @@ namespace UE::RHI::GPUProfiler
 		TEXT("Captures statistics about a frame of GPU work and prints the results to the log."),
 		FConsoleCommandWithArgsDelegate::CreateLambda([](const TArray<FString>& Args)
 		{
-			GGPUProfilerSink_StatSystem.bIsProfiling = true;
-			GGPUProfilerSink_StatSystem.bTriggerProfile = true;
+			GGPUProfiler.bIsProfiling = true;
+			GGPUProfiler.bTriggerProfile = true;
 		}));
 #endif // WITH_PROFILEGPU
 
 	RHI_API bool IsProfiling()
 	{
 	#if WITH_PROFILEGPU
-		return GGPUProfilerSink_StatSystem.bIsProfiling;
+		return GGPUProfiler.bIsProfiling;
 	#else
 		return false;
 	#endif
diff --git a/Engine/Source/Runtime/RHI/Public/GPUProfiler.h b/Engine/Source/Runtime/RHI/Public/GPUProfiler.h
--- a/Engine/Source/Runtime/RHI/Public/GPUProfiler.h
+++ b/Engine/Source/Runtime/RHI/Public/GPUProfiler.h
@@ -317,8 +317,6 @@ namespace UE::RHI::GPUProfiler
 
 	class FEventStream
 	{
-		friend struct FEventSink;
-
 	private:
 		struct FChunk
 		{
@@ -326,6 +324,7 @@ namespace UE::RHI::GPUProfiler
 			{
 				FChunk* Next = nullptr;
 				uint32 Num = 0;
+				uint32 Iter = 0;
 
 			#if WITH_RHI_BREADCRUMBS
 				FRHIBreadcrumbAllocatorArray BreadcrumbAllocators;
@@ -471,59 +470,12 @@ namespace UE::RHI::GPUProfiler
 			Other.Current = nullptr;
 			Other.First = nullptr;
 		}
-	};
-
-	struct FEventSink
-	{
-	protected:
-		struct FIterator
-		{
-			friend FEventSink;
-
-		private:
-			TSharedRef<FEventStream> Stream;
-
-			FEventStream::FChunk* Current;
-			uint32 Index = 0;
-
-		public:
-			FIterator(TSharedRef<FEventStream> const& Stream)
-				: Stream(Stream)
-				, Current(Stream->First)
-			{}
-
-			FEvent const* Peek() const
-			{
-				return Current ? Current->GetElement(Index) : nullptr;
-			}
 
-			FEvent const* Pop()
-			{
-				FEvent const* Result = Peek();
-				if (Result)
-				{
-					++Index;
-
-					while (Current && Index >= Current->Header.Num)
-					{
-						Current = Current->Header.Next;
-						Index = 0;
-					}
-				}
-
-				return Result;
-			}
-		};
-
-		RHI_API FEventSink();
-		RHI_API ~FEventSink();
-
-		FEventSink(FEventSink const&) = delete;
-		FEventSink(FEventSink&&) = delete;
+	private:
+		friend struct FGPUProfiler;
 
-	public:
-		virtual void ProcessStreams(TConstArrayView<TSharedRef<FEventStream>> EventStreams) = 0;
-		virtual void InitializeQueues(TConstArrayView<FQueue> Queues) = 0;
+		FEvent const* Peek() const;
+		void Pop();
 	};
 
 	RHI_API void ProcessEvents(TArrayView<FEventStream> EventStreams);
