Fix RHI breadcrumb lifetime use-after-free crashes
- The RenderGraph breadcrumb allocator was not being included in the BreadcrumbAllocatorRefs array on the RHI command list it is attached to.
- RenderGraph itself was keeping the allocator alive for some time via its async delete system, but this isn't always long enough to include the work done by the GPU profiler etc.

Also removed the Alloc function on FRHIBreadcrumbAllocator, so that only breadcrumbs may be allocated from it.

#Synced-CL 48210367

diff --git a/Engine/Source/Runtime/RHI/Private/RHICommandList.cpp b/Engine/Source/Runtime/RHI/Private/RHICommandList.cpp
--- a/Engine/Source/Runtime/RHI/Private/RHICommandList.cpp
+++ b/Engine/Source/Runtime/RHI/Private/RHICommandList.cpp
@@ -327,11 +327,10 @@ void FRHICommandListBase::AttachBreadcrumbSubTree(FRHIBreadcrumbAllocator& Alloc
 	// Switch the current breadcrumb allocator out for the subtree one.
 	if (BreadcrumbAllocator.Get() != &Allocator)
 	{
-		if (BreadcrumbAllocator)
-		{
-			BreadcrumbAllocatorRefs.AddUnique(BreadcrumbAllocator.Get());
-		}
+		check(!BreadcrumbAllocator.IsValid() || BreadcrumbAllocatorRefs.Contains(BreadcrumbAllocator.Get()));
+
 		BreadcrumbAllocator = Allocator.AsShared();
+		BreadcrumbAllocatorRefs.AddUnique(BreadcrumbAllocator.Get());
 	}
 }
 #endif
diff --git a/Engine/Source/Runtime/RHI/Public/RHIBreadcrumbs.h b/Engine/Source/Runtime/RHI/Public/RHIBreadcrumbs.h
--- a/Engine/Source/Runtime/RHI/Public/RHIBreadcrumbs.h
+++ b/Engine/Source/Runtime/RHI/Public/RHIBreadcrumbs.h
@@ -265,7 +265,7 @@ namespace UE::RHI::GPUProfiler
 			, ID(NextID.fetch_add(1, std::memory_order_relaxed) | 0x80000000) // Set the top bit to avoid collision with zero (i.e. "no breadcrumb")
 		{}
 
-		FRHIBreadcrumbNode*      & GetNextPtr(ERHIPipeline Pipeline)       { return NextPtrs[GetRHIPipelineIndex(Pipeline)]; }
+		FRHIBreadcrumbNode*&       GetNextPtr(ERHIPipeline Pipeline)       { return NextPtrs[GetRHIPipelineIndex(Pipeline)]; }
 		FRHIBreadcrumbNode* const& GetNextPtr(ERHIPipeline Pipeline) const { return NextPtrs[GetRHIPipelineIndex(Pipeline)]; }
 
 		FRHIBreadcrumbNode* GetParent() const { return Parent; }
@@ -342,6 +342,7 @@ namespace UE::RHI::GPUProfiler
 	{
 	public:
 		inline void AddUnique(FRHIBreadcrumbAllocator* Allocator);
+		inline bool Contains(FRHIBreadcrumbAllocator* Allocator);
 	};
 
 	class FRHIBreadcrumbAllocator : public TSharedFromThis<FRHIBreadcrumbAllocator>
@@ -359,18 +360,6 @@ namespace UE::RHI::GPUProfiler
 
 		FRHIBreadcrumbAllocatorArray const& GetParents() const { return Parents; }
 
-		template <typename TType, typename... TArgs>
-		TType* Alloc(TArgs&&... Args)
-		{
-			static_assert(std::is_trivially_destructible<TType>::value, "Only trivially destructable types may be used with the RHI breadcrumb allocator.");
-			return new (Inner.Alloc(sizeof(TType), alignof(TType))) TType(Forward<TArgs>(Args)...);
-		}
-
-		void* Alloc(uint32 Size, uint32 Align)
-		{
-			return Inner.Alloc(Size, Align);
-		}
-
 		template <typename TDesc, typename... TValues>
 		inline FRHIBreadcrumbNode* AllocBreadcrumb(TRHIBreadcrumbInitializer<TDesc, TValues...> const& Args);
 
@@ -381,16 +370,24 @@ namespace UE::RHI::GPUProfiler
 	};
 
 	inline void FRHIBreadcrumbAllocatorArray::AddUnique(FRHIBreadcrumbAllocator* Allocator)
+	{
+		if (!Contains(Allocator))
+		{
+			Add(Allocator->AsShared());
+		}
+	}
+
+	inline bool FRHIBreadcrumbAllocatorArray::Contains(FRHIBreadcrumbAllocator* Allocator)
 	{
 		for (TSharedRef<FRHIBreadcrumbAllocator> const& Existing : *this)
 		{
 			if (Allocator == &Existing.Get())
 			{
-				return;
+				return true;
 			}
 		}
 
-		Add(Allocator->AsShared());
+		return false;
 	}
 
 	//
diff --git a/Engine/Source/Runtime/RHI/Public/RHICommandList.h b/Engine/Source/Runtime/RHI/Public/RHICommandList.h
--- a/Engine/Source/Runtime/RHI/Public/RHICommandList.h
+++ b/Engine/Source/Runtime/RHI/Public/RHICommandList.h
@@ -3254,6 +3254,7 @@ public:
 		if (!BreadcrumbAllocator.IsValid())
 		{
 			BreadcrumbAllocator = MakeShared<FRHIBreadcrumbAllocator>();
+			BreadcrumbAllocatorRefs.Add(BreadcrumbAllocator.ToSharedRef()); // Not AddUnique because we know this is a new instance.
 		}
 
 		return *BreadcrumbAllocator;
diff --git a/Engine/Source/Runtime/RHI/Public/RHICommandList.inl b/Engine/Source/Runtime/RHI/Public/RHICommandList.inl
--- a/Engine/Source/Runtime/RHI/Public/RHICommandList.inl
+++ b/Engine/Source/Runtime/RHI/Public/RHICommandList.inl
@@ -191,7 +191,11 @@ inline FRHIResourceReplaceBatcher::~FRHIResourceReplaceBatcher()
 
 		return std::apply([&](auto&&... Values)
 		{
-			return Alloc<UE::RHI::Breadcrumbs::Private::TRHIBreadcrumb<TDesc>>(*this, *Desc, std::forward<decltype(Values)>(Values)...);
+			using TType = UE::RHI::Breadcrumbs::Private::TRHIBreadcrumb<TDesc>;
+			static_assert(std::is_trivially_destructible<TType>::value, "Only trivially destructable types may be used with the RHI breadcrumb allocator.");
+
+			void* Addr = Inner.Alloc(sizeof(TType), alignof(TType));
+			return new (Addr) TType(*this, *Desc, std::forward<decltype(Values)>(Values)...);
 		}, std::get<1>(Args));
 	}
 
