Deprecating a number of global RHI functions that use implied command lists.

#Synced-CL 48371251

diff --git a/Engine/Source/Runtime/RHI/Public/RHICommandList.h b/Engine/Source/Runtime/RHI/Public/RHICommandList.h
--- a/Engine/Source/Runtime/RHI/Public/RHICommandList.h
+++ b/Engine/Source/Runtime/RHI/Public/RHICommandList.h
@@ -5581,6 +5581,7 @@ inline void RHIClearTextureReference(FRHITextureReference* TextureRef)
 	FRHICommandListImmediate::Get().UpdateTextureReference(TextureRef, nullptr);
 }
 
+UE_DEPRECATED(5.8, "RHICreateTexture with an implied immediate command list is deprecated")
 inline FTextureRHIRef RHICreateTexture(const FRHITextureCreateDesc& CreateDesc)
 {
 	return FRHICommandListImmediate::Get().CreateTexture(CreateDesc);
@@ -5593,24 +5594,13 @@ inline FTextureRHIRef RHIAsyncCreateTexture2D(uint32 SizeX, uint32 SizeY, uint8
 	return GDynamicRHI->RHIAsyncCreateTexture2D(SizeX, SizeY, Format, NumMips, Flags, ResourceState, InitialMipData, NumInitialMips, DebugName, OutCompletionEvent);
 }
 
+UE_DEPRECATED(5.8, "RHIAsyncReallocateTexture2D with an implied immediate command list is deprecated")
 inline FTextureRHIRef RHIAsyncReallocateTexture2D(FRHITexture* Texture2D, int32 NewMipCount, int32 NewSizeX, int32 NewSizeY, FThreadSafeCounter* RequestStatus)
 {
 	return FRHICommandListExecutor::GetImmediateCommandList().AsyncReallocateTexture2D(Texture2D, NewMipCount, NewSizeX, NewSizeY, RequestStatus);
 }
 
-UE_DEPRECATED(5.6, "RHIFinalizeAsyncReallocateTexture2D is no longer implemented.")
-inline ETextureReallocationStatus RHIFinalizeAsyncReallocateTexture2D(FRHITexture* Texture2D, bool bBlockUntilCompleted)
-{
-	return TexRealloc_Succeeded;
-}
-
-UE_DEPRECATED(5.6, "RHICancelAsyncReallocateTexture2D is no longer implemented.")
-inline ETextureReallocationStatus RHICancelAsyncReallocateTexture2D(FRHITexture* Texture2D, bool bBlockUntilCompleted)
-{
-	return TexRealloc_Succeeded;
-}
-
-//UE_DEPRECATED(5.6, "RHICmdList.LockTexture should be used for all texture type locks.")
+// UE_DEPRECATED(5.6, "RHICmdList.LockTexture should be used for all texture type locks.")
 inline void* RHILockTexture2D(FRHITexture* Texture, uint32 MipIndex, EResourceLockMode LockMode, uint32& DestStride, bool bLockWithinMiptail, bool bFlushRHIThread = true, uint64* OutLockedByteCount = nullptr)
 {
 	return FRHICommandListExecutor::GetImmediateCommandList().LockTexture2D(Texture, MipIndex, LockMode, DestStride, bLockWithinMiptail, bFlushRHIThread, OutLockedByteCount);
@@ -5622,50 +5612,55 @@ inline void RHIUnlockTexture2D(FRHITexture* Texture, uint32 MipIndex, bool bLock
 	FRHICommandListExecutor::GetImmediateCommandList().UnlockTexture2D(Texture, MipIndex, bLockWithinMiptail, bFlushRHIThread);
 }
 
-//UE_DEPRECATED(5.6, "RHICmdList.LockTexture should be used for all texture type locks.")
+UE_DEPRECATED(5.8, "RHILockTexture2DArray with an implied immediate command list is deprecated")
 inline void* RHILockTexture2DArray(FRHITexture* Texture, uint32 TextureIndex, uint32 MipIndex, EResourceLockMode LockMode, uint32& DestStride, bool bLockWithinMiptail)
 {
 	return FRHICommandListExecutor::GetImmediateCommandList().LockTexture2DArray(Texture, TextureIndex, MipIndex, LockMode, DestStride, bLockWithinMiptail);
 }
 
-//UE_DEPRECATED(5.6, "RHICmdList.UnlockTexture should be used for all texture type unlocks.")
+UE_DEPRECATED(5.8, "RHIUnlockTexture2DArray with an implied immediate command list is deprecated")
 inline void RHIUnlockTexture2DArray(FRHITexture* Texture, uint32 TextureIndex, uint32 MipIndex, bool bLockWithinMiptail)
 {
 	 FRHICommandListExecutor::GetImmediateCommandList().UnlockTexture2DArray(Texture, TextureIndex, MipIndex, bLockWithinMiptail);
 }
 
-//UE_DEPRECATED(5.6, "RHICmdList.LockTexture should be used for all texture type locks.")
+UE_DEPRECATED(5.8, "RHILockTextureCubeFace with an implied immediate command list is deprecated")
 inline void* RHILockTextureCubeFace(FRHITexture* Texture, uint32 FaceIndex, uint32 ArrayIndex, uint32 MipIndex, EResourceLockMode LockMode, uint32& DestStride, bool bLockWithinMiptail)
 {
 	return FRHICommandListExecutor::GetImmediateCommandList().LockTextureCubeFace(Texture, FaceIndex, ArrayIndex, MipIndex, LockMode, DestStride, bLockWithinMiptail);
 }
 
-//UE_DEPRECATED(5.6, "RHICmdList.UnlockTexture should be used for all texture type unlocks.")
+UE_DEPRECATED(5.8, "RHIUnlockTextureCubeFace with an implied immediate command list is deprecated")
 inline void RHIUnlockTextureCubeFace(FRHITexture* Texture, uint32 FaceIndex, uint32 ArrayIndex, uint32 MipIndex, bool bLockWithinMiptail)
 {
 	FRHICommandListExecutor::GetImmediateCommandList().UnlockTextureCubeFace(Texture, FaceIndex, ArrayIndex, MipIndex, bLockWithinMiptail);
 }
 
+UE_DEPRECATED(5.8, "RHIUpdateTexture2D with an implied immediate command list is deprecated")
 inline void RHIUpdateTexture2D(FRHITexture* Texture, uint32 MipIndex, const struct FUpdateTextureRegion2D& UpdateRegion, uint32 SourcePitch, const uint8* SourceData)
 {
 	 FRHICommandListExecutor::GetImmediateCommandList().UpdateTexture2D(Texture, MipIndex, UpdateRegion, SourcePitch, SourceData);
 }
 
+UE_DEPRECATED(5.8, "RHIBeginUpdateTexture3D with an implied immediate command list is deprecated")
 inline FUpdateTexture3DData RHIBeginUpdateTexture3D(FRHITexture* Texture, uint32 MipIndex, const struct FUpdateTextureRegion3D& UpdateRegion)
 {
 	return FRHICommandListExecutor::GetImmediateCommandList().BeginUpdateTexture3D(Texture, MipIndex, UpdateRegion);
 }
 
+UE_DEPRECATED(5.8, "RHIEndUpdateTexture3D with an implied immediate command list is deprecated")
 inline void RHIEndUpdateTexture3D(FUpdateTexture3DData& UpdateData)
 {
 	FRHICommandListExecutor::GetImmediateCommandList().EndUpdateTexture3D(UpdateData);
 }
 
+UE_DEPRECATED(5.8, "RHIEndMultiUpdateTexture3D with an implied immediate command list is deprecated")
 inline void RHIEndMultiUpdateTexture3D(TArray<FUpdateTexture3DData>& UpdateDataArray)
 {
 	FRHICommandListExecutor::GetImmediateCommandList().EndMultiUpdateTexture3D(UpdateDataArray);
 }
 
+UE_DEPRECATED(5.8, "RHIUpdateTexture3D with an implied immediate command list is deprecated")
 inline void RHIUpdateTexture3D(FRHITexture* Texture, uint32 MipIndex, const struct FUpdateTextureRegion3D& UpdateRegion, uint32 SourceRowPitch, uint32 SourceDepthPitch, const uint8* SourceData)
 {
 	 FRHICommandListExecutor::GetImmediateCommandList().UpdateTexture3D(Texture, MipIndex, UpdateRegion, SourceRowPitch, SourceDepthPitch, SourceData);
@@ -5676,11 +5671,13 @@ inline void RHIFlushResources()
 	return FRHICommandListExecutor::GetImmediateCommandList().FlushResources();
 }
 
+UE_DEPRECATED(5.8, "RHIVirtualTextureSetFirstMipInMemory with an implied immediate command list is deprecated")
 inline void RHIVirtualTextureSetFirstMipInMemory(FRHITexture* Texture, uint32 FirstMip)
 {
 	 FRHICommandListExecutor::GetImmediateCommandList().VirtualTextureSetFirstMipInMemory(Texture, FirstMip);
 }
 
+UE_DEPRECATED(5.8, "RHIVirtualTextureSetFirstMipVisible with an implied immediate command list is deprecated")
 inline void RHIVirtualTextureSetFirstMipVisible(FRHITexture* Texture, uint32 FirstMip)
 {
 	 FRHICommandListExecutor::GetImmediateCommandList().VirtualTextureSetFirstMipVisible(Texture, FirstMip);
@@ -5737,16 +5734,19 @@ inline FRHIShaderLibraryRef RHICreateShaderLibrary(EShaderPlatform Platform, FSt
     return GDynamicRHI->RHICreateShaderLibrary(Platform, FilePath, Name);
 }
 
+UE_DEPRECATED(5.8, "RHILockStagingBuffer with an implied immediate command list is deprecated")
 inline void* RHILockStagingBuffer(FRHIStagingBuffer* StagingBuffer, uint32 Offset, uint32 Size)
 {
 	return FRHICommandListExecutor::GetImmediateCommandList().LockStagingBuffer(StagingBuffer, Offset, Size);
 }
 
+UE_DEPRECATED(5.8, "RHILockStagingBuffer with an implied immediate command list is deprecated")
 inline void* RHILockStagingBuffer(FRHIStagingBuffer* StagingBuffer, FRHIGPUFence* Fence, uint32 Offset, uint32 Size)
 {
 	return FRHICommandListExecutor::GetImmediateCommandList().LockStagingBuffer(StagingBuffer, Fence, Offset, Size);
 }
 
+UE_DEPRECATED(5.8, "RHIUnlockStagingBuffer with an implied immediate command list is deprecated")
 inline void RHIUnlockStagingBuffer(FRHIStagingBuffer* StagingBuffer)
 {
 	 FRHICommandListExecutor::GetImmediateCommandList().UnlockStagingBuffer(StagingBuffer);
