Removing deprecated RHI code from 5.6 and earlier.

#Synced-CL 48400404

diff --git a/Engine/Source/Runtime/RHI/Private/GPUProfiler.cpp b/Engine/Source/Runtime/RHI/Private/GPUProfiler.cpp
--- a/Engine/Source/Runtime/RHI/Private/GPUProfiler.cpp
+++ b/Engine/Source/Runtime/RHI/Private/GPUProfiler.cpp
@@ -2810,6 +2810,8 @@ FRHIGPUFrameTimeHistory::EResult FRHIGPUFrameTimeHistory::PopFrameCycles(FState&
 	}
 }
 
+uint32 GGPUFrameTime = 0;
+
 void FRHIGPUFrameTimeHistory::PushFrameCycles(double GPUFrequency, uint64 GPUCycles)
 {
 	double Seconds = double(GPUCycles) / GPUFrequency;
@@ -2821,16 +2823,12 @@ void FRHIGPUFrameTimeHistory::PushFrameCycles(double GPUFrequency, uint64 GPUCyc
 		History[NextIndex++ % MaxLength] = uint64(Cycles64);
 	}
 
-PRAGMA_DISABLE_DEPRECATION_WARNINGS
 	FPlatformAtomics::InterlockedExchange(reinterpret_cast<volatile int32*>(&GGPUFrameTime), int32(Cycles32));
-PRAGMA_ENABLE_DEPRECATION_WARNINGS
 }
 
 RHI_API uint32 RHIGetGPUFrameCycles(uint32 GPUIndex)
 {
-PRAGMA_DISABLE_DEPRECATION_WARNINGS
 	return (uint32)FPlatformAtomics::AtomicRead(reinterpret_cast<volatile int32*>(&GGPUFrameTime));
-PRAGMA_ENABLE_DEPRECATION_WARNINGS
 }
 
 #undef LOCTEXT_NAMESPACE
diff --git a/Engine/Source/Runtime/RHI/Private/PipelineFileCache.cpp b/Engine/Source/Runtime/RHI/Private/PipelineFileCache.cpp
--- a/Engine/Source/Runtime/RHI/Private/PipelineFileCache.cpp
+++ b/Engine/Source/Runtime/RHI/Private/PipelineFileCache.cpp
@@ -3850,17 +3850,6 @@ void FPipelineFileCacheManager::RegisterPSOCompileFailure(uint32 RunTimeHash, FG
 	}
 }
 
-FPipelineStateStats* FPipelineFileCacheManager::RegisterPSOStats(uint32 RunTimeHash)
-{
-	if(IsPipelineFileCacheEnabled() && LogPSOtoFileCache())
-	{
-		FRWScopeLock Lock(FileCacheLock, SLT_ReadOnly);
-		return RegisterPSOStatsInternal(Lock, RunTimeHash);
-	}
-	
-	return nullptr;
-}
-
 FPipelineStateStats* FPipelineFileCacheManager::RegisterPSOStatsInternal(FRWScopeLock& Lock, uint32 RunTimeHash)
 {
 	// May fail registration if the user cache has been closed
diff --git a/Engine/Source/Runtime/RHI/Private/RHIGlobals.cpp b/Engine/Source/Runtime/RHI/Private/RHIGlobals.cpp
--- a/Engine/Source/Runtime/RHI/Private/RHIGlobals.cpp
+++ b/Engine/Source/Runtime/RHI/Private/RHIGlobals.cpp
@@ -5,8 +5,6 @@
 
 FRHIGlobals GRHIGlobals;
 
-uint32 GGPUFrameTime = 0;
-
 //
 //	MSAA sample offsets.
 //
diff --git a/Engine/Source/Runtime/RHI/Public/PipelineFileCache.h b/Engine/Source/Runtime/RHI/Public/PipelineFileCache.h
--- a/Engine/Source/Runtime/RHI/Public/PipelineFileCache.h
+++ b/Engine/Source/Runtime/RHI/Public/PipelineFileCache.h
@@ -353,9 +353,6 @@ public:
 	// true if the named PSOFC is currently open.
 	RHI_API static bool HasPipelineFileCache(const FString& PSOCacheKey);
 
-	UE_DEPRECATED(5.6, "Do not call this function, please instead use the OutStats parameter in CacheGraphicsPSO/CacheComputePSO")
-	RHI_API static FPipelineStateStats* RegisterPSOStats(uint32 RunTimeHash);
-	
 	/*
 	 * This PSO has failed compile and is invalid - this cache should not return this invalid PSO from subsequent calls for PreCompile requests.
 	 * Note: Not implementated for Compute that has no flag to say it came from this cache - don't want to consume failures that didn't propagate from this cache.
diff --git a/Engine/Source/Runtime/RHI/Public/RHICommandList.h b/Engine/Source/Runtime/RHI/Public/RHICommandList.h
--- a/Engine/Source/Runtime/RHI/Public/RHICommandList.h
+++ b/Engine/Source/Runtime/RHI/Public/RHICommandList.h
@@ -811,83 +811,6 @@ public:
 		return Initializer.Finalize();
 	}
 
-PRAGMA_DISABLE_DEPRECATION_WARNINGS
-	UE_DEPRECATED(5.6, "CreateBuffer without FRHIBufferCreateDesc is deprecated")
-	inline FBufferRHIRef CreateNullBuffer(ERHIAccess ResourceState, FRHIResourceCreateInfo& CreateInfo)
-	{
-		const FRHIBufferCreateDesc CreateDesc =
-			FRHIBufferCreateDesc::CreateNull(CreateInfo.DebugName)
-			.SetGPUMask(CreateInfo.GPUMask)
-			.SetInitialState(ResourceState)
-			.SetClassName(CreateInfo.ClassName)
-			.SetOwnerName(CreateInfo.OwnerName);
-
-		return CreateBuffer(CreateDesc);
-	}
-
-	UE_DEPRECATED(5.6, "CreateBuffer without FRHIBufferCreateDesc is deprecated")
-	inline FBufferRHIRef CreateBuffer(uint32 Size, EBufferUsageFlags Usage, uint32 Stride, ERHIAccess ResourceState, FRHIResourceCreateInfo& CreateInfo)
-	{
-		if (CreateInfo.bWithoutNativeResource)
-		{
-			return CreateNullBuffer(ResourceState, CreateInfo);
-		}
-
-		FRHIBufferCreateDesc CreateDesc =
-			FRHIBufferCreateDesc::Create(CreateInfo.DebugName, Size, Stride, Usage)
-			.SetGPUMask(CreateInfo.GPUMask)
-			.SetInitialState(ResourceState)
-			.SetClassName(CreateInfo.ClassName)
-			.SetOwnerName(CreateInfo.OwnerName);
-
-		if (CreateInfo.ResourceArray)
-		{
-			CreateDesc.SetInitActionResourceArray(CreateInfo.ResourceArray);
-		}
-
-		return CreateBuffer(CreateDesc);
-	}
-
-	UE_DEPRECATED(5.6, "CreateBuffer without FRHIBufferCreateDesc is deprecated")
-	inline FBufferRHIRef CreateVertexBuffer(uint32 Size, EBufferUsageFlags Usage, ERHIAccess ResourceState, FRHIResourceCreateInfo& CreateInfo)
-	{
-		return CreateBuffer(Size, Usage | EBufferUsageFlags::VertexBuffer, 0, ResourceState, CreateInfo);
-	}
-
-	UE_DEPRECATED(5.6, "CreateBuffer without FRHIBufferCreateDesc is deprecated")
-	inline FBufferRHIRef CreateVertexBuffer(uint32 Size, EBufferUsageFlags Usage, FRHIResourceCreateInfo& CreateInfo)
-	{
-		ERHIAccess ResourceState = RHIGetDefaultResourceState(Usage | EBufferUsageFlags::VertexBuffer, false);
-		return CreateVertexBuffer(Size, Usage, ResourceState, CreateInfo);
-	}
-
-	UE_DEPRECATED(5.6, "CreateBuffer without FRHIBufferCreateDesc is deprecated")
-	inline FBufferRHIRef CreateStructuredBuffer(uint32 Stride, uint32 Size, EBufferUsageFlags Usage, ERHIAccess ResourceState, FRHIResourceCreateInfo& CreateInfo)
-	{
-		return CreateBuffer(Size, Usage | EBufferUsageFlags::StructuredBuffer, Stride, ResourceState, CreateInfo);
-	}
-
-	UE_DEPRECATED(5.6, "CreateBuffer without FRHIBufferCreateDesc is deprecated")
-	inline FBufferRHIRef CreateStructuredBuffer(uint32 Stride, uint32 Size, EBufferUsageFlags Usage, FRHIResourceCreateInfo& CreateInfo)
-	{
-		ERHIAccess ResourceState = RHIGetDefaultResourceState(Usage | EBufferUsageFlags::StructuredBuffer, false);
-		return CreateStructuredBuffer(Stride, Size, Usage, ResourceState, CreateInfo);
-	}
-
-	UE_DEPRECATED(5.6, "CreateBuffer without FRHIBufferCreateDesc is deprecated")
-	inline FBufferRHIRef CreateIndexBuffer(uint32 Stride, uint32 Size, EBufferUsageFlags Usage, ERHIAccess ResourceState, FRHIResourceCreateInfo& CreateInfo)
-	{
-		return CreateBuffer(Size, Usage | EBufferUsageFlags::IndexBuffer, Stride, ResourceState, CreateInfo);
-	}
-
-	UE_DEPRECATED(5.6, "CreateBuffer without FRHIBufferCreateDesc is deprecated")
-	inline FBufferRHIRef CreateIndexBuffer(uint32 Stride, uint32 Size, EBufferUsageFlags Usage, FRHIResourceCreateInfo& CreateInfo)
-	{
-		ERHIAccess ResourceState = RHIGetDefaultResourceState(Usage | EBufferUsageFlags::IndexBuffer, false);
-		return CreateIndexBuffer(Stride, Size, Usage, ResourceState, CreateInfo);
-	}
-PRAGMA_ENABLE_DEPRECATION_WARNINGS
-
 	inline void UpdateUniformBuffer(FRHIUniformBuffer* UniformBufferRHI, const void* Contents)
 	{
 		FRHICommandListScopedPipelineGuard ScopedPipeline(*this);
@@ -1008,142 +931,6 @@ PRAGMA_ENABLE_DEPRECATION_WARNINGS
 		return CreateShaderResourceView(Initializer.Buffer, Initializer);
 	}
 
-	UE_DEPRECATED(5.6, "Use the CreateUnorderedAccessView function that takes an FRHIViewDesc.")
-	inline FUnorderedAccessViewRHIRef CreateUnorderedAccessView(FRHIBuffer* Buffer, bool bUseUAVCounter, bool bAppendBuffer)
-	{
-		return CreateUnorderedAccessView(Buffer, FRHIViewDesc::CreateBufferUAV()
-			.SetTypeFromBuffer(Buffer)
-			.SetAtomicCounter(bUseUAVCounter)
-			.SetAppendBuffer(bAppendBuffer)
-		);
-	}
-
-	UE_DEPRECATED(5.6, "Use the CreateUnorderedAccessView function that takes an FRHIViewDesc.")
-	inline FUnorderedAccessViewRHIRef CreateUnorderedAccessView(FRHIBuffer* Buffer, uint8 Format)
-	{
-		// For back-compat reasons, SRVs of byte-address buffers created via this function ignore the Format, and instead create raw views.
-		if (Buffer && EnumHasAnyFlags(Buffer->GetDesc().Usage, BUF_ByteAddressBuffer))
-		{
-			return CreateUnorderedAccessView(Buffer, FRHIViewDesc::CreateBufferUAV()
-				.SetType(FRHIViewDesc::EBufferType::Raw)
-			);
-		}
-		else
-		{
-			return CreateUnorderedAccessView(Buffer, FRHIViewDesc::CreateBufferUAV()
-				.SetType(FRHIViewDesc::EBufferType::Typed)
-				.SetFormat(EPixelFormat(Format))
-			);
-		}
-	}
-
-	UE_DEPRECATED(5.6, "Use the CreateUnorderedAccessView function that takes an FRHIViewDesc.")
-	inline FUnorderedAccessViewRHIRef CreateUnorderedAccessView(FRHITexture* Texture, uint32 MipLevel = 0, uint16 FirstArraySlice = 0, uint16 NumArraySlices = 0)
-	{
-		check(MipLevel < 256);
-
-		return CreateUnorderedAccessView(Texture, FRHIViewDesc::CreateTextureUAV()
-			.SetDimensionFromTexture(Texture)
-			.SetMipLevel(uint8(MipLevel))
-			.SetArrayRange(FirstArraySlice, NumArraySlices)
-		);
-	}
-
-	UE_DEPRECATED(5.6, "Use the CreateUnorderedAccessView function that takes an FRHIViewDesc.")
-	inline FUnorderedAccessViewRHIRef CreateUnorderedAccessView(FRHITexture* Texture, uint32 MipLevel, uint8 Format, uint16 FirstArraySlice = 0, uint16 NumArraySlices = 0)
-	{
-		check(MipLevel < 256);
-
-		return CreateUnorderedAccessView(Texture, FRHIViewDesc::CreateTextureUAV()
-			.SetDimensionFromTexture(Texture)
-			.SetMipLevel(uint8(MipLevel))
-			.SetFormat(EPixelFormat(Format))
-			.SetArrayRange(FirstArraySlice, NumArraySlices)
-		);
-	}
-
-	UE_DEPRECATED(5.6, "Use the CreateShaderResourceView function that takes an FRHIViewDesc.")
-	inline FShaderResourceViewRHIRef CreateShaderResourceView(FRHIBuffer* Buffer)
-	{
-		FShaderResourceViewRHIRef SRVRef = CreateShaderResourceView(Buffer, FRHIViewDesc::CreateBufferSRV()
-			.SetTypeFromBuffer(Buffer));
-		checkf(SRVRef->GetDesc().Buffer.SRV.BufferType != FRHIViewDesc::EBufferType::Typed,
-			TEXT("Typed buffer should be created using CreateShaderResourceView where Format is specified."));
-		return SRVRef;
-	}
-
-	UE_DEPRECATED(5.6, "Use the CreateShaderResourceView function that takes an FRHIViewDesc.")
-	inline FShaderResourceViewRHIRef CreateShaderResourceView(FRHIBuffer* Buffer, uint32 Stride, uint8 Format)
-	{
-		check(Format != PF_Unknown);
-		check(Stride == GPixelFormats[Format].BlockBytes);
-
-		// For back-compat reasons, SRVs of byte-address buffers created via this function ignore the Format, and instead create raw views.
-		if (Buffer && EnumHasAnyFlags(Buffer->GetDesc().Usage, BUF_ByteAddressBuffer))
-		{
-			return CreateShaderResourceView(Buffer, FRHIViewDesc::CreateBufferSRV()
-				.SetType(FRHIViewDesc::EBufferType::Raw)
-			);
-		}
-		else
-		{
-			return CreateShaderResourceView(Buffer, FRHIViewDesc::CreateBufferSRV()
-				.SetType(FRHIViewDesc::EBufferType::Typed)
-				.SetFormat(EPixelFormat(Format))
-			);
-		}
-	}
-
-	UE_DEPRECATED(5.6, "Use the CreateShaderResourceView function that takes an FRHIViewDesc.")
-	inline FShaderResourceViewRHIRef CreateShaderResourceView(FRHITexture* Texture, const FRHITextureSRVCreateInfo& CreateInfo)
-	{
-		return CreateShaderResourceView(Texture, FRHIViewDesc::CreateTextureSRV()
-			.SetDimensionFromTexture(Texture)
-			.SetFormat     (CreateInfo.Format)
-			.SetMipRange   (CreateInfo.MipLevel, CreateInfo.NumMipLevels)
-			.SetDisableSRGB(CreateInfo.SRGBOverride == SRGBO_ForceDisable)
-			.SetArrayRange (CreateInfo.FirstArraySlice, CreateInfo.NumArraySlices)
-			.SetPlane      (CreateInfo.MetaData)
-		);
-	}
-
-	UE_DEPRECATED(5.6, "Use the CreateShaderResourceView function that takes an FRHIViewDesc.")
-	inline FShaderResourceViewRHIRef CreateShaderResourceView(FRHITexture* Texture, uint8 MipLevel)
-	{
-		return CreateShaderResourceView(Texture, FRHIViewDesc::CreateTextureSRV()
-			.SetDimensionFromTexture(Texture)
-			.SetMipRange(MipLevel, 1)
-		);
-	}
-
-	UE_DEPRECATED(5.6, "Use the CreateShaderResourceView function that takes an FRHIViewDesc.")
-	inline FShaderResourceViewRHIRef CreateShaderResourceView(FRHITexture* Texture, uint8 MipLevel, uint8 NumMipLevels, EPixelFormat Format)
-	{
-		return CreateShaderResourceView(Texture, FRHIViewDesc::CreateTextureSRV()
-			.SetDimensionFromTexture(Texture)
-			.SetMipRange(MipLevel, NumMipLevels)
-			.SetFormat(Format)
-		);
-	}
-
-	UE_DEPRECATED(5.6, "Use the CreateShaderResourceView function that takes an FRHIViewDesc.")
-	inline FShaderResourceViewRHIRef CreateShaderResourceViewWriteMask(FRHITexture* Texture2DRHI)
-	{
-		return CreateShaderResourceView(Texture2DRHI, FRHIViewDesc::CreateTextureSRV()
-			.SetDimensionFromTexture(Texture2DRHI)
-			.SetPlane(ERHITexturePlane::CMask)
-		);
-	}
-
-	UE_DEPRECATED(5.6, "Use the CreateShaderResourceView function that takes an FRHIViewDesc.")
-	inline FShaderResourceViewRHIRef CreateShaderResourceViewFMask(FRHITexture* Texture2DRHI)
-	{
-		return CreateShaderResourceView(Texture2DRHI, FRHIViewDesc::CreateTextureSRV()
-			.SetDimensionFromTexture(Texture2DRHI)
-			.SetPlane(ERHITexturePlane::FMask)
-		);
-	}
-
 	inline FRHIResourceCollectionRef CreateResourceCollection(TConstArrayView<FRHIResourceCollectionMember> InMembers)
 	{
 		LLM_SCOPE_BYNAME(TEXT("RHIMisc/CreateResourceCollection"));
@@ -3782,12 +3569,6 @@ public:
 				{
 					InlineParams[Index].ReferencedGeometries = AllocArray(Params[Index].ReferencedGeometries);
 				}
-				PRAGMA_DISABLE_DEPRECATION_WARNINGS
-				if (Params[Index].PerInstanceGeometries.Num())
-				{
-					InlineParams[Index].PerInstanceGeometries = AllocArray(Params[Index].PerInstanceGeometries);
-				}
-				PRAGMA_ENABLE_DEPRECATION_WARNINGS
 			}
 
 			ALLOC_COMMAND(FRHICommandBuildSceneAccelerationStructures)(MakeConstArrayView(InlineParams, Params.Num()));
@@ -4759,19 +4540,6 @@ public:
 		LLM_SCOPE(ELLMTag::Textures);
 		return GDynamicRHI->AsyncReallocateTexture2D_RenderThread(*this, Texture2D, NewMipCount, NewSizeX, NewSizeY, RequestStatus);
 	}
-	
-	UE_DEPRECATED(5.6, "FinalizeAsyncReallocateTexture2D is no longer implemented.")
-	inline ETextureReallocationStatus FinalizeAsyncReallocateTexture2D(FRHITexture* Texture2D, bool bBlockUntilCompleted)
-	{
-		LLM_SCOPE(ELLMTag::Textures);
-		return TexRealloc_Succeeded;
-	}
-	
-	UE_DEPRECATED(5.6, "CancelAsyncReallocateTexture2D is no longer implemented.")
-	inline ETextureReallocationStatus CancelAsyncReallocateTexture2D(FRHITexture* Texture2D, bool bBlockUntilCompleted)
-	{
-		return TexRealloc_Succeeded;
-	}
 
 	inline FRHILockTextureResult LockTexture(const FRHILockTextureArgs& Arguments)
 	{
@@ -4928,12 +4696,6 @@ public:
 		return GDynamicRHI->RHIFlushResources();
 	}
 
-	UE_DEPRECATED(5.6, "GetGPUFrameCycles is deprecated. Use the global scope RHIGetGPUFrameCycles() function.")
-	inline uint32 GetGPUFrameCycles()
-	{
-		return RHIGetGPUFrameCycles(GetGPUMask().ToIndex());
-	}
-
 	inline void SubmitAndBlockUntilGPUIdle()
 	{
 		QUICK_SCOPE_CYCLE_COUNTER(STAT_RHIMETHOD_SubmitAndBlockUntilGPUIdle_Flush);
diff --git a/Engine/Source/Runtime/RHI/Public/RHIContext.h b/Engine/Source/Runtime/RHI/Public/RHIContext.h
--- a/Engine/Source/Runtime/RHI/Public/RHIContext.h
+++ b/Engine/Source/Runtime/RHI/Public/RHIContext.h
@@ -626,19 +626,7 @@ struct FRayTracingSceneBuildParams
 	// Any referenced geometry is kept alive while the scene is alive.
 	TConstArrayView<FRHIRayTracingGeometry*> ReferencedGeometries;
 
-	UE_DEPRECATED(5.6, "No longer necessary.")
-	TConstArrayView<FRHIRayTracingGeometry*> PerInstanceGeometries;
-
 	EAccelerationStructureBuildMode BuildMode = EAccelerationStructureBuildMode::Build;
-
-	PRAGMA_DISABLE_DEPRECATION_WARNINGS
-	FRayTracingSceneBuildParams() = default;
-	FRayTracingSceneBuildParams(const FRayTracingSceneBuildParams&) = default;
-	FRayTracingSceneBuildParams& operator=(const FRayTracingSceneBuildParams&) = default;
-	FRayTracingSceneBuildParams(FRayTracingSceneBuildParams&&) = default;
-	FRayTracingSceneBuildParams& operator=(FRayTracingSceneBuildParams&&) = default;
-	~FRayTracingSceneBuildParams() = default;
-	PRAGMA_ENABLE_DEPRECATION_WARNINGS
 };
 
 struct FRayTracingClusterOperationParams
diff --git a/Engine/Source/Runtime/RHI/Public/RHIDefinitions.h b/Engine/Source/Runtime/RHI/Public/RHIDefinitions.h
--- a/Engine/Source/Runtime/RHI/Public/RHIDefinitions.h
+++ b/Engine/Source/Runtime/RHI/Public/RHIDefinitions.h
@@ -105,17 +105,6 @@ public:
 
 	RHI_API FString ToString() const;
 
-	UE_DEPRECATED(5.6,
-		"FDebugName::ToString(TCHAR* Out, uint32 OutSize) is dangerous and can lead to buffer overflow if the provided "
-		"buffer is smaller than FDebugName::StringBufferSize, even if the OutSize parameter indicates the buffer is "
-		"smaller than this value. Use the templated ToString() or ToStringTruncate() functions to format the name "
-		"string into a pre-allocated array, or use the allocating ToString() function that returns an FString.")
-	uint32 ToString(TCHAR* Out, uint32 OutSize) const
-	{
-		check(OutSize >= StringBufferSize);
-		return ToStringInternal(Out, OutSize);
-	}
-
 	/**
 	 * Converts the FDebugName to a string buffer, avoiding dynamic allocations.
 	 * Returns the length of the string, excluding the null terminator.
@@ -289,12 +278,6 @@ enum
 };
 static_assert(MaxSimultaneousRenderTargets <= (1 << MaxSimultaneousRenderTargets_NumBits), "MaxSimultaneousRenderTargets will not fit on MaxSimultaneousRenderTargets_NumBits");
 
-/** The number of UAVs that may be simultaneously bound to a shader. */
-enum
-{
-	MaxSimultaneousUAVs UE_DEPRECATED(5.6, "Do not use MaxSimultaneousUAVs, this is now device reported by GRHIGlobals.MaxSimultaneousUAVs") = 8
-};
-
 enum class ERHIZBuffer
 {
 	// Before changing this, make sure all math & shader assumptions are correct! Also wrap your C++ assumptions with
diff --git a/Engine/Source/Runtime/RHI/Public/RHIGlobals.h b/Engine/Source/Runtime/RHI/Public/RHIGlobals.h
--- a/Engine/Source/Runtime/RHI/Public/RHIGlobals.h
+++ b/Engine/Source/Runtime/RHI/Public/RHIGlobals.h
@@ -758,10 +758,6 @@ PRAGMA_ENABLE_DEPRECATION_WARNINGS
 
 extern RHI_API FRHIGlobals GRHIGlobals;
 
-/** The GPU time taken to render the last frame. Same metric as FPlatformTime::Cycles(). */
-UE_DEPRECATED(5.6, "Direct use of GGPUFrameTime is deprecated. Call the global scope RHIGetGPUFrameCycles() function instead.")
-extern RHI_API uint32 GGPUFrameTime;
-
 //
 // Deprecated old-style names
 
diff --git a/Engine/Source/Runtime/RHI/Public/RHIResources.h b/Engine/Source/Runtime/RHI/Public/RHIResources.h
--- a/Engine/Source/Runtime/RHI/Public/RHIResources.h
+++ b/Engine/Source/Runtime/RHI/Public/RHIResources.h
@@ -368,43 +368,6 @@ struct FClearValueBinding
 	static RHI_API const FClearValueBinding DefaultNormal8Bit;
 };
 
-struct UE_DEPRECATED(5.6, "FRHIResourceCreateInfo is no longer used. Please use FRHIBufferCreateDesc or FRHITextureCreateDesc.") FRHIResourceCreateInfo
-{
-PRAGMA_DISABLE_DEPRECATION_WARNINGS
-	FRHIResourceCreateInfo() = default;
-
-	FRHIResourceCreateInfo(const TCHAR* InDebugName)
-		: DebugName(InDebugName)
-	{
-		check(InDebugName);
-	}
-
-	// for CreateBuffer calls
-	FRHIResourceCreateInfo(const TCHAR* InDebugName, FResourceArrayUploadInterface* InResourceArray)
-		: FRHIResourceCreateInfo(InDebugName)
-	{
-		ResourceArray = InResourceArray;
-	}
-PRAGMA_ENABLE_DEPRECATION_WARNINGS
-
-	FName GetTraceClassName() const									{ const static FLazyName FRHIBufferName(TEXT("FRHIBuffer")); return (ClassName == NAME_None) ? FRHIBufferName : ClassName; }
-
-	// for CreateBuffer calls
-	UE_DEPRECATED(5.6, "Please use FRHIBufferCreateDesc for creating buffers with Resource Arrays")
-	FResourceArrayUploadInterface* ResourceArray = nullptr;
-
-	// set of GPUs on which to create the resource
-	FRHIGPUMask GPUMask = FRHIGPUMask::All();
-
-	// whether to create an RHI object with no underlying resource
-	bool bWithoutNativeResource = false;
-
-	const TCHAR* DebugName;
-
-	FName ClassName = NAME_None;	// The owner class of FRHIBuffer used for Insight asset metadata tracing
-	FName OwnerName = NAME_None;	// The owner name used for Insight asset metadata tracing
-};
-
 class FExclusiveDepthStencil
 {
 public:
@@ -2147,16 +2110,6 @@ struct FRHITextureCreateDesc : public FRHITextureDesc
 		return SetInitAction(ERHITextureInitAction::BulkData);
 	}
 
-	UE_DEPRECATED(5.6, "SetBulkData has been renamed to SetInitActionBulkData")
-	FRHITextureCreateDesc& SetBulkData(FResourceBulkDataInterface* InBulkData)
-	{
-		if (InBulkData)
-		{
-			SetInitActionBulkData(InBulkData);
-		}
-		return *this;
-	}
-
 	FName GetTraceClassName() const
 	{
 		const static FLazyName FRHITextureName(TEXT("FRHITexture"));
@@ -3944,12 +3897,6 @@ public:
 		return Initializer;
 	}
 
-	UE_DEPRECATED(5.6, "GetOrCreateInlineBufferSRV is deprecated. Use GetInlineBufferSize and provide the buffer to RHICommitInlineRayTracingBuffer function.")
-	FRHIShaderResourceView* GetOrCreateInlineBufferSRV(FRHICommandListBase& RHICmdList)
-	{
-		return nullptr;
-	}
-
 	// Returns the size and stride of the structured buffer for RHI-specific inline parameters associated with this SBT.
 	// Returns 0 if current RHI does not require this buffer.
 	virtual FRHISizeAndStride GetInlineBindingDataSizeAndStride() const
diff --git a/Engine/Source/Runtime/RHI/Public/RHIShaderPlatform.h b/Engine/Source/Runtime/RHI/Public/RHIShaderPlatform.h
--- a/Engine/Source/Runtime/RHI/Public/RHIShaderPlatform.h
+++ b/Engine/Source/Runtime/RHI/Public/RHIShaderPlatform.h
@@ -11,9 +11,7 @@ enum EShaderPlatform : uint16
 {
 	SP_PCD3D_SM5 = 0,
 	SP_METAL_ES3_1_IOS = 11,
-	SP_METAL UE_DEPRECATED(5.6, "Use SP_METAL_ES3_1_IOS enumeration value instead.") = SP_METAL_ES3_1_IOS,
 	SP_METAL_SM5_IOS = 12,
-	SP_METAL_MRT UE_DEPRECATED(5.6, "Use SP_METAL_SM5_IOS enumeration value instead.") = SP_METAL_SM5_IOS,
 	SP_PCD3D_ES3_1 = 14,
 	SP_OPENGL_PCES3_1 = 15,
 	SP_METAL_SM5 = 16,
@@ -21,12 +19,9 @@ enum EShaderPlatform : uint16
 	SP_VULKAN_SM5 = 20,
 	SP_VULKAN_ES3_1_ANDROID = 21,
 	SP_METAL_ES3_1 = 22,
-	SP_METAL_MACES3_1 UE_DEPRECATED(5.6, "Use SP_METAL_ES3_1 enumeration value instead.") = SP_METAL_ES3_1,
 	SP_OPENGL_ES3_1_ANDROID = 24,
 	SP_METAL_ES3_1_TVOS = 30,
-	SP_METAL_TVOS UE_DEPRECATED(5.6, "Use SP_METAL_ES3_1_TVOS enumeration value instead.") = SP_METAL_ES3_1_TVOS,
 	SP_METAL_SM5_TVOS = 31,
-	SP_METAL_MRT_TVOS UE_DEPRECATED(5.6, "Use SP_METAL_SM5_TVOS enumeration value instead.") = SP_METAL_SM5_TVOS,
 	/**********************************************************************************/
 	/* !! Do not add any new platforms here. Add them below SP_StaticPlatform_Last !! */
 	/**********************************************************************************/
diff --git a/Engine/Source/Runtime/RHI/Public/RHITransition.h b/Engine/Source/Runtime/RHI/Public/RHITransition.h
--- a/Engine/Source/Runtime/RHI/Public/RHITransition.h
+++ b/Engine/Source/Runtime/RHI/Public/RHITransition.h
@@ -324,37 +324,11 @@ struct FRHITransientAliasingInfo
 		return Info;
 	}
 
-	UE_DEPRECATED(5.5, "Discard aliasing ops are no longer necessary.")
-	static FRHITransientAliasingInfo Discard(class FRHITexture* Texture)
-	{
-		FRHITransientAliasingInfo Info;
-		Info.Texture = Texture;
-		Info.Type = EType::Texture;
-		Info.Action = EAction::Discard;
-		return Info;
-	}
-
-	UE_DEPRECATED(5.5, "Discard aliasing ops are no longer necessary.")
-	static FRHITransientAliasingInfo Discard(class FRHIBuffer* Buffer)
-	{
-		FRHITransientAliasingInfo Info;
-		Info.Buffer = Buffer;
-		Info.Type = EType::Buffer;
-		Info.Action = EAction::Discard;
-		return Info;
-	}
-
 	bool IsAcquire() const
 	{
 		return Action == EAction::Acquire;
 	}
 
-	UE_DEPRECATED(5.5, "Discard aliasing ops are no longer necessary.")
-	bool IsDiscard() const
-	{
-		return Action == EAction::Discard;
-	}
-
 	bool IsTexture() const
 	{
 		return Type == EType::Texture;
diff --git a/Engine/Source/Runtime/RHI/Public/RHIUtilities.h b/Engine/Source/Runtime/RHI/Public/RHIUtilities.h
--- a/Engine/Source/Runtime/RHI/Public/RHIUtilities.h
+++ b/Engine/Source/Runtime/RHI/Public/RHIUtilities.h
@@ -735,10 +735,6 @@ extern RHI_API void RHIGetPresentThresholds(float& OutTopPercent, float& OutBott
 /** Returns the value of the rhi.SyncAllowVariable CVar. */
 extern RHI_API bool RHIGetSyncAllowVariable();
 
-/** Signals the completion of the specified task graph event when the given frame has flipped. */
-UE_DEPRECATED(5.4, "RHICompleteGraphEventOnFlip is replaced with RHITriggerTaskEventOnFlip")
-inline void RHICompleteGraphEventOnFlip(uint64 PresentIndex, FGraphEventRef Event) { checkNoEntry(); }
-
 extern RHI_API void RHITriggerTaskEventOnFlip(uint64 PresentIndex, const UE::Tasks::FTaskEvent& TaskEvent);
 
 /** Sets the FrameIndex and InputTime for the current frame. */
diff --git a/Engine/Source/Runtime/RHICore/Private/RHICoreNvidiaAftermath.cpp b/Engine/Source/Runtime/RHICore/Private/RHICoreNvidiaAftermath.cpp
--- a/Engine/Source/Runtime/RHICore/Private/RHICoreNvidiaAftermath.cpp
+++ b/Engine/Source/Runtime/RHICore/Private/RHICoreNvidiaAftermath.cpp
@@ -989,15 +989,6 @@ namespace UE::RHICore::Nvidia::Aftermath
 		FScopeLock Lock(&ShaderRegistration.Lock);
 		ShaderRegistration.ShaderBinaryToDebugHash.Remove(Hash);
 	}
-
-	FCrashResult OnGPUCrash()
-	{
-		TArray<FCrashResult> Results;
-		OnGPUCrash(Results);
-
-		// Just report the last one, this method is pending deprecation
-		return !Results.IsEmpty() ? Results.Last() : FCrashResult {};
-	}
 	
 	bool OnGPUCrash(TArray<FCrashResult>& OutResults)
 	{
diff --git a/Engine/Source/Runtime/RHICore/Public/RHICoreNvidiaAftermath.h b/Engine/Source/Runtime/RHICore/Public/RHICoreNvidiaAftermath.h
--- a/Engine/Source/Runtime/RHICore/Public/RHICoreNvidiaAftermath.h
+++ b/Engine/Source/Runtime/RHICore/Public/RHICoreNvidiaAftermath.h
@@ -64,13 +64,6 @@
 			TOptional<uint64> GPUFaultAddress;
 		};
 
-		//
-		// Called by platform RHIs when a GPU crash is detected.
-		// Waits for Aftermath to finish crash dump processing, then returns the result.
-		//
-		UE_DEPRECATED(5.6, "Use OnGPUCrash(OutResults) overload instead, Aftermath may report multiple crashes")
-		RHICORE_API FCrashResult OnGPUCrash();
-
 		//
 		// Called by platform RHIs when a GPU crash is detected.
 		// Waits for Aftermath to finish crash dump processing, then returns the result.
