Remove RHICreateShaderLibrary_RenderThread
- The function was not part of a public API, and nothing was calling it.
- The only callsite in ShaderCodeLibrary.cpp uses the global scope RHICreateShaderLibrary function, which calls directly into the RHI via GDynamicRHI->

#Synced-CL 48503794

diff --git a/Engine/Source/Runtime/Apple/MetalRHI/Private/MetalDynamicRHI.h b/Engine/Source/Runtime/Apple/MetalRHI/Private/MetalDynamicRHI.h
--- a/Engine/Source/Runtime/Apple/MetalRHI/Private/MetalDynamicRHI.h
+++ b/Engine/Source/Runtime/Apple/MetalRHI/Private/MetalDynamicRHI.h
@@ -217,8 +217,6 @@ public:
 	virtual void* RHILockStagingBuffer(FRHIStagingBuffer* StagingBuffer, FRHIGPUFence* Fence, uint32 Offset, uint32 SizeRHI) final override;
 	virtual void RHIUnlockStagingBuffer(FRHIStagingBuffer* StagingBuffer) final override;
 
-	virtual FRHIShaderLibraryRef RHICreateShaderLibrary_RenderThread(class FRHICommandListImmediate& RHICmdList, EShaderPlatform Platform, FString FilePath, FString Name) final override;
-
 	virtual void RHIUpdateUniformBuffer(FRHICommandListBase& RHICmdList, FRHIUniformBuffer* UniformBufferRHI, const void* Contents) final override;
 
 	virtual uint16 RHIGetPlatformTextureMaxSampleCount() override;
diff --git a/Engine/Source/Runtime/Apple/MetalRHI/Private/MetalDynamicRHI_Shaders.cpp b/Engine/Source/Runtime/Apple/MetalRHI/Private/MetalDynamicRHI_Shaders.cpp
--- a/Engine/Source/Runtime/Apple/MetalRHI/Private/MetalDynamicRHI_Shaders.cpp
+++ b/Engine/Source/Runtime/Apple/MetalRHI/Private/MetalDynamicRHI_Shaders.cpp
@@ -249,8 +249,3 @@ FBoundShaderStateRHIRef FMetalDynamicRHI::RHICreateBoundShaderState(
 	NOT_SUPPORTED("RHICreateBoundShaderState");
 	return nullptr;
 }
-
-FRHIShaderLibraryRef FMetalDynamicRHI::RHICreateShaderLibrary_RenderThread(class FRHICommandListImmediate& RHICmdList, EShaderPlatform Platform, FString FilePath, FString Name)
-{
-	return RHICreateShaderLibrary(Platform, FilePath, Name);
-}
diff --git a/Engine/Source/Runtime/RHI/Private/RHICommandList.cpp b/Engine/Source/Runtime/RHI/Private/RHICommandList.cpp
--- a/Engine/Source/Runtime/RHI/Private/RHICommandList.cpp
+++ b/Engine/Source/Runtime/RHI/Private/RHICommandList.cpp
@@ -2257,12 +2257,6 @@ void FDynamicRHI::RHIEndMultiUpdateTexture3D(FRHICommandListBase& RHICmdList, TA
 	}
 }
 
-FRHIShaderLibraryRef FDynamicRHI::RHICreateShaderLibrary_RenderThread(class FRHICommandListImmediate& RHICmdList, EShaderPlatform Platform, FString FilePath, FString Name)
-{
-	FScopedRHIThreadStaller StallRHIThread(RHICmdList);
-	return GDynamicRHI->RHICreateShaderLibrary(Platform, FilePath, Name);
-}
-
 void FDynamicRHI::RHIMapStagingSurface_RenderThread(class FRHICommandListImmediate& RHICmdList, FRHITexture* Texture, uint32 GPUIndex, FRHIGPUFence* Fence, void*& OutData, int32& OutWidth, int32& OutHeight)
 {
 	if (Fence == nullptr || !Fence->Poll() || Fence->NumPendingWriteCommands.GetValue() != 0)
diff --git a/Engine/Source/Runtime/RHI/Public/DynamicRHI.h b/Engine/Source/Runtime/RHI/Public/DynamicRHI.h
--- a/Engine/Source/Runtime/RHI/Public/DynamicRHI.h
+++ b/Engine/Source/Runtime/RHI/Public/DynamicRHI.h
@@ -938,8 +938,6 @@ public:
 
 	RHI_API virtual void RHIEndMultiUpdateTexture3D(FRHICommandListBase& RHICmdList, TArray<FUpdateTexture3DData>& UpdateDataArray);
 
-	RHI_API virtual FRHIShaderLibraryRef RHICreateShaderLibrary_RenderThread(class FRHICommandListImmediate& RHICmdList, EShaderPlatform Platform, FString FilePath, FString Name);
-
 	RHI_API virtual void RHIMapStagingSurface_RenderThread(class FRHICommandListImmediate& RHICmdList, FRHITexture* Texture, uint32 GPUIndex, FRHIGPUFence* Fence, void*& OutData, int32& OutWidth, int32& OutHeight);
 	RHI_API virtual void RHIUnmapStagingSurface_RenderThread(class FRHICommandListImmediate& RHICmdList, FRHITexture* Texture, uint32 GPUIndex);
 	RHI_API virtual void RHIReadSurfaceFloatData_RenderThread(class FRHICommandListImmediate& RHICmdList, FRHITexture* Texture, FIntRect Rect, TArray<FFloat16Color>& OutData, ECubeFace CubeFace, int32 ArrayIndex, int32 MipIndex);
diff --git a/Engine/Source/Runtime/RHI/Public/RHIValidation.h b/Engine/Source/Runtime/RHI/Public/RHIValidation.h
--- a/Engine/Source/Runtime/RHI/Public/RHIValidation.h
+++ b/Engine/Source/Runtime/RHI/Public/RHIValidation.h
@@ -1000,11 +1000,6 @@ public:
 		RHI->RHIEndMultiUpdateTexture3D(RHICmdList, UpdateDataArray);
 	}
 
-	virtual FRHIShaderLibraryRef RHICreateShaderLibrary_RenderThread(class FRHICommandListImmediate& RHICmdList, EShaderPlatform Platform, FString FilePath, FString Name) override final
-	{
-		return RHI->RHICreateShaderLibrary_RenderThread(RHICmdList, Platform, FilePath, Name);
-	}
-
 	virtual void RHIReadSurfaceFloatData_RenderThread(class FRHICommandListImmediate& RHICmdList, FRHITexture* Texture, FIntRect Rect, TArray<FFloat16Color>& OutData, ECubeFace CubeFace, int32 ArrayIndex, int32 MipIndex) override final
 	{
 		RHI->RHIReadSurfaceFloatData_RenderThread(RHICmdList, Texture, Rect, OutData, CubeFace, ArrayIndex, MipIndex);
