RHICommandList cleanup
- Adding optional EThreadFence argument to EnqueueLambda to match EnqueueLambdaMultiPipe.
- Removing unused ValidateBoundShader variants.
- Removing unused BoundWorkGraphShaderRHI member.

#Synced-CL 48482604

diff --git a/Engine/Source/Runtime/RHI/Public/RHICommandList.h b/Engine/Source/Runtime/RHI/Public/RHICommandList.h
--- a/Engine/Source/Runtime/RHI/Public/RHICommandList.h
+++ b/Engine/Source/Runtime/RHI/Public/RHICommandList.h
@@ -581,8 +581,14 @@ public:
 		return AllocCommand(sizeof(TCmd), alignof(TCmd));
 	}
 
+	enum class EThreadFence
+	{
+		Enabled,
+		Disabled
+	};
+
 	template <typename LAMBDA>
-	inline void EnqueueLambda(const TCHAR* LambdaName, LAMBDA&& Lambda)
+	inline void EnqueueLambda(EThreadFence ThreadFence, const TCHAR* LambdaName, LAMBDA&& Lambda)
 	{
 		if (IsBottomOfPipe())
 		{
@@ -591,20 +597,25 @@ public:
 		else
 		{
 			ALLOC_COMMAND(TRHILambdaCommand<FRHICommandListBase, LAMBDA>)(Forward<LAMBDA>(Lambda), LambdaName);
+
+			if (ThreadFence == EThreadFence::Enabled)
+			{
+				RHIThreadFence(true);
+			}
 		}
 	}
 
+	template <typename LAMBDA>
+	inline void EnqueueLambda(const TCHAR* LambdaName, LAMBDA&& Lambda)
+	{
+		FRHICommandListBase::EnqueueLambda(EThreadFence::Disabled, LambdaName, Forward<LAMBDA>(Lambda));
+	}
+
 	template <typename LAMBDA>
 	inline void EnqueueLambda(LAMBDA&& Lambda)
 	{
 		FRHICommandListBase::EnqueueLambda(TEXT("TRHILambdaCommand"), Forward<LAMBDA>(Lambda));
 	}
-	
-	enum class EThreadFence
-	{
-		Enabled,
-		Disabled
-	};
 
 	template <typename LAMBDA>
 	void EnqueueLambdaMultiPipe(ERHIPipeline Pipelines, EThreadFence ThreadFence, const TCHAR* LambdaName, LAMBDA&& Lambda)
@@ -829,11 +840,11 @@ public:
 		}
 		else
 		{
-			EnqueueLambda([this, StreamSourceSlotRHI, BufferRHI] (FRHICommandListBase&)
-			{
-				StreamSourceSlotRHI->Buffer = BufferRHI;
-			});
-			RHIThreadFence(true);
+			EnqueueLambda(EThreadFence::Enabled, TEXT("UpdateStreamSourceSlot"),
+				[this, StreamSourceSlotRHI, BufferRHI](FRHICommandListBase&)
+				{
+					StreamSourceSlotRHI->Buffer = BufferRHI;
+				});
 		}
 	}
 
@@ -1056,25 +1067,22 @@ public:
 protected:
 	FMemStackBase& GetAllocator() { return MemManager; }
 
-	inline void ValidateBoundShader(FRHIVertexShader*        ShaderRHI) { checkSlow(PersistentState.BoundShaderInput.VertexShaderRHI          == ShaderRHI); }
-	inline void ValidateBoundShader(FRHIPixelShader*         ShaderRHI) { checkSlow(PersistentState.BoundShaderInput.PixelShaderRHI           == ShaderRHI); }
-	inline void ValidateBoundShader(FRHIGeometryShader*      ShaderRHI) { checkSlow(PersistentState.BoundShaderInput.GetGeometryShader()      == ShaderRHI); }
-	inline void ValidateBoundShader(FRHIComputeShader*       ShaderRHI) { checkSlow(PersistentState.BoundComputeShaderRHI                     == ShaderRHI); }
-	inline void ValidateBoundShader(FRHIWorkGraphShader*     ShaderRHI) { checkSlow(PersistentState.BoundWorkGraphShaderRHI                   == ShaderRHI); }
-	inline void ValidateBoundShader(FRHIMeshShader*          ShaderRHI) { checkSlow(PersistentState.BoundShaderInput.GetMeshShader()          == ShaderRHI); }
-	inline void ValidateBoundShader(FRHIAmplificationShader* ShaderRHI) { checkSlow(PersistentState.BoundShaderInput.GetAmplificationShader() == ShaderRHI); }
+	inline void ValidateBoundShader(const FRHIComputeShader* ShaderRHI)
+	{
+		checkSlow(PersistentState.BoundComputeShaderRHI == ShaderRHI);
+	}
 
-	inline void ValidateBoundShader(FRHIGraphicsShader* ShaderRHI)
+	inline void ValidateBoundShader(const FRHIGraphicsShader* ShaderRHI)
 	{
 #if DO_GUARD_SLOW
 		switch (ShaderRHI->GetFrequency())
 		{
-		case SF_Vertex:        checkSlow(PersistentState.BoundShaderInput.VertexShaderRHI          == ShaderRHI); break;
+		case SF_Vertex:        checkSlow(PersistentState.BoundShaderInput.GetVertexShader()        == ShaderRHI); break;
 		case SF_Mesh:          checkSlow(PersistentState.BoundShaderInput.GetMeshShader()          == ShaderRHI); break;
 		case SF_Amplification: checkSlow(PersistentState.BoundShaderInput.GetAmplificationShader() == ShaderRHI); break;
-		case SF_Pixel:         checkSlow(PersistentState.BoundShaderInput.PixelShaderRHI           == ShaderRHI); break;
+		case SF_Pixel:         checkSlow(PersistentState.BoundShaderInput.GetPixelShader()         == ShaderRHI); break;
 		case SF_Geometry:      checkSlow(PersistentState.BoundShaderInput.GetGeometryShader()      == ShaderRHI); break;
-		default: checkfSlow(false, TEXT("Unexpected graphics shader type %d"), ShaderRHI->GetFrequency());
+		default: checkfSlow(false, TEXT("Unexpected shader type %d"), ShaderRHI->GetFrequency());
 		}
 #endif // DO_GUARD_SLOW
 	}
@@ -1292,7 +1300,6 @@ protected:
 
 		FBoundShaderStateInput BoundShaderInput;
 		FRHIComputeShader* BoundComputeShaderRHI = nullptr;
-		FRHIWorkGraphShader* BoundWorkGraphShaderRHI = nullptr;
 		FRHICommandListScopedFence* CurrentFenceScope = nullptr;
 
 #if WITH_RHI_BREADCRUMBS
@@ -2547,7 +2554,7 @@ public:
 	{}
 
 	template <typename LAMBDA>
-	inline void EnqueueLambda(const TCHAR* LambdaName, LAMBDA&& Lambda)
+	inline void EnqueueLambda(EThreadFence ThreadFence, const TCHAR* LambdaName, LAMBDA&& Lambda)
 	{
 		if (IsBottomOfPipe())
 		{
@@ -2556,6 +2563,11 @@ public:
 		else
 		{
 			ALLOC_COMMAND(TRHILambdaCommand<FRHIComputeCommandList, LAMBDA>)(Forward<LAMBDA>(Lambda), LambdaName);
+
+			if (ThreadFence == EThreadFence::Enabled)
+			{
+				RHIThreadFence(true);
+			}
 		}
 	}
 
@@ -2573,6 +2585,12 @@ public:
 		}
 	}
 
+	template <typename LAMBDA>
+	inline void EnqueueLambda(const TCHAR* LambdaName, LAMBDA&& Lambda)
+	{
+		FRHIComputeCommandList::EnqueueLambda(EThreadFence::Disabled, LambdaName, Forward<LAMBDA>(Lambda));
+	}
+
 	template <typename LAMBDA>
 	inline void EnqueueLambda(LAMBDA&& Lambda)
 	{
@@ -3634,7 +3652,7 @@ public:
 	inline FRHIGeometryShader*      GetBoundGeometryShader     () const { return PersistentState.BoundShaderInput.GetGeometryShader();      }
 
 	template <typename LAMBDA>
-	inline void EnqueueLambda(const TCHAR* LambdaName, LAMBDA&& Lambda)
+	inline void EnqueueLambda(EThreadFence ThreadFence, const TCHAR* LambdaName, LAMBDA&& Lambda)
 	{
 		if (IsBottomOfPipe())
 		{
@@ -3643,9 +3661,20 @@ public:
 		else
 		{
 			ALLOC_COMMAND(TRHILambdaCommand<FRHICommandList, LAMBDA>)(Forward<LAMBDA>(Lambda), LambdaName);
+
+			if (ThreadFence == EThreadFence::Enabled)
+			{
+				RHIThreadFence(true);
+			}
 		}
 	}
 
+	template <typename LAMBDA>
+	inline void EnqueueLambda(const TCHAR* LambdaName, LAMBDA&& Lambda)
+	{
+		FRHICommandList::EnqueueLambda(EThreadFence::Disabled, LambdaName, Forward<LAMBDA>(Lambda));
+	}
+
 	template <typename LAMBDA>
 	inline void EnqueueLambda(LAMBDA&& Lambda)
 	{
@@ -4485,7 +4514,7 @@ public:
 	RHI_API void InitializeImmediateContexts();
 
 	template <typename LAMBDA>
-	inline void EnqueueLambda(const TCHAR* LambdaName, LAMBDA&& Lambda)
+	inline void EnqueueLambda(EThreadFence ThreadFence, const TCHAR* LambdaName, LAMBDA&& Lambda)
 	{
 		if (IsBottomOfPipe())
 		{
@@ -4494,9 +4523,20 @@ public:
 		else
 		{
 			ALLOC_COMMAND(TRHILambdaCommand<FRHICommandListImmediate, LAMBDA>)(Forward<LAMBDA>(Lambda), LambdaName);
+
+			if (ThreadFence == EThreadFence::Enabled)
+			{
+				RHIThreadFence(true);
+			}
 		}
 	}
 
+	template <typename LAMBDA>
+	inline void EnqueueLambda(const TCHAR* LambdaName, LAMBDA&& Lambda)
+	{
+		FRHICommandListImmediate::EnqueueLambda(EThreadFence::Disabled, TEXT("TRHILambdaCommand"), Forward<LAMBDA>(Lambda));
+	}
+
 	template <typename LAMBDA>
 	inline void EnqueueLambda(LAMBDA&& Lambda)
 	{
