Remove legacy GPU profiler from common engine code
- Removed all code where RHI_NEW_GPU_PROFILER == 0, and any associated code only used by these code paths.
- Also removed GPUPROFILERTRACE_ENABLED, which is always zero when the new profiler is enabled.

Todo: Deprecate and remove use of RDG_GPU_STAT_SCOPE, and macros in RealtimeGPUProfiler.h

#Synced-CL 48676711

diff --git a/Engine/Config/BaseInput.ini b/Engine/Config/BaseInput.ini
--- a/Engine/Config/BaseInput.ini
+++ b/Engine/Config/BaseInput.ini
@@ -242,7 +242,6 @@ MaxScrollbackSize=1024
 +ManualAutoCompleteList=(Command="StopRecordingTake",Desc="Stop recording the current sequence. Only one sequence recording can be active at one time.")
 +ManualAutoCompleteList=(Command="FreezeRendering",Desc="Toggle freezing of most aspects of rendering (such as visibility calculations), useful in conjunction with ToggleDebugCamera to fly around and see how frustum and occlusion culling is working")
 +ManualAutoCompleteList=(Command="ProfileGPU",Desc="Profile one frame of rendering commands sent to the GPU")
-+ManualAutoCompleteList=(Command="ProfileGPUHitches",Desc="Toggle profiling of GPU hitches.")
 +ManualAutoCompleteList=(Command="DumpGPU",Desc="Dump one frame of rendering intermediary resources to disk.")
 +ManualAutoCompleteList=(Command="Automation",Desc="Run an automation command (e.g., Automation RunTests TestName)")
 +ManualAutoCompleteList=(Command="CsvProfile Start",Desc="Start CSV profiling.")
diff --git a/Engine/Config/Mac/BaseMacEngine.ini b/Engine/Config/Mac/BaseMacEngine.ini
--- a/Engine/Config/Mac/BaseMacEngine.ini
+++ b/Engine/Config/Mac/BaseMacEngine.ini
@@ -35,7 +35,6 @@ QualityLevelMapping="high"
 [ConsoleVariables]
 r.HairStrands.Visibility.MSAA.SamplePerPixel=1
 r.SceneCulling.ExplicitCellBounds=1
-r.GPUStatsEnabled=0
 ; Workaround for the webkit WebBrowser issues with slate
 Slate.NoAnimationOnTabForegroundedEvent=True
 
diff --git a/Engine/Plugins/Runtime/nDisplay/Source/DisplayCluster/Private/Game/EngineClasses/Basics/DisplayClusterViewportClient.cpp b/Engine/Plugins/Runtime/nDisplay/Source/DisplayCluster/Private/Game/EngineClasses/Basics/DisplayClusterViewportClient.cpp
--- a/Engine/Plugins/Runtime/nDisplay/Source/DisplayCluster/Private/Game/EngineClasses/Basics/DisplayClusterViewportClient.cpp
+++ b/Engine/Plugins/Runtime/nDisplay/Source/DisplayCluster/Private/Game/EngineClasses/Basics/DisplayClusterViewportClient.cpp
@@ -60,14 +60,6 @@ static FAutoConsoleVariableRef CVarDisplayClusterForceCopyCrossGPU(
 	ECVF_RenderThreadSafe
 );
 
-int32 GDisplayClusterShowStats = 0;
-static FAutoConsoleVariableRef CVarDisplayClusterShowStats(
-	TEXT("DC.Stats"),
-	GDisplayClusterShowStats,
-	TEXT("Show per-view profiling stats for display cluster rendering."),
-	ECVF_RenderThreadSafe
-);
-
 int32 GDisplayClusterSingleRender = 1;
 static FAutoConsoleVariableRef CVarDisplayClusterSingleRender(
 	TEXT("DC.SingleRender"),
@@ -172,93 +164,6 @@ public:
 		}
 	}
 
-	virtual int32 DrawStatsHUD(FCanvas* InCanvas, int32 InX, int32 InY) override
-	{
-#if GPUPROFILERTRACE_ENABLED && (RHI_NEW_GPU_PROFILER == 0)
-		if (GDisplayClusterShowStats)
-		{
-			// Get GPU perf results
-			TArray<FRealtimeGPUProfilerDescriptionResult> PerfResults;
-			FRealtimeGPUProfiler::Get()->FetchPerfByDescription(PerfResults);
-
-			UFont* StatsFont = GetStatsFont();
-
-			const FLinearColor HeaderColor = FLinearColor(1.f, 0.2f, 0.f);
-
-			if (PerfResults.Num())
-			{
-				// Get CPU perf results
-				TArray<float> CpuPerfResults;
-				CpuPerfResults.AddUninitialized(PerfResults.Num());
-				{
-					FRWScopeLock Lock(CpuHistoryMutex, SLT_Write);
-
-					for (int32 ResultIndex = 0; ResultIndex < PerfResults.Num(); ResultIndex++)
-					{
-						CpuPerfResults[ResultIndex] = FetchHistoryAverage(PerfResults[ResultIndex].Description);
-					}
-				}
-
-				// Compute column sizes
-				int32 YIgnore;
-
-				const TCHAR* DescriptionHeader = TEXT("Display Cluster Stats");
-				int32 DescriptionColumnWidth;
-				StringSize(StatsFont, DescriptionColumnWidth, YIgnore, DescriptionHeader);
-
-				for (const FRealtimeGPUProfilerDescriptionResult& PerfResult : PerfResults)
-				{
-					int32 XL;
-					StringSize(StatsFont, XL, YIgnore, *PerfResult.Description);
-
-					DescriptionColumnWidth = FMath::Max(DescriptionColumnWidth, XL);
-				}
-
-				int32 NumberColumnWidth;
-				StringSize(StatsFont, NumberColumnWidth, YIgnore, *FString::ChrN(7, 'W'));
-
-				// Render header
-				InCanvas->DrawShadowedString(InX, InY, DescriptionHeader, StatsFont, HeaderColor);
-				RightJustify(InCanvas, StatsFont, InX + DescriptionColumnWidth + 1 * NumberColumnWidth, InY, TEXT("GPUs"), HeaderColor);
-				RightJustify(InCanvas, StatsFont, InX + DescriptionColumnWidth + 2 * NumberColumnWidth, InY, TEXT("Average"), HeaderColor);
-				RightJustify(InCanvas, StatsFont, InX + DescriptionColumnWidth + 3 * NumberColumnWidth, InY, TEXT("CPU"), HeaderColor);
-				InY += StatsFont->GetMaxCharHeight();
-
-				// Render rows
-				int32 ResultIndex = 0;
-				const FLinearColor StatColor = FLinearColor(0.f, 1.f, 0.f);
-
-				for (const FRealtimeGPUProfilerDescriptionResult& PerfResult : PerfResults)
-				{
-					InCanvas->DrawTile(InX, InY, DescriptionColumnWidth + 3 * NumberColumnWidth, StatsFont->GetMaxCharHeight(),
-						0, 0, 1, 1,
-						(ResultIndex & 1) ? FLinearColor(0.02f, 0.02f, 0.02f, 0.88f) : FLinearColor(0.05f, 0.05f, 0.05f, 0.92f),
-						GWhiteTexture, true);
-
-					// Source GPU times are in microseconds, CPU times in seconds, so we need to divide one by 1000, and multiply the other by 1000
-					InCanvas->DrawShadowedString(InX, InY, *PerfResult.Description, StatsFont, StatColor);
-					RightJustify(InCanvas, StatsFont, InX + DescriptionColumnWidth + 1 * NumberColumnWidth, InY, *FString::Printf(TEXT("%d"), PerfResult.GPUMask.GetNative()), StatColor);
-					RightJustify(InCanvas, StatsFont, InX + DescriptionColumnWidth + 2 * NumberColumnWidth, InY, *FString::Printf(TEXT("%.2f"), PerfResult.AverageTime / 1000.f), StatColor);
-					RightJustify(InCanvas, StatsFont, InX + DescriptionColumnWidth + 3 * NumberColumnWidth, InY, *FString::Printf(TEXT("%.2f"), CpuPerfResults[ResultIndex] * 1000.f), StatColor);
-
-					InY += StatsFont->GetMaxCharHeight();
-
-					ResultIndex++;
-				}
-			}
-			else
-			{
-				InCanvas->DrawShadowedString(InX, InY, TEXT("Display Cluster Stats [NO DATA]"), StatsFont, HeaderColor);
-				InY += StatsFont->GetMaxCharHeight();
-			}
-
-			InY += StatsFont->GetMaxCharHeight();
-		}
-#endif  // GPUPROFILERTRACE_ENABLED
-
-		return InY;
-	}
-
 	float* GetNextHistoryWriteAddress(const FString& Description)
 	{
 		FRWScopeLock Lock(CpuHistoryMutex, SLT_Write);
diff --git a/Engine/Source/Programs/TraceAnalyzer/TraceAnalyzer.Target.cs b/Engine/Source/Programs/TraceAnalyzer/TraceAnalyzer.Target.cs
--- a/Engine/Source/Programs/TraceAnalyzer/TraceAnalyzer.Target.cs
+++ b/Engine/Source/Programs/TraceAnalyzer/TraceAnalyzer.Target.cs
@@ -33,7 +33,6 @@ public class TraceAnalyzerTarget : TargetRules
 
 		bEnableTrace = false; // UE_TRACE_ENABLED
 		GlobalDefinitions.Add("CPUPROFILERTRACE_ENABLED=0");
-		GlobalDefinitions.Add("GPUPROFILERTRACE_ENABLED=0");
 		GlobalDefinitions.Add("LOADTIMEPROFILERTRACE_ENABLED=0");
 		GlobalDefinitions.Add("STATSTRACE_ENABLED=0");
 		GlobalDefinitions.Add("COUNTERSTRACE_ENABLED=0");
diff --git a/Engine/Source/Programs/UnrealInsights/UnrealInsights.Target.cs b/Engine/Source/Programs/UnrealInsights/UnrealInsights.Target.cs
--- a/Engine/Source/Programs/UnrealInsights/UnrealInsights.Target.cs
+++ b/Engine/Source/Programs/UnrealInsights/UnrealInsights.Target.cs
@@ -71,8 +71,7 @@ public class UnrealInsightsTarget : TargetRules
 			//GlobalDefinitions.Add("LOGTRACE_ENABLED=0");
 			//GlobalDefinitions.Add("MISCTRACE_ENABLED=0");
 			//GlobalDefinitions.Add("CPUPROFILERTRACE_ENABLED=0");
-			//GlobalDefinitions.Add("GPUPROFILERTRACE_ENABLED=0"); // old GPU Profiler
-			//GlobalDefinitions.Add("UE_TRACE_GPU_PROFILER_ENABLED=0"); // new GPU Profiler
+			//GlobalDefinitions.Add("UE_TRACE_GPU_PROFILER_ENABLED=0");
 			//GlobalDefinitions.Add("LOADTIMEPROFILERTRACE_ENABLED=0");
 			//GlobalDefinitions.Add("STATSTRACE_ENABLED=0");
 			//GlobalDefinitions.Add("COUNTERSTRACE_ENABLED=0");
diff --git a/Engine/Source/Runtime/Apple/MetalRHI/Private/MetalRHIRenderQuery.cpp b/Engine/Source/Runtime/Apple/MetalRHI/Private/MetalRHIRenderQuery.cpp
--- a/Engine/Source/Runtime/Apple/MetalRHI/Private/MetalRHIRenderQuery.cpp
+++ b/Engine/Source/Runtime/Apple/MetalRHI/Private/MetalRHIRenderQuery.cpp
@@ -35,7 +35,10 @@ FAutoConsoleVariableRef CVarQueryUseControlPrecisionBit(
 );
 
 FMetalQueryBuffer::FMetalQueryBuffer(FMetalQueryBufferPool* InPool, FMetalBufferPtr InBuffer)
+PRAGMA_DISABLE_DEPRECATION_WARNINGS
+	// @todo Metal RHI - Remove the use of RRT_TimestampCalibrationQuery. FMetalQueryBuffer probably shouldn't be an FRHIResource type.
 	: FRHIResource(RRT_TimestampCalibrationQuery)
+PRAGMA_ENABLE_DEPRECATION_WARNINGS
 	, Pool(InPool)
 	, Buffer(InBuffer)
 	, WriteOffset(0)
diff --git a/Engine/Source/Runtime/Core/Private/Stats/StatsCommand.cpp b/Engine/Source/Runtime/Core/Private/Stats/StatsCommand.cpp
--- a/Engine/Source/Runtime/Core/Private/Stats/StatsCommand.cpp
+++ b/Engine/Source/Runtime/Core/Private/Stats/StatsCommand.cpp
@@ -2144,8 +2144,7 @@ static void StatCmd(FString InCmd, bool bStatCommand, FOutputDevice* Ar /*= null
 			FStatSlowParams Params(Cmd);
 			FHUDGroupManager::Get(Stats).HandleCommand(Params, true);
 		}
-		// Skip this macro command if a real stat group called "GPU" exists. Required for legacy GPU profiler support until we permanently enable RHI_NEW_GPU_PROFILER.
-		else if (!Stats.Groups.Contains(TEXT("STATGROUP_GPU")) && FParse::Command(&Cmd, TEXT("gpu")))
+		else if (FParse::Command(&Cmd, TEXT("gpu")))
 		{
 			// Macro command to toggle all the stat groups that contain "STATCAT_GPU" stats.
 			if (FCoreDelegates::StatCheckEnabled.IsBound())
diff --git a/Engine/Source/Runtime/Engine/Classes/Engine/Engine.h b/Engine/Source/Runtime/Engine/Classes/Engine/Engine.h
--- a/Engine/Source/Runtime/Engine/Classes/Engine/Engine.h
+++ b/Engine/Source/Runtime/Engine/Classes/Engine/Engine.h
@@ -2420,10 +2420,6 @@ public:
 	ENGINE_API bool HandleDumpParticleMemCommand( const TCHAR* Cmd, FOutputDevice& Ar );
 #endif
 
-#if WITH_PROFILEGPU && (RHI_NEW_GPU_PROFILER == 0)
-	ENGINE_API bool HandleProfileGPUCommand( const TCHAR* Cmd, FOutputDevice& Ar );	
-#endif
-
 #if WITH_DUMPGPU
 	ENGINE_API bool HandleDumpGPUCommand( const TCHAR* Cmd, FOutputDevice& Ar );
 #endif
@@ -2455,8 +2451,6 @@ public:
 #if WITH_EDITOR
 	ENGINE_API bool HandleDumpShaderCompileStatsCommand( const TCHAR* Cmd, FOutputDevice& Ar );
 #endif
-	ENGINE_API bool HandleProfileCommand( const TCHAR* Cmd, FOutputDevice& Ar );
-	ENGINE_API bool HandleProfileGPUHitchesCommand( const TCHAR* Cmd, FOutputDevice& Ar );
 	ENGINE_API bool HandleShaderComplexityCommand( const TCHAR* Cmd, FOutputDevice& Ar );
 	ENGINE_API bool HandleFreezeRenderingCommand( const TCHAR* Cmd, FOutputDevice& Ar, UWorld* InWorld );
 	ENGINE_API bool HandleStartFPSChartCommand( const TCHAR* Cmd, FOutputDevice& Ar );
diff --git a/Engine/Source/Runtime/Engine/Private/StatsRender2.cpp b/Engine/Source/Runtime/Engine/Private/StatsRender2.cpp
--- a/Engine/Source/Runtime/Engine/Private/StatsRender2.cpp
+++ b/Engine/Source/Runtime/Engine/Private/StatsRender2.cpp
@@ -782,7 +782,6 @@ void RenderArrayOfStats( FCanvas* Canvas, const int32 X, int32& Y, const TArray<
 	}
 }
 
-#if RHI_NEW_GPU_PROFILER
 struct FGpuStatMessage
 {
 	const FComplexStatMessage* Busy;
@@ -917,7 +916,6 @@ void RenderArrayOfGpuStats(FCanvas* Canvas, const int32 X, int32& Y, const TMap<
 		Y += RenderMoreStatsLine(Canvas, X, Y, Aggregates.Num() - MaxStatsPerGroup);
 	}
 }
-#endif
 
 static int32 RenderFlatCycle( const FGameThreadStatsData& ViewData, const FComplexStatMessage& Item, class FCanvas* Canvas, const int32 X, const int32 Y, const  float Budget, const bool bIsBudgetIgnored )
 {
@@ -1025,7 +1023,6 @@ static void RenderGroupedWithHierarchy(const FGameThreadStatsData& ViewData, FVi
 		}
 
 		// Render GPU stats
-	#if RHI_NEW_GPU_PROFILER
 		if (StatGroup.GpuStatsAggregate.Num())
 		{
 			// Group together the individual GPU Busy, Wait and Idle stats 
@@ -1052,7 +1049,6 @@ static void RenderGroupedWithHierarchy(const FGameThreadStatsData& ViewData, FVi
 			RenderArrayOfGpuStats(Canvas, X, Y, GroupedStats, ViewData);
 			Y += Globals.GetFontHeight();
 		}
-	#endif
 	}
 }
 
diff --git a/Engine/Source/Runtime/Engine/Private/UnrealEngine.cpp b/Engine/Source/Runtime/Engine/Private/UnrealEngine.cpp
--- a/Engine/Source/Runtime/Engine/Private/UnrealEngine.cpp
+++ b/Engine/Source/Runtime/Engine/Private/UnrealEngine.cpp
@@ -5394,13 +5394,6 @@ bool UEngine::Exec( UWorld* InWorld, const TCHAR* Cmd, FOutputDevice& Ar )
 	}
 #endif
 
-#if WITH_PROFILEGPU && (RHI_NEW_GPU_PROFILER == 0)
-	else if( FParse::Command(&Cmd,TEXT("PROFILEGPU")) )
-	{
-		return HandleProfileGPUCommand( Cmd, Ar );
-	}	
-#endif
-
 #if WITH_DUMPGPU
 	else if (FParse::Command(&Cmd, TEXT("DUMPGPU")))
 	{
@@ -5520,14 +5513,6 @@ bool UEngine::Exec_Dev( UWorld* InWorld, const TCHAR* Cmd, FOutputDevice& Ar )
 		extern bool HandleVisualizeRT();
 		return HandleVisualizeRT();
 	}
-	else if( FParse::Command(&Cmd,TEXT("PROFILE")) )
-	{
-		return HandleProfileCommand( Cmd, Ar );
-	}
-	else if( FParse::Command(&Cmd,TEXT("PROFILEGPUHITCHES")) )				    
-	{
-		return HandleProfileGPUHitchesCommand( Cmd, Ar );
-	}
 	else if( FParse::Command(&Cmd,TEXT("SHADERCOMPLEXITY")) )
 	{
 		return HandleShaderComplexityCommand( Cmd, Ar );
@@ -6250,25 +6235,6 @@ bool UEngine::HandleShaderComplexityCommand( const TCHAR* Cmd, FOutputDevice& Ar
 	return true; 
 }
 
-bool UEngine::HandleProfileGPUHitchesCommand( const TCHAR* Cmd, FOutputDevice& Ar )
-{
-	if (!FApp::CanEverRender())
-	{
-		return true;
-	}
-
-	GTriggerGPUHitchProfile = !GTriggerGPUHitchProfile;
-	if (GTriggerGPUHitchProfile)
-	{
-		Ar.Logf(TEXT("Profiling GPU hitches."));
-	}
-	else
-	{
-		Ar.Logf(TEXT("Stopped profiling GPU hitches."));
-	}
-	return true;
-}
-
 bool UEngine::HandleToggleRenderingThreadCommand( const TCHAR* Cmd, FOutputDevice& Ar )
 {
 	check(IsInGameThread());
@@ -6328,76 +6294,8 @@ bool UEngine::HandleDumpShaderCompileStatsCommand(const TCHAR* Cmd, FOutputDevic
 }
 #endif
 
-bool UEngine::HandleProfileCommand( const TCHAR* Cmd, FOutputDevice& Ar )
-{
-#if WITH_PROFILEGPU && (RHI_NEW_GPU_PROFILER == 0)
-	if ( FParse::Command(&Cmd,TEXT("GPU")) )
-	{
-		if (!FApp::CanEverRender())
-		{
-			return true;
-		}
-
-		if (!GTriggerGPUHitchProfile)
-		{
-			ENQUEUE_RENDER_COMMAND(HandleProfileCommand)(
-				[](FRHICommandListImmediate& RHICmdList)
-			{
-				RHICmdList.ImmediateFlush(EImmediateFlushType::FlushRHIThread);
-				GTriggerGPUProfile = true;
-			});
-			Ar.Logf(TEXT("Profiling the next GPU frame"));
-		}
-		else
-		{
-			Ar.Logf(TEXT("Can't do a gpu profile during a hitch profile!"));
-		}
-		return true;
-	}
-#endif
-	return false;
-}
-
 #endif // !(UE_BUILD_SHIPPING || UE_BUILD_TEST)
 
-#if WITH_PROFILEGPU && (RHI_NEW_GPU_PROFILER == 0)
-bool UEngine::HandleProfileGPUCommand( const TCHAR* Cmd, FOutputDevice& Ar )
-{
-	if (!FApp::CanEverRender())
-	{
-		return true;
-	}
-
-	if (FParse::Command(&Cmd, TEXT("TRACE")))
-	{
-		FString Filename = CreateProfileDirectoryAndFilename(TEXT(""), TEXT(".rtt"));
-		//FPaths::MakePlatformFilename(Filename);
-		GGPUTraceFileName = Filename;
-		Ar.Logf(TEXT("Tracing the next GPU frame"));
-	}
-	else
-	{
-		if (!GTriggerGPUHitchProfile)
-		{
-			ENQUEUE_RENDER_COMMAND(HandleProfileGPUCommand)(
-				[](FRHICommandListImmediate& RHICmdList)
-			{
-				RHICmdList.ImmediateFlush(EImmediateFlushType::FlushRHIThread);
-				GTriggerGPUProfile = true;
-			});
-			Ar.Logf(TEXT("Profiling the next GPU frame"));
-		}
-		else
-		{
-			Ar.Logf(TEXT("Can't do a gpu profile during a hitch profile!"));
-		}
-	}
-
-	return true;
-}
-
-#endif // WITH_PROFILEGPU && (RHI_NEW_GPU_PROFILER == 0)
-
 #if WITH_DUMPGPU
 
 bool UEngine::HandleDumpGPUCommand(const TCHAR* Cmd, FOutputDevice& Ar)
@@ -19288,8 +19186,7 @@ int32 UEngine::RenderStatDrawCount(UWorld* World, FViewport* Viewport, FCanvas*
 #if HAS_GPU_STATS
 
 	int32 TotalCount[MAX_NUM_GPUS] = { 0 };
-	// Display all the categories of draw counts. This may always report 0 in some modes if AreGPUStatsEnabled is not enabled.
-	// Most likely because we are not currently capturing a CSV.
+	// Display all the categories of draw counts.
 	FRHIDrawStatsCategory::FManager const& Manager = FRHIDrawStatsCategory::GetManager();
 	for (int32 Index = 0; Index < Manager.NumCategory; ++Index)
 	{
diff --git a/Engine/Source/Runtime/Launch/Private/LaunchEngineLoop.cpp b/Engine/Source/Runtime/Launch/Private/LaunchEngineLoop.cpp
--- a/Engine/Source/Runtime/Launch/Private/LaunchEngineLoop.cpp
+++ b/Engine/Source/Runtime/Launch/Private/LaunchEngineLoop.cpp
@@ -5152,10 +5152,6 @@ void FEngineLoop::Exit()
 
 	IStreamingManager::Shutdown();
 
-#if HAS_GPU_STATS && (RHI_NEW_GPU_PROFILER == 0)
-	FRealtimeGPUProfiler::SafeRelease();
-#endif
-
 	DestroyMoviePlayer();
 
 	// Move earlier?
@@ -5467,19 +5463,6 @@ static inline void BeginFrameRenderThread(FRHICommandListImmediate& RHICmdList,
 	TRACE_BEGIN_FRAME(TraceFrameType_Rendering);
 	GFrameNumberRenderThread++;
 
-#if !UE_BUILD_SHIPPING 
-
-		// If we are profiling, kick off a long GPU task to make the GPU always behind the CPU so that we
-		// won't get GPU idle time measured in profiling results
-	#if WITH_PROFILEGPU && (RHI_NEW_GPU_PROFILER == 0)
-		if (GTriggerGPUProfile && !GTriggerGPUHitchProfile)
-		{
-			IssueScalableLongGPUTask(RHICmdList);
-		}
-	#endif
-
-#endif // !UE_BUILD_SHIPPING
-
 #if WITH_RHI_BREADCRUMBS
 #if CSV_PROFILER
 	if (FCsvProfiler::Get()->IsCapturing_Renderthread())
@@ -5667,21 +5650,6 @@ void FEngineLoop::Tick()
 	{
 		SCOPE_CYCLE_COUNTER(STAT_FrameTime);
 
-		#if WITH_PROFILEGPU && !UE_BUILD_SHIPPING && (RHI_NEW_GPU_PROFILER == 0)
-			// Issue the measurement of the execution time of a basic LongGPUTask unit on the very first frame
-			// The results will be retrived on the first call of IssueScalableLongGPUTask
-			if (GFrameCounter == 0 && IsFeatureLevelSupported(GMaxRHIShaderPlatform, ERHIFeatureLevel::SM5) && FApp::CanEverRender())
-			{
-				FlushRenderingCommands();
-
-				ENQUEUE_RENDER_COMMAND(MeasureLongGPUTaskExecutionTimeCmd)(
-					[](FRHICommandListImmediate& RHICmdList)
-					{
-						MeasureLongGPUTaskExecutionTime(RHICmdList);
-					});
-			}
-		#endif
-
 #if WITH_ENGINE && CSV_PROFILER
 		UpdateCoreCsvStats_BeginFrame();
 #endif
diff --git a/Engine/Source/Runtime/NullDrv/Private/NullRHI.cpp b/Engine/Source/Runtime/NullDrv/Private/NullRHI.cpp
--- a/Engine/Source/Runtime/NullDrv/Private/NullRHI.cpp
+++ b/Engine/Source/Runtime/NullDrv/Private/NullRHI.cpp
@@ -9,9 +9,7 @@
 #include "DataDrivenShaderPlatformInfo.h"
 #include "GPUProfiler.h"
 
-#if RHI_NEW_GPU_PROFILER
 constexpr UE::RHI::GPUProfiler::FQueue NullRHIGPUProfilerQueue(UE::RHI::GPUProfiler::FQueue::EType::Graphics, 0, 0);
-#endif
 
 FNullDynamicRHI::FNullDynamicRHI()
 {
@@ -62,10 +60,8 @@ void FNullDynamicRHI::Init()
 
 	GIsRHIInitialized = true;
 
-#if RHI_NEW_GPU_PROFILER
 	// Register a single dummy graphics GPU queue. This is required so that we can still submit the frame boundary events in RHIEndFrame.
 	UE::RHI::GPUProfiler::InitializeQueues(MakeConstArrayView(&NullRHIGPUProfilerQueue, 1));
-#endif
 }
 
 
@@ -108,9 +104,8 @@ void* FNullDynamicRHI::GetStaticTextureBuffer(int32 SizeX, int32 SizeY, EPixelFo
 void FNullDynamicRHI::RHIEndFrame(const FRHIEndFrameArgs& Args)
 {
 	// The RHI must publish the frame boundary events for engine stats to work correctly.
-
-#if RHI_NEW_GPU_PROFILER
 	UE::RHI::GPUProfiler::FEventStream EventStream(NullRHIGPUProfilerQueue);
+
 	EventStream.Emplace<UE::RHI::GPUProfiler::FEvent::FFrameBoundary>(FPlatformTime::Cycles64(), Args.FrameNumber
 	#if WITH_RHI_BREADCRUMBS
 		, Args.GPUBreadcrumbs[ERHIPipeline::Graphics]
@@ -121,7 +116,6 @@ void FNullDynamicRHI::RHIEndFrame(const FRHIEndFrameArgs& Args)
 	);
 
 	UE::RHI::GPUProfiler::ProcessEvents(MakeArrayView(&EventStream, 1));
-#endif
 }
 
 /** Value between 0-100 that determines the percentage of the vertical scan that is allowed to pass while still allowing us to swap when VSYNC'ed.
diff --git a/Engine/Source/Runtime/RHI/Private/GPUProfiler.cpp b/Engine/Source/Runtime/RHI/Private/GPUProfiler.cpp
--- a/Engine/Source/Runtime/RHI/Private/GPUProfiler.cpp
+++ b/Engine/Source/Runtime/RHI/Private/GPUProfiler.cpp
@@ -113,648 +113,6 @@ static TAutoConsoleVariable<int> CVarGPUCsvStatsEnabled(
 	0,
 	TEXT("Enables or disables GPU stat recording to CSVs"));
 
-#if (RHI_NEW_GPU_PROFILER == 0)
-
-static TAutoConsoleVariable<FString> GProfileGPUPatternCVar(
-	TEXT("r.ProfileGPU.Pattern"),
-	TEXT("*"),
-	TEXT("Allows to filter the entries when using ProfileGPU, the pattern match is case sensitive.\n")
-	TEXT("'*' can be used in the end to get all entries starting with the string.\n")
-	TEXT("    '*' without any leading characters disables the pattern matching and uses a time threshold instead (default).\n")
-	TEXT("'?' allows to ignore one character.\n")
-	TEXT("e.g. AmbientOcclusionSetup, AmbientOcclusion*, Ambient???lusion*, *"),
-	ECVF_Default);
-
-static TAutoConsoleVariable<int32> GProfileShowEventHistogram(
-	TEXT("r.ProfileGPU.ShowEventHistogram"),
-	0,
-	TEXT("Whether the event histogram should be shown."),
-	ECVF_Default);
-
-TAutoConsoleVariable<int32> GProfileGPUTransitions(
-	TEXT("r.ProfileGPU.ShowTransitions"),
-	0,
-	TEXT("Allows profileGPU to display resource transition events."),
-	ECVF_Default);
-
-// Should we print a summary at the end?
-static TAutoConsoleVariable<int32> GProfilePrintAssetSummary(
-	TEXT("r.ProfileGPU.PrintAssetSummary"),
-	0,
-	TEXT("Should we print a summary split by asset (r.ShowMaterialDrawEvents is strongly recommended as well).\n"),
-	ECVF_Default);
-
-// Should we print a summary at the end?
-static TAutoConsoleVariable<FString> GProfileAssetSummaryCallOuts(
-	TEXT("r.ProfileGPU.AssetSummaryCallOuts"),
-	TEXT(""),
-	TEXT("Comma separated list of substrings that deserve special mention in the final summary (e.g., \"LOD,HeroName\"\n")
-	TEXT("r.ProfileGPU.PrintAssetSummary must be true to enable this feature"),
-	ECVF_Default);
-
-static TAutoConsoleVariable<int32> GSaveScreenshotAfterProfilingGPUCVar(
-	TEXT("r.ProfileGPU.Screenshot"),
-	1,
-	TEXT("Whether a screenshot should be taken when profiling the GPU. 0:off, 1:on (default)"),
-	ECVF_RenderThreadSafe);
-
-static TAutoConsoleVariable<float> GGPUHitchThresholdCVar(
-	TEXT("RHI.GPUHitchThreshold"),
-	100.0f,
-	TEXT("Threshold for detecting hitches on the GPU (in milliseconds).")
-);
-
-static TAutoConsoleVariable<int32> CVarGPUCrashDataCollectionEnable(
-	TEXT("r.gpucrash.collectionenable"),
-	1,
-	TEXT("Stores GPU crash data from scoped events when a applicable crash debugging system is available."),
-	ECVF_RenderThreadSafe);
-
-static TAutoConsoleVariable<int32> CVarGPUCrashDataDepth(
-	TEXT("r.gpucrash.datadepth"),
-	-1,
-	TEXT("Limits the amount of marker scope depth we record for GPU crash debugging to the given scope depth."),
-	ECVF_RenderThreadSafe);
-
-namespace RHIConfig
-{
-	bool ShouldSaveScreenshotAfterProfilingGPU()
-	{
-		return GSaveScreenshotAfterProfilingGPUCVar.GetValueOnAnyThread() != 0;
-	}
-
-	bool ShouldShowProfilerAfterProfilingGPU()
-	{
-		return GCVarProfileGPU_ShowUI.GetValueOnAnyThread() != 0;
-	}
-
-	float GetGPUHitchThreshold()
-	{
-		return GGPUHitchThresholdCVar.GetValueOnAnyThread() * 0.001f;
-	}
-}
-
-/** Recursively generates a histogram of nodes and stores their timing in TimingResult. */
-static void GatherStatsEventNode(FGPUProfilerEventNode* Node, int32 Depth, TMap<FString, FGPUProfilerEventNodeStats>& EventHistogram)
-{
-	if (Node->NumDraws > 0 || Node->NumDispatches > 0 || Node->Children.Num() > 0)
-	{
-		Node->TimingResult = Node->GetTiming() * 1000.0f;
-		Node->NumTotalDraws = Node->NumDraws;
-		Node->NumTotalDispatches = Node->NumDispatches;
-		Node->NumTotalPrimitives = Node->NumPrimitives;
-		Node->NumTotalVertices = Node->NumVertices;
-
-		FGPUProfilerEventNode* Parent = Node->Parent;		
-		while (Parent)
-		{
-			Parent->NumTotalDraws += Node->NumDraws;
-			Parent->NumTotalDispatches += Node->NumDispatches;
-			Parent->NumTotalPrimitives += Node->NumPrimitives;
-			Parent->NumTotalVertices += Node->NumVertices;
-
-			Parent = Parent->Parent;
-		}
-
-		for (int32 ChildIndex = 0; ChildIndex < Node->Children.Num(); ChildIndex++)
-		{
-			// Traverse children
-			GatherStatsEventNode(Node->Children[ChildIndex], Depth + 1, EventHistogram);
-		}
-
-		FGPUProfilerEventNodeStats* FoundHistogramBucket = EventHistogram.Find(Node->Name);
-		if (FoundHistogramBucket)
-		{
-			FoundHistogramBucket->NumDraws += Node->NumTotalDraws;
-			FoundHistogramBucket->NumPrimitives += Node->NumTotalPrimitives;
-			FoundHistogramBucket->NumVertices += Node->NumTotalVertices;
-			FoundHistogramBucket->TimingResult += Node->TimingResult;
-			FoundHistogramBucket->NumEvents++;
-		}
-		else
-		{
-			FGPUProfilerEventNodeStats NewNodeStats;
-			NewNodeStats.NumDraws = Node->NumTotalDraws;
-			NewNodeStats.NumPrimitives = Node->NumTotalPrimitives;
-			NewNodeStats.NumVertices = Node->NumTotalVertices;
-			NewNodeStats.TimingResult = Node->TimingResult;
-			NewNodeStats.NumEvents = 1;
-			EventHistogram.Add(Node->Name, NewNodeStats);
-		}
-	}
-}
-
-struct FGPUProfileInfoPair
-{
-	int64 Triangles;
-	int32 DrawCalls;
-
-	FGPUProfileInfoPair()
-		: Triangles(0)
-		, DrawCalls(0)
-	{
-	}
-
-	void AddDraw(int64 InTriangleCount)
-	{
-		Triangles += InTriangleCount;
-		++DrawCalls;
-	}
-};
-
-struct FGPUProfileStatSummary
-{
-	TMap<FString, FGPUProfileInfoPair> TrianglesPerMaterial;
-	TMap<FString, FGPUProfileInfoPair> TrianglesPerMesh;
-	TMap<FString, FGPUProfileInfoPair> TrianglesPerNonMesh;
-
-	int32 TotalNumNodes;
-	int32 TotalNumDraws;
-
-	bool bGatherSummaryStats;
-	bool bDumpEventLeafNodes;
-
-	FGPUProfileStatSummary()
-		: TotalNumNodes(0)
-		, TotalNumDraws(0)
-		, bGatherSummaryStats(false)
-		, bDumpEventLeafNodes(false)
-	{
-		bDumpEventLeafNodes = GCVarProfileGPU_ShowLeafEvents.GetValueOnRenderThread() != 0;
-		bGatherSummaryStats = GProfilePrintAssetSummary.GetValueOnRenderThread() != 0;
-	}
-
-	void ProcessMatch(FGPUProfilerEventNode* Node)
-	{
-		if (bGatherSummaryStats && (Node->NumTotalPrimitives > 0) && (Node->NumTotalVertices > 0) && (Node->Children.Num() == 0))
-		{
-			FString MaterialPart;
-			FString AssetPart;
-			if (Node->Name.Split(TEXT(" "), &MaterialPart, &AssetPart, ESearchCase::CaseSensitive))
-			{
-				TrianglesPerMaterial.FindOrAdd(MaterialPart).AddDraw(Node->NumTotalPrimitives);
-				TrianglesPerMesh.FindOrAdd(AssetPart).AddDraw(Node->NumTotalPrimitives);
-			}
-			else
-			{
-				TrianglesPerNonMesh.FindOrAdd(Node->Name).AddDraw(Node->NumTotalPrimitives);
-			}
-		}
-	}
-
-	void PrintSummary()
-	{
-		UE_LOG(LogRHI, Log, TEXT("Total Nodes %u Draws %u"), TotalNumNodes, TotalNumDraws);
-		UE_LOG(LogRHI, Log, TEXT(""));
-		UE_LOG(LogRHI, Log, TEXT(""));
-
-		if (bGatherSummaryStats)
-		{
-			// Sort the lists and print them out
-			TrianglesPerMesh.ValueSort([](const FGPUProfileInfoPair& A, const FGPUProfileInfoPair& B){ return A.Triangles > B.Triangles; });
-			UE_LOG(LogRHI, Log, TEXT(""));
-			UE_LOG(LogRHI, Log, TEXT("MeshList,TriangleCount,DrawCallCount"));
-			for (auto& Pair : TrianglesPerMesh)
-			{
-				UE_LOG(LogRHI, Log, TEXT("%s,%" INT64_FMT ",%d"), *Pair.Key, Pair.Value.Triangles, Pair.Value.DrawCalls);
-			}
-
-			TrianglesPerMaterial.ValueSort([](const FGPUProfileInfoPair& A, const FGPUProfileInfoPair& B){ return A.Triangles > B.Triangles; });
-			UE_LOG(LogRHI, Log, TEXT(""));
-			UE_LOG(LogRHI, Log, TEXT("MaterialList,TriangleCount,DrawCallCount"));
-			for (auto& Pair : TrianglesPerMaterial)
-			{
-				UE_LOG(LogRHI, Log, TEXT("%s,%" INT64_FMT ",%d"), *Pair.Key, Pair.Value.Triangles, Pair.Value.DrawCalls);
-			}
-
-			TrianglesPerNonMesh.ValueSort([](const FGPUProfileInfoPair& A, const FGPUProfileInfoPair& B){ return A.Triangles > B.Triangles; });
-			UE_LOG(LogRHI, Log, TEXT(""));
-			UE_LOG(LogRHI, Log, TEXT("MiscList,TriangleCount,DrawCallCount"));
-			for (auto& Pair : TrianglesPerNonMesh)
-			{
-				UE_LOG(LogRHI, Log, TEXT("%s,%" INT64_FMT ",%d"), *Pair.Key, Pair.Value.Triangles, Pair.Value.DrawCalls);
-			}
-
-			// See if we want to call out any particularly interesting matches
-			TArray<FString> InterestingSubstrings;
-			GProfileAssetSummaryCallOuts.GetValueOnRenderThread().ParseIntoArray(InterestingSubstrings, TEXT(","), true);
-
-			if (InterestingSubstrings.Num() > 0)
-			{
-				UE_LOG(LogRHI, Log, TEXT(""));
-				UE_LOG(LogRHI, Log, TEXT("Information about specified mesh substring matches (r.ProfileGPU.AssetSummaryCallOuts)"));
-				for (const FString& InterestingSubstring : InterestingSubstrings)
-				{
-					int32 InterestingNumDraws = 0;
-					int64 InterestingNumTriangles = 0;
-
-					for (auto& Pair : TrianglesPerMesh)
-					{
-						if (Pair.Key.Contains(InterestingSubstring))
-						{
-							InterestingNumDraws += Pair.Value.DrawCalls;
-							InterestingNumTriangles += Pair.Value.Triangles;
-						}
-					}
-
-					UE_LOG(LogRHI, Log, TEXT("Matching '%s': %d draw calls, with %" INT64_FMT " tris (%.2f M)"), *InterestingSubstring, InterestingNumDraws, InterestingNumTriangles, InterestingNumTriangles * 1e-6);
-				}
-				UE_LOG(LogRHI, Log, TEXT(""));
-			}
-		}
-	}
-};
-
-/** Recursively dumps stats for each node with a depth first traversal. */
-static void DumpStatsEventNode(FGPUProfilerEventNode* Node, float RootResult, int32 Depth, const FWildcardString& WildcardFilter, bool bParentMatchedFilter, float& ReportedTiming, FGPUProfileStatSummary& Summary)
-{
-	Summary.TotalNumNodes++;
-	ReportedTiming = 0;
-
-	if (Node->NumDraws > 0 || Node->NumDispatches > 0 || Node->Children.Num() > 0 || Summary.bDumpEventLeafNodes)
-	{
-		Summary.TotalNumDraws += Node->NumDraws;
-		// Percent that this node was of the total frame time
-		const float Percent = Node->TimingResult * 100.0f / (RootResult * 1000.0f);
-		const float PercentThreshold = GCVarProfileGPU_ThresholdPercent.GetValueOnRenderThread();
-		const int32 EffectiveDepth = FMath::Max(Depth - 1, 0);
-		const bool bDisplayEvent = (bParentMatchedFilter || WildcardFilter.IsMatch(Node->Name)) && (Percent > PercentThreshold || Summary.bDumpEventLeafNodes);
-
-		if (bDisplayEvent)
-		{
-			FString NodeStats;
-
-			if (Node->NumTotalDraws > 0)
-			{
-				NodeStats = FString::Printf(TEXT("%u %s %u prims %u verts "), Node->NumTotalDraws, Node->NumTotalDraws == 1 ? TEXT("draw") : TEXT("draws"), Node->NumTotalPrimitives, Node->NumTotalVertices);
-			}
-
-			if (Node->NumTotalDispatches > 0)
-			{
-				NodeStats += FString::Printf(TEXT("%u %s"), Node->NumTotalDispatches, Node->NumTotalDispatches == 1 ? TEXT("dispatch") : TEXT("dispatches"));
-			
-				// Cumulative group stats are not meaningful, only include dispatch stats if there was one in the current node
-				if (Node->GroupCount.X > 0 && Node->NumDispatches == 1)
-				{
-					NodeStats += FString::Printf(TEXT(" %u"), Node->GroupCount.X);
-
-					if (Node->GroupCount.Y > 1)
-					{
-						NodeStats += FString::Printf(TEXT("x%u"), Node->GroupCount.Y);
-					}
-
-					if (Node->GroupCount.Z > 1)
-					{
-						NodeStats += FString::Printf(TEXT("x%u"), Node->GroupCount.Z);
-					}
-
-					NodeStats += TEXT(" groups");
-				}
-			}
-
-			// Print information about this node, padded to its depth in the tree
-			UE_LOG(LogRHI, Log, TEXT("%s%4.1f%%%5.2fms   %s %s"), 
-				*FString().LeftPad(EffectiveDepth * 3), 
-				Percent,
-				Node->TimingResult,
-				*Node->Name,
-				*NodeStats
-				);
-
-			ReportedTiming = Node->TimingResult;
-			Summary.ProcessMatch(Node);
-		}
-
-		struct FCompareGPUProfileNode
-		{
-			EGPUProfileSortMode SortMode;
-			FCompareGPUProfileNode(EGPUProfileSortMode InSortMode)
-				: SortMode(InSortMode)
-			{}
-			FORCEINLINE bool operator()(const FGPUProfilerEventNode* A, const FGPUProfilerEventNode* B) const
-			{
-				switch (SortMode)
-				{
-					case EGPUProfileSortMode::NumPrims:
-						return B->NumTotalPrimitives < A->NumTotalPrimitives;
-					case EGPUProfileSortMode::NumVerts:
-						return B->NumTotalVertices < A->NumTotalVertices;
-					case EGPUProfileSortMode::TimeElapsed:
-					default:
-						return B->TimingResult < A->TimingResult;
-				}
-			}
-		};
-
-		EGPUProfileSortMode SortMode = (EGPUProfileSortMode)FMath::Clamp(GCVarProfileGPU_Sort.GetValueOnRenderThread(), 0, ((int32)EGPUProfileSortMode::Max - 1));
-		if (SortMode != EGPUProfileSortMode::Chronological)
-		{
-			Node->Children.Sort(FCompareGPUProfileNode(SortMode));
-		}
-
-		float TotalChildTime = 0;
-		uint32 TotalChildDraws = 0;
-		for (int32 ChildIndex = 0; ChildIndex < Node->Children.Num(); ChildIndex++)
-		{
-			FGPUProfilerEventNode* ChildNode = Node->Children[ChildIndex];
-
-			// Traverse children			
-			const int32 PrevNumDraws = Summary.TotalNumDraws;
-			float ChildReportedTiming = 0;
-			DumpStatsEventNode(Node->Children[ChildIndex], RootResult, Depth + 1, WildcardFilter, bDisplayEvent, ChildReportedTiming, Summary);
-			const int32 NumChildDraws = Summary.TotalNumDraws - PrevNumDraws;
-
-			TotalChildTime += ChildReportedTiming;
-			TotalChildDraws += NumChildDraws;
-		}
-
-		const float UnaccountedTime = FMath::Max(Node->TimingResult - TotalChildTime, 0.0f);
-		const float UnaccountedPercent = UnaccountedTime * 100.0f / (RootResult * 1000.0f);
-
-		// Add an 'Other Children' node if necessary to show time spent in the current node that is not in any of its children
-		if (bDisplayEvent && Node->Children.Num() > 0 && TotalChildDraws > 0 && (UnaccountedPercent > 2.0f || UnaccountedTime > .2f))
-		{
-			UE_LOG(LogRHI, Log, TEXT("%s%4.1f%%%5.2fms   Other Children"), 
-				*FString().LeftPad((EffectiveDepth + 1) * 3), 
-				UnaccountedPercent,
-				UnaccountedTime);
-		}
-	}
-}
-
-#if !UE_BUILD_SHIPPING
-
-/**
- * Converts GPU profile data to Visualizer data
- *
- * @param InProfileData GPU profile data
- * @param OutVisualizerData Visualizer data
- */
-static TSharedPtr< FVisualizerEvent > CreateVisualizerDataRecursively( const TRefCountPtr< class FGPUProfilerEventNode >& InNode, TSharedPtr< FVisualizerEvent > InParentEvent, const double InStartTimeMs, const double InTotalTimeMs )
-{
-	TSharedPtr< FVisualizerEvent > VisualizerEvent( new FVisualizerEvent( InStartTimeMs / InTotalTimeMs, InNode->TimingResult / InTotalTimeMs, InNode->TimingResult, 0, InNode->Name ) );
-	VisualizerEvent->ParentEvent = InParentEvent;
-
-	double ChildStartTimeMs = InStartTimeMs;
-	for( int32 ChildIndex = 0; ChildIndex < InNode->Children.Num(); ChildIndex++ )
-	{
-		TRefCountPtr< FGPUProfilerEventNode > ChildNode = InNode->Children[ ChildIndex ];
-		TSharedPtr< FVisualizerEvent > ChildEvent = CreateVisualizerDataRecursively( ChildNode, VisualizerEvent, ChildStartTimeMs, InTotalTimeMs );
-		VisualizerEvent->Children.Add( ChildEvent );
-
-		ChildStartTimeMs += ChildNode->TimingResult;
-	}
-
-	return VisualizerEvent;
-}
-
-/**
- * Converts GPU profile data to Visualizer data
- *
- * @param InProfileData GPU profile data
- * @param OutVisualizerData Visualizer data
- */
-static TSharedPtr< FVisualizerEvent > CreateVisualizerData( const TArray<TRefCountPtr<class FGPUProfilerEventNode> >& InProfileData )
-{
-	// Calculate total time first
-	double TotalTimeMs = 0.0;
-	for( int32 Index = 0; Index < InProfileData.Num(); ++Index )
-	{
-		TotalTimeMs += InProfileData[ Index ]->TimingResult;
-	}
-	
-	// Assumption: InProfileData contains only one (root) element. Otherwise an extra FVisualizerEvent root event is required.
-	TSharedPtr< FVisualizerEvent > DummyRoot;
-	// Recursively create visualizer event data.
-	TSharedPtr< FVisualizerEvent > StatEvents( CreateVisualizerDataRecursively( InProfileData[0], DummyRoot, 0.0, TotalTimeMs ) );
-	return StatEvents;
-}
-
-#endif
-
-void FGPUProfilerEventNodeFrame::DumpEventTree()
-{
-	if (EventTree.Num() > 0)
-	{
-		float RootResult = GetRootTimingResults();
-
-		FString ConfigString;
-
-		if (GCVarProfileGPU_Root.GetValueOnRenderThread() != TEXT("*"))
-		{
-			ConfigString += FString::Printf(TEXT("Root filter: %s "), *GCVarProfileGPU_Root.GetValueOnRenderThread());
-		}
-
-		if (GCVarProfileGPU_ThresholdPercent.GetValueOnRenderThread() > 0.0f)
-		{
-			ConfigString += FString::Printf(TEXT("Threshold: %.2f%% "), GCVarProfileGPU_ThresholdPercent.GetValueOnRenderThread());
-		}
-
-		if (ConfigString.Len() > 0)
-		{
-			ConfigString = FString(", ") + ConfigString;
-		}
-
-		UE_LOG(LogRHI, Log, TEXT("Perf marker hierarchy, total GPU time %.2fms%s"), RootResult * 1000.0f, *ConfigString);
-		UE_LOG(LogRHI, Log, TEXT(""));
-
-		// Display a warning if this is a GPU profile and the GPU was profiled with v-sync enabled
-		FText VsyncEnabledWarningText = FText::GetEmpty();
-		static IConsoleVariable* CVSyncVar = IConsoleManager::Get().FindConsoleVariable(TEXT("r.VSync"));
-		if (CVSyncVar->GetInt() != 0 && !PlatformDisablesVSync())
-		{
-			VsyncEnabledWarningText = LOCTEXT("GpuProfileVsyncEnabledWarning", "WARNING: This GPU profile was captured with v-sync enabled.  V-sync wait time may show up in any bucket, and as a result the data in this profile may be skewed. Please profile with v-sync disabled to obtain the most accurate data.");
-			UE_LOG(LogRHI, Log, TEXT("%s"), *(VsyncEnabledWarningText.ToString()));
-		}
-
-		LogDisjointQuery();
-
-		TMap<FString, FGPUProfilerEventNodeStats> EventHistogram;
-		for (int32 BaseNodeIndex = 0; BaseNodeIndex < EventTree.Num(); BaseNodeIndex++)
-		{
-			GatherStatsEventNode(EventTree[BaseNodeIndex], 0, EventHistogram);
-		}
-
-		FString RootWildcardString = GCVarProfileGPU_Root.GetValueOnRenderThread();
-		FWildcardString RootWildcard(RootWildcardString);
-
-		FGPUProfileStatSummary Summary;
-		for (int32 BaseNodeIndex = 0; BaseNodeIndex < EventTree.Num(); BaseNodeIndex++)
-		{
-			float Unused = 0;
-			DumpStatsEventNode(EventTree[BaseNodeIndex], RootResult, 0, RootWildcard, false, Unused, /*inout*/ Summary);
-		}
-		Summary.PrintSummary();
-
-		const bool bShowHistogram = GProfileShowEventHistogram.GetValueOnRenderThread() != 0;
-
-		if (RootWildcardString == TEXT("*") && bShowHistogram)
-		{
-			struct FNodeStatsCompare
-			{
-				/** Sorts nodes by descending durations. */
-				FORCEINLINE bool operator()(const FGPUProfilerEventNodeStats& A, const FGPUProfilerEventNodeStats& B) const
-				{
-					return B.TimingResult < A.TimingResult;
-				}
-			};
-
-			// Sort descending based on node duration
-			EventHistogram.ValueSort( FNodeStatsCompare() );
-
-			// Log stats about the node histogram
-			UE_LOG(LogRHI, Log, TEXT("Node histogram %u buckets"), EventHistogram.Num());
-
-			// bad: reading on render thread but we don't support ECVF_RenderThreadSafe on strings yet
-			// It's very unlikely to cause a problem as the cvar is only changes by the user.
-			FString WildcardString = GProfileGPUPatternCVar.GetValueOnRenderThread();
-
-			FGPUProfilerEventNodeStats Sum;
-
-			const float ThresholdInMS = 5.0f;
-
-			if(WildcardString == FString("*"))
-			{
-				// disable Wildcard functionality
-				WildcardString.Empty();
-			}
-
-			if(WildcardString.IsEmpty())
-			{
-				UE_LOG(LogRHI, Log, TEXT(" r.ProfileGPU.Pattern = '*' (using threshold of %g ms)"), ThresholdInMS);
-			}
-			else
-			{
-				UE_LOG(LogRHI, Log, TEXT(" r.ProfileGPU.Pattern = '%s' (not using time threshold)"), *WildcardString);
-			}
-
-			FWildcardString Wildcard(WildcardString);
-
-			int32 NumNotShown = 0;
-			for (TMap<FString, FGPUProfilerEventNodeStats>::TIterator It(EventHistogram); It; ++It)
-			{
-				const FGPUProfilerEventNodeStats& NodeStats = It.Value();
-
-				bool bDump = NodeStats.TimingResult > RootResult * ThresholdInMS;
-
-				if(!Wildcard.IsEmpty())
-				{
-					// if a Wildcard string was specified, we want to always dump all entries
-					bDump = Wildcard.IsMatch(*It.Key());
-				}
-
-				if (bDump)
-				{
-					UE_LOG(LogRHI, Log, TEXT("   %.2fms   %s   Events %u   Draws %u"), NodeStats.TimingResult, *It.Key(), NodeStats.NumEvents, NodeStats.NumDraws);
-					Sum += NodeStats;
-				}
-				else
-				{
-					NumNotShown++;
-				}
-			}
-
-			UE_LOG(LogRHI, Log, TEXT("   Total %.2fms   Events %u   Draws %u,    %u buckets not shown"), 
-				Sum.TimingResult, Sum.NumEvents, Sum.NumDraws, NumNotShown);
-		}
-
-#if !UE_BUILD_SHIPPING
-		// Create and display profile visualizer data
-		if (RHIConfig::ShouldShowProfilerAfterProfilingGPU())
-		{
-
-		// execute on main thread
-			{
-				struct FDisplayProfilerVisualizer
-				{
-					void Thread( TSharedPtr<FVisualizerEvent> InVisualizerData, const FText InVsyncEnabledWarningText )
-					{
-						static FName ProfileVisualizerModule("ProfileVisualizer");			
-						if (FModuleManager::Get().IsModuleLoaded(ProfileVisualizerModule))
-						{
-							IProfileVisualizerModule& ProfileVisualizer = FModuleManager::GetModuleChecked<IProfileVisualizerModule>(ProfileVisualizerModule);
-							// Display a warning if this is a GPU profile and the GPU was profiled with v-sync enabled (otherwise InVsyncEnabledWarningText is empty)
-							ProfileVisualizer.DisplayProfileVisualizer( InVisualizerData, TEXT("GPU"), InVsyncEnabledWarningText, FLinearColor::Red );
-						}
-					}
-				} DisplayProfilerVisualizer;
-
-				TSharedPtr<FVisualizerEvent> VisualizerData = CreateVisualizerData( EventTree );
-
-				DECLARE_CYCLE_STAT(TEXT("FSimpleDelegateGraphTask.DisplayProfilerVisualizer"),
-					STAT_FSimpleDelegateGraphTask_DisplayProfilerVisualizer,
-					STATGROUP_TaskGraphTasks);
-
-				FSimpleDelegateGraphTask::CreateAndDispatchWhenReady(
-					FSimpleDelegateGraphTask::FDelegate::CreateRaw(&DisplayProfilerVisualizer, &FDisplayProfilerVisualizer::Thread, VisualizerData, VsyncEnabledWarningText),
-					GET_STATID(STAT_FSimpleDelegateGraphTask_DisplayProfilerVisualizer), nullptr, ENamedThreads::GameThread
-				);
-			}
-
-			
-		}
-#endif
-	}
-}
-
-void FGPUProfiler::PushEvent(const TCHAR* Name, FColor Color)
-{
-	if (bTrackingEvents)
-	{
-		check(StackDepth >= 0);
-		StackDepth++;
-
-		check(IsInRenderingThread() || IsInRHIThread());
-		if (CurrentEventNode)
-		{
-			// Add to the current node's children
-			CurrentEventNode->Children.Add(CreateEventNode(Name, CurrentEventNode));
-			CurrentEventNode = CurrentEventNode->Children.Last();
-		}
-		else
-		{
-			// Add a new root node to the tree
-			CurrentEventNodeFrame->EventTree.Add(CreateEventNode(Name, NULL));
-			CurrentEventNode = CurrentEventNodeFrame->EventTree.Last();
-		}
-
-		check(CurrentEventNode);
-		// Start timing the current node
-		CurrentEventNode->StartTiming();
-	}
-}
-
-void FGPUProfiler::PopEvent()
-{
-	if (bTrackingEvents)
-	{
-		check(StackDepth >= 1);
-		StackDepth--;
-
-		check(CurrentEventNode && (IsInRenderingThread() || IsInRHIThread()));
-		// Stop timing the current node and move one level up the tree
-		CurrentEventNode->StopTiming();
-		CurrentEventNode = CurrentEventNode->Parent;
-	}
-}
-
-/** Whether GPU timing measurements are supported by the driver. */
-bool FGPUTiming::GIsSupported = false;
-
-/** Frequency for the timing values, in number of ticks per seconds, or 0 if the feature isn't supported. */
-TStaticArray<uint64, MAX_NUM_GPUS> FGPUTiming::GTimingFrequency(InPlace, 0);
-
-/**
-* Two timestamps performed on GPU and CPU at nearly the same time.
-* This can be used to visualize GPU and CPU timing events on the same timeline.
-*/
-TStaticArray<FGPUTimingCalibrationTimestamp, MAX_NUM_GPUS> FGPUTiming::GCalibrationTimestamp;
-
-/** Whether the static variables have been initialized. */
-bool FGPUTiming::GAreGlobalsInitialized = false;
-
-#else
-
 CSV_DEFINE_CATEGORY(RHI, true);
 
 #if UE_BUILD_SHIPPING
@@ -2774,8 +2132,6 @@ namespace UE::RHI::GPUProfiler
 	}
 }
 
-#endif // RHI_NEW_GPU_PROFILER
-
 RHI_API FRHIGPUFrameTimeHistory GRHIGPUFrameTimeHistory;
 
 FRHIGPUFrameTimeHistory::EResult FRHIGPUFrameTimeHistory::FState::PopFrameCycles(uint64& OutCycles64)
diff --git a/Engine/Source/Runtime/RHI/Private/GpuProfilerTrace.cpp b/Engine/Source/Runtime/RHI/Private/GpuProfilerTrace.cpp
--- a/Engine/Source/Runtime/RHI/Private/GpuProfilerTrace.cpp
+++ b/Engine/Source/Runtime/RHI/Private/GpuProfilerTrace.cpp
@@ -13,259 +13,11 @@
 #include "Trace/Detail/Important/ImportantLogScope.h"
 #include "Trace/Detail/LogScope.h"
 
-// Both (old and new) GPU profilers uses same "GPU" trace channel.
-#if GPUPROFILERTRACE_ENABLED || UE_TRACE_GPU_PROFILER_ENABLED
+#if UE_TRACE_GPU_PROFILER_ENABLED
 	UE_TRACE_CHANNEL_EXTERN(GpuChannel, RHI_API)
 	UE_TRACE_CHANNEL_DEFINE(GpuChannel)
 #endif
 
-// The old GpuProfilerTrace is deprecated in UE 5.6
-#if GPUPROFILERTRACE_ENABLED
-
-namespace GpuProfilerTrace
-{
-
-static TAutoConsoleVariable<int32> CVarGpuProfilerMaxEventBufferSizeKB(
-	TEXT("r.GpuProfilerMaxEventBufferSizeKB"),
-	32,
-	TEXT("Size of the scratch buffer in kB."),
-	ECVF_Default);
-
-
-struct FGpuTraceFrame
-{
-	int64							CalibrationBias;
-	FGPUTimingCalibrationTimestamp	Calibration;
-	uint64							TimestampBase;
-	uint64							LastTimestamp;
-	uint32							RenderingFrameNumber;
-	uint32							EventBufferSize;
-	bool							bActive;
-	uint8*							EventBuffer = nullptr;
-	uint32							MaxEventBufferSize = 0;
-};
-
-#if RHI_NEW_GPU_PROFILER
-FGpuTraceFrame GCurrentFrames[2];
-#else
-FGpuTraceFrame GCurrentFrame;
-#endif
-
-static TSet<uint32> GEventNames;
-
-// deprecated in UE 5.6
-UE_TRACE_EVENT_BEGIN(GpuProfiler, EventSpec, NoSync|Important)
-	UE_TRACE_EVENT_FIELD(uint32, EventType)
-	UE_TRACE_EVENT_FIELD(uint16[], Name)
-UE_TRACE_EVENT_END()
-
-// GPU Index 0 - deprecated in UE 5.6
-UE_TRACE_EVENT_BEGIN(GpuProfiler, Frame)
-	UE_TRACE_EVENT_FIELD(uint64, CalibrationBias)
-	UE_TRACE_EVENT_FIELD(uint64, TimestampBase)
-	UE_TRACE_EVENT_FIELD(uint32, RenderingFrameNumber)
-	UE_TRACE_EVENT_FIELD(uint8[], Data)
-UE_TRACE_EVENT_END()
-
-// GPU Index 1 - deprecated in UE 5.6
-UE_TRACE_EVENT_BEGIN(GpuProfiler, Frame2)
-	UE_TRACE_EVENT_FIELD(uint64, CalibrationBias)
-	UE_TRACE_EVENT_FIELD(uint64, TimestampBase)
-	UE_TRACE_EVENT_FIELD(uint32, RenderingFrameNumber)
-	UE_TRACE_EVENT_FIELD(uint8[], Data)
-UE_TRACE_EVENT_END()
-
-} // namespace GpuProfilerTrace
-
-#if RHI_NEW_GPU_PROFILER
-#define GPU_TRACE_ARG , uint32 GPUIndex
-#else
-#define GPU_TRACE_ARG
-#endif
-
-void FGpuProfilerTrace::BeginFrame(FGPUTimingCalibrationTimestamp& Calibration GPU_TRACE_ARG)
-{
-	using namespace GpuProfilerTrace;
-
-#if RHI_NEW_GPU_PROFILER
-	FGpuTraceFrame& GCurrentFrame = GCurrentFrames[GPUIndex];
-#endif
-
-	if (!bool(GpuChannel))
-	{
-		return;
-	}
-
-	GCurrentFrame.Calibration = Calibration;
-	ensure(GCurrentFrame.Calibration.CPUMicroseconds > 0 && GCurrentFrame.Calibration.GPUMicroseconds > 0);
-	GCurrentFrame.TimestampBase = 0;
-	GCurrentFrame.EventBufferSize = 0;
-	GCurrentFrame.bActive = true;
-
-	int32 NeededSize = CVarGpuProfilerMaxEventBufferSizeKB.GetValueOnAnyThread() * 1024;
-	if ((GCurrentFrame.MaxEventBufferSize != NeededSize) && (NeededSize > 0))
-	{
-		FMemory::Free(GCurrentFrame.EventBuffer);
-		GCurrentFrame.EventBuffer = (uint8*)FMemory::Malloc(NeededSize);
-		GCurrentFrame.MaxEventBufferSize = NeededSize;
-	}
-}
-
-void FGpuProfilerTrace::SpecifyEventByName(const FName& Name GPU_TRACE_ARG)
-{
-	using namespace GpuProfilerTrace;
-
-#if RHI_NEW_GPU_PROFILER
-	FGpuTraceFrame& GCurrentFrame = GCurrentFrames[GPUIndex];
-#endif
-
-	if (!GCurrentFrame.bActive)
-	{
-		return;
-	}
-
-	// This function is only called from FRealtimeGPUProfilerFrame::UpdateStats
-	// at the end of the frame, so the access to this container is thread safe
-
-	uint32 Index = Name.GetComparisonIndex().ToUnstableInt();
-	if (!GEventNames.Contains(Index))
-	{
-		GEventNames.Add(Index);
-
-		FString String = Name.ToString();
-		uint32 NameLength = String.Len() + 1;
-		static_assert(sizeof(TCHAR) == sizeof(uint16), "");
-
-		UE_TRACE_LOG(GpuProfiler, EventSpec, GpuChannel, NameLength * sizeof(uint16))
-			<< EventSpec.EventType(Index)
-			<< EventSpec.Name((const uint16*)(*String), NameLength);
-	}
-}
-
-void FGpuProfilerTrace::BeginEventByName(const FName& Name, uint32 FrameNumber, uint64 TimestampMicroseconds GPU_TRACE_ARG)
-{
-	using namespace GpuProfilerTrace;
-
-#if RHI_NEW_GPU_PROFILER
-	FGpuTraceFrame& GCurrentFrame = GCurrentFrames[GPUIndex];
-#endif
-
-	if (!GCurrentFrame.bActive)
-	{
-		return;
-	}
-
-	// Prevent buffer overrun
-	if (GCurrentFrame.EventBufferSize + 10 + sizeof(uint32) > GCurrentFrame.MaxEventBufferSize) // 10 is the max size that FTraceUtils::Encode7bit might use + some space for the FName index (uint32)
-	{
-		UE_LOG(LogRHI, Error, TEXT("GpuProfiler's scratch buffer is out of space for this frame (current size : %d kB). Dropping this frame. The size can be increased dynamically with the console variable r.GpuProfilerMaxEventBufferSizeKB"), GCurrentFrame.MaxEventBufferSize / 1024);
-
-		// Deactivate for the current frame to avoid errors while decoding an incomplete trace
-		GCurrentFrame.bActive = false;
-		return;
-	}
-
-	if (GCurrentFrame.TimestampBase == 0)
-	{
-		GCurrentFrame.TimestampBase = TimestampMicroseconds;
-		GCurrentFrame.LastTimestamp = GCurrentFrame.TimestampBase;
-		GCurrentFrame.RenderingFrameNumber = FrameNumber;
-		if (!GCurrentFrame.Calibration.GPUMicroseconds)
-		{
-			GCurrentFrame.Calibration.GPUMicroseconds = TimestampMicroseconds;
-		}
-	}
-	uint8* BufferPtr = GCurrentFrame.EventBuffer + GCurrentFrame.EventBufferSize;
-	uint64 TimestampDelta = TimestampMicroseconds - GCurrentFrame.LastTimestamp;
-	GCurrentFrame.LastTimestamp = TimestampMicroseconds;
-	FTraceUtils::Encode7bit((TimestampDelta << 1ull) | 0x1, BufferPtr);
-	*reinterpret_cast<uint32*>(BufferPtr) = uint32(Name.GetComparisonIndex().ToUnstableInt());
-	GCurrentFrame.EventBufferSize = BufferPtr - GCurrentFrame.EventBuffer + sizeof(uint32);
-}
-
-void FGpuProfilerTrace::EndEvent(uint64 TimestampMicroseconds GPU_TRACE_ARG)
-{
-	using namespace GpuProfilerTrace;
-
-#if RHI_NEW_GPU_PROFILER
-	FGpuTraceFrame& GCurrentFrame = GCurrentFrames[GPUIndex];
-#endif
-
-	if (!GCurrentFrame.bActive)
-	{
-		return;
-	}
-
-	// Prevent buffer overrun
-	if (GCurrentFrame.EventBufferSize + 10 > GCurrentFrame.MaxEventBufferSize) // 10 is the max size that FTraceUtils::Encode7bit might use
-	{
-		UE_LOG(LogRHI, Error, TEXT("GpuProfiler's scratch buffer is out of space for this frame (current size : %d kB). Dropping this frame. The size can be increased dynamically with the console variable r.GpuProfilerMaxEventBufferSizeKB"), GCurrentFrame.MaxEventBufferSize / 1024);
-
-		// Deactivate for the current frame to avoid errors while decoding an incomplete trace
-		GCurrentFrame.bActive = false;
-		return;
-	}
-
-	uint64 TimestampDelta = TimestampMicroseconds - GCurrentFrame.LastTimestamp;
-	GCurrentFrame.LastTimestamp = TimestampMicroseconds;
-	uint8* BufferPtr = GCurrentFrame.EventBuffer + GCurrentFrame.EventBufferSize;
-	FTraceUtils::Encode7bit(TimestampDelta << 1ull, BufferPtr);
-	GCurrentFrame.EventBufferSize = BufferPtr - GCurrentFrame.EventBuffer;
-}
-
-void FGpuProfilerTrace::EndFrame(uint32 GPUIndex)
-{
-	using namespace GpuProfilerTrace;
-
-#if RHI_NEW_GPU_PROFILER
-	FGpuTraceFrame& GCurrentFrame = GCurrentFrames[GPUIndex];
-#endif
-
-	if (GCurrentFrame.bActive && GCurrentFrame.EventBufferSize)
-	{
-		// This subtraction is intended to be performed on uint64 to leverage the wrap around behavior defined by the standard
-		uint64 Bias = GCurrentFrame.Calibration.CPUMicroseconds - GCurrentFrame.Calibration.GPUMicroseconds;
-
-		if (GPUIndex == 0)
-		{
-			UE_TRACE_LOG(GpuProfiler, Frame, GpuChannel)
-				<< Frame.CalibrationBias(Bias)
-				<< Frame.TimestampBase(GCurrentFrame.TimestampBase)
-				<< Frame.RenderingFrameNumber(GCurrentFrame.RenderingFrameNumber)
-				<< Frame.Data(GCurrentFrame.EventBuffer, GCurrentFrame.EventBufferSize);
-		}
-		else if (GPUIndex == 1)
-		{
-			UE_TRACE_LOG(GpuProfiler, Frame2, GpuChannel)
-				<< Frame2.CalibrationBias(Bias)
-				<< Frame2.TimestampBase(GCurrentFrame.TimestampBase)
-				<< Frame2.RenderingFrameNumber(GCurrentFrame.RenderingFrameNumber)
-				<< Frame2.Data(GCurrentFrame.EventBuffer, GCurrentFrame.EventBufferSize);
-		}
-	}
-
-	GCurrentFrame.EventBufferSize = 0;
-	GCurrentFrame.bActive = false;
-}
-
-#undef GPU_TRACE_ARG
-
-void FGpuProfilerTrace::Deinitialize()
-{
-	using namespace GpuProfilerTrace;
-
-#if RHI_NEW_GPU_PROFILER
-	for (FGpuTraceFrame& GCurrentFrame : GCurrentFrames)
-#endif
-	{
-		FMemory::Free(GCurrentFrame.EventBuffer);
-		GCurrentFrame.EventBuffer = nullptr;
-		GCurrentFrame.MaxEventBufferSize = 0;
-	}
-}
-
-#endif // GPUPROFILERTRACE_ENABLED
-
 namespace UE::RHI::GPUProfiler
 {
 // Tracing for the new GPU Profiler
diff --git a/Engine/Source/Runtime/RHI/Private/RHI.cpp b/Engine/Source/Runtime/RHI/Private/RHI.cpp
--- a/Engine/Source/Runtime/RHI/Private/RHI.cpp
+++ b/Engine/Source/Runtime/RHI/Private/RHI.cpp
@@ -37,18 +37,6 @@ IMPLEMENT_MODULE(FDefaultModuleImpl, RHI);
 /** RHI Logging. */
 DEFINE_LOG_CATEGORY(LogRHI);
 
-#if (RHI_NEW_GPU_PROFILER == 0)
-
-	CSV_DEFINE_CATEGORY(RHI, true);
-
-	#if UE_BUILD_SHIPPING
-	CSV_DEFINE_CATEGORY(DrawCall, false);
-	#else
-	CSV_DEFINE_CATEGORY(DrawCall, true);
-	#endif
-
-#endif
-
 IMPLEMENT_TYPE_LAYOUT(FRHIUniformBufferLayoutInitializer);
 IMPLEMENT_TYPE_LAYOUT(FRHIUniformBufferResourceInitializer);
 
@@ -141,13 +129,6 @@ const FClearValueBinding FClearValueBinding::DefaultNormal8Bit(FLinearColor(128.
 
 #if HAS_GPU_STATS
 
-#if (RHI_NEW_GPU_PROFILER == 0)
-	FRHIDrawStatsCategory::FRHIDrawStatsCategory()
-		: Name(NAME_None)
-		, Index(-1)
-	{}
-#endif
-
 	FRHIDrawStatsCategory::FRHIDrawStatsCategory(FName InName)
 		: Name(InName)
 		, Index(GetManager().NumCategory++)
@@ -1259,94 +1240,6 @@ ERHIAccess GRHIMultiPipelineMergeableAccessMask = ERHIAccess::ReadOnlyExclusiveM
 
 ERHIAccess GRHIMultiSubresourceDiscardIntermediateAccess = ERHIAccess::SRVCompute;
 
-#if (RHI_NEW_GPU_PROFILER == 0)
-
-void FRHIDrawStats::Accumulate(FRHIDrawStats& Other)
-{
-	for (uint32 GPUIndex = 0; GPUIndex < GNumExplicitGPUsForRendering; ++GPUIndex)
-	{
-		FPerGPU& LeftGPU = GetGPU(GPUIndex);
-		FPerGPU& RightGPU = Other.GetGPU(GPUIndex);
-
-		for (int32 CategoryIndex = 0; CategoryIndex < NumCategories; ++CategoryIndex)
-		{
-			FPerCategory& LeftCategory = LeftGPU.Categories[CategoryIndex];
-			FPerCategory& RightCategory = RightGPU.Categories[CategoryIndex];
-
-			LeftCategory += RightCategory;
-		}
-	}
-}
-
-RHI_API void FRHIDrawStats::ProcessAsFrameStats()
-{
-#if HAS_GPU_STATS
-	// Only copy the display counters every half second keep things more stable.
-	constexpr float TimeoutSeconds = 0.5;
-
-	static double LastTime = 0.0;
-	double CurrentTime = FPlatformTime::Seconds();
-
-	bool bCopyDisplayFrames = false;
-	if (CurrentTime - LastTime > TimeoutSeconds)
-	{
-		LastTime = CurrentTime;
-		bCopyDisplayFrames = true;
-	}
-
-	FRHIDrawStatsCategory::FManager& Manager = FRHIDrawStatsCategory::GetManager();
-#endif
-
-	// Summed stats across all GPUs
-	FPerCategory Total = {};
-	TStaticArray<FPerCategory, FRHIDrawStats::NumCategories> TotalPerCategory;
-	FMemory::Memzero(TotalPerCategory);
-
-	for (int32 GPUIndex = 0; GPUIndex < MAX_NUM_GPUS; ++GPUIndex)
-	{
-		FPerCategory TotalPerGPU = {};
-		for (int32 CategoryIndex = 0; CategoryIndex < FRHIDrawStats::NumCategories; ++CategoryIndex)
-		{
-			FPerCategory& Category = GPUs[GPUIndex].Categories[CategoryIndex];
-
-			TotalPerCategory[CategoryIndex] += Category;
-			TotalPerGPU                     += Category;
-			Total                           += Category;
-
-#if HAS_GPU_STATS
-			if (bCopyDisplayFrames && CategoryIndex < Manager.NumCategory)
-			{
-				Manager.DisplayCounts[CategoryIndex][GPUIndex] = Category.Draws;
-			}
-#endif // HAS_GPU_STATS
-		}
-		
-		GNumDrawCallsRHI      [GPUIndex] = TotalPerGPU.Draws;
-		GNumPrimitivesDrawnRHI[GPUIndex] = TotalPerGPU.GetTotalPrimitives();
-	}
-
-	// Multi-GPU support : CSV stats do not support MGPU yet. We're summing the totals across all GPUs here.
-	CSV_CUSTOM_STAT(RHI, DrawCalls      , int32(Total.Draws               ), ECsvCustomStatOp::Set);
-	CSV_CUSTOM_STAT(RHI, PrimitivesDrawn, int32(Total.GetTotalPrimitives()), ECsvCustomStatOp::Set);
-
-#if HAS_GPU_STATS
-	SET_DWORD_STAT(STAT_RHITriangles         , Total.Triangles);
-	SET_DWORD_STAT(STAT_RHILines             , Total.Lines    );
-	SET_DWORD_STAT(STAT_RHIDrawPrimitiveCalls, Total.Draws    );
-
-	#if CSV_PROFILER_STATS
-	for (int32 CategoryIndex = 0; CategoryIndex < Manager.NumCategory; ++CategoryIndex)
-	{
-		FCsvProfiler::RecordCustomStat(Manager.Array[CategoryIndex]->Name, CSV_CATEGORY_INDEX(DrawCall), int32(TotalPerCategory[CategoryIndex].Draws), ECsvCustomStatOp::Set);
-	}
-	#endif
-#endif // HAS_GPU_STATS
-
-	Reset();
-}
-
-#endif // (RHI_NEW_GPU_PROFILER == 0)
-
 //
 // The current shader platform.
 //
diff --git a/Engine/Source/Runtime/RHI/Private/RHICommandList.cpp b/Engine/Source/Runtime/RHI/Private/RHICommandList.cpp
--- a/Engine/Source/Runtime/RHI/Private/RHICommandList.cpp
+++ b/Engine/Source/Runtime/RHI/Private/RHICommandList.cpp
@@ -99,8 +99,6 @@ FAutoConsoleTaskPriority CPrio_SceneRenderingTask(
 	ENamedThreads::HighTaskPriority
 );
 
-extern TAutoConsoleVariable<int32> GProfileGPUTransitions;
-
 DECLARE_CYCLE_STAT(TEXT("Parallel Translate"),                 STAT_ParallelTranslate,      STATGROUP_RHICMDLIST);
 DECLARE_CYCLE_STAT(TEXT("RHI Thread Parallel Translate Wait"), STAT_ParallelTranslateWait,  STATGROUP_RHICMDLIST);
 DECLARE_CYCLE_STAT(TEXT("Explicit wait for tasks"),            STAT_ExplicitWait,           STATGROUP_RHICMDLIST);
@@ -287,31 +285,6 @@ void FRHICommandListBase::FinishRecording()
 	DispatchEvent->DispatchSubsequents();
 }
 
-#if (RHI_NEW_GPU_PROFILER == 0) && HAS_GPU_STATS
-TOptional<FRHIDrawStatsCategory const*> FRHICommandListBase::SetDrawStatsCategory(TOptional<FRHIDrawStatsCategory const*> Category)
-{
-	check(!Category.IsSet() || ((*Category) == nullptr || (*Category)->ShouldCountDraws()));
-
-	TOptional<FRHIDrawStatsCategory const*> Previous = PersistentState.CurrentDrawStatsCategory;
-
-	if (Previous != Category)
-	{
-		PersistentState.CurrentDrawStatsCategory = Category;
-
-		EnqueueLambda([Category](FRHICommandListBase& ExecutingCmdList)
-		{
-			// InitialDrawStatsCategory will be unset in Bypass() mode, but we shouldn't 
-			// be using it as the Category should have already been determined.
-			ExecutingCmdList.PersistentState.CurrentDrawStatsCategory = Category.IsSet()
-				? Category.GetValue()
-				: ExecutingCmdList.InitialDrawStatsCategory.GetValue();
-		});
-	}
-
-	return Previous;
-}
-#endif
-
 #if WITH_RHI_BREADCRUMBS
 void FRHICommandListBase::AttachBreadcrumbSubTree(FRHIBreadcrumbAllocator& Allocator, FRHIBreadcrumbList& Nodes)
 {
@@ -893,7 +866,7 @@ void FRHICommandListExecutor::FSubmitState::Dispatch(FRHICommandListBase* CmdLis
 	ConditionalSplitTranslateJob(CmdList);	
 
 	{
-#if WITH_RHI_BREADCRUMBS
+	#if WITH_RHI_BREADCRUMBS
 		// Fixup unknown breadcrumb parents
 		for (FRHIBreadcrumbNode* Node : CmdList->CPUBreadcrumbState.UnknownParentList.IterateAndUnlink())
 		{
@@ -975,18 +948,7 @@ void FRHICommandListExecutor::FSubmitState::Dispatch(FRHICommandListBase* CmdLis
 			CmdList->CPUBreadcrumbState.Current = CPUFirst;
 			CmdList->PersistentState.LocalBreadcrumb = CPUFirst;
 		}
-#endif // WITH_RHI_BREADCRUMBS
-
-#if (RHI_NEW_GPU_PROFILER == 0) && HAS_GPU_STATS
-		CmdList->InitialDrawStatsCategory = GRHICommandList.CurrentDrawStatsCategory;
-
-		if (CmdList->PersistentState.CurrentDrawStatsCategory.IsSet())
-		{
-			GRHICommandList.CurrentDrawStatsCategory = CmdList->PersistentState.CurrentDrawStatsCategory.GetValue();
-		}
-
-		CmdList->PersistentState.CurrentDrawStatsCategory = CmdList->InitialDrawStatsCategory;
-#endif
+	#endif // WITH_RHI_BREADCRUMBS
 
 		FGraphEventArray Prereqs;
 		if (!CmdList->AllowParallelTranslate())
@@ -1120,10 +1082,6 @@ void FRHICommandListExecutor::FTranslateState::Translate(FRHICommandListBase* Cm
 	BreadcrumbAllocatorRefs.Append(MoveTemp(CmdList->BreadcrumbAllocatorRefs));
 #endif
 
-#if (RHI_NEW_GPU_PROFILER == 0)
-	DrawStats.Accumulate(CmdList->DrawStats);
-#endif
-
 	if (CmdList->ParallelRenderPassBegin && !bShouldFinalize)
 	{
 		// Forward the acquired contexts to the next chain
@@ -1235,10 +1193,6 @@ void FRHICommandListExecutor::FSubmitState::Submit(const FSubmitArgs& Args)
 				FinalizedCmdLists.Add(TranslateState.FinalizedCmdList);
 			}
 		}
-
-#if (RHI_NEW_GPU_PROFILER == 0)
-		GRHICommandList.FrameDrawStats.Accumulate(Job->DrawStats);
-#endif
 	}
 
 	if (GDynamicRHI)
@@ -1291,9 +1245,6 @@ void FRHICommandListExecutor::FSubmitState::Submit(const FSubmitArgs& Args)
 		};
 
 		GDynamicRHI->RHIEndFrame(EndFrameArgs);
-#if (RHI_NEW_GPU_PROFILER == 0)
-		GRHICommandList.FrameDrawStats.ProcessAsFrameStats();
-#endif
 	}
 
 	CompletionEvent->DispatchSubsequents();
@@ -1389,8 +1340,8 @@ RHI_API FGraphEventRef FRHICommandListExecutor::Submit(TConstArrayView<FRHIComma
 		SubmitState->bEmitBreadcrumbs = bEmitBreadcrumbs;
 #endif
 
-		// Prevent use of parallel contexts if unsupported by the RHI, while the legacy 'profilegpu' command is active, or while disabled by the cvar.
-		if (GRHISupportsParallelRHIExecute && !GTriggerGPUProfile && CVarRHICmdParallelTranslateEnable.GetValueOnRenderThread())
+		// Prevent use of parallel contexts if unsupported by the RHI, or while disabled by the cvar.
+		if (GRHISupportsParallelRHIExecute && CVarRHICmdParallelTranslateEnable.GetValueOnRenderThread())
 		{
 			SubmitState->MaxCommandsPerTranslate     = CVarRHICmdParallelTranslateMaxCommandsPerTranslate.GetValueOnRenderThread();
 			SubmitState->bAllowSingleParallelCombine = CVarRHICmdParallelTranslateCombineSingleAndParallel.GetValueOnRenderThread();
@@ -1556,10 +1507,6 @@ void FRHICommandListImmediate::InitializeImmediateContexts()
 			GPUBreadcrumbState[Pipeline].Latest = GRHICommandList.Breadcrumbs.GPU[Pipeline].Current;
 		}
 #endif
-
-#if (RHI_NEW_GPU_PROFILER == 0) && HAS_GPU_STATS
-		InitialDrawStatsCategory = GRHICommandList.CurrentDrawStatsCategory;
-#endif
 	}
 
 	// This can be called before the RHI is initialized, in which case
@@ -1938,13 +1885,6 @@ void FRHICommandListImmediate::EndFrame()
 	GDynamicRHI->RHIEndFrame_RenderThread(*this);
 }
 
-#if WITH_PROFILEGPU && (RHI_NEW_GPU_PROFILER == 0)
-int32 FRHIComputeCommandList::GetGProfileGPUTransitions()
-{
-	return GProfileGPUTransitions.GetValueOnAnyThread();
-}
-#endif
-
 void FRHIComputeCommandList::Transition(TArrayView<const FRHITransitionInfo> Infos, ERHITransitionCreateFlags CreateFlags)
 {
 	TransitionInternal(Infos, CreateFlags);
diff --git a/Engine/Source/Runtime/RHI/Private/RHIStats.cpp b/Engine/Source/Runtime/RHI/Private/RHIStats.cpp
--- a/Engine/Source/Runtime/RHI/Private/RHIStats.cpp
+++ b/Engine/Source/Runtime/RHI/Private/RHIStats.cpp
@@ -12,16 +12,10 @@ int32 GNumPrimitivesDrawnRHI[MAX_NUM_GPUS] = {};
 // Define counter stats.
 #if HAS_GPU_STATS
 
-#if RHI_NEW_GPU_PROFILER
 	DEFINE_STAT(STAT_RHIDraws);
 	DEFINE_STAT(STAT_RHIDispatches);
 	DEFINE_STAT(STAT_RHIPrimitives);
 	DEFINE_STAT(STAT_RHIVertices);
-#else
-	DEFINE_STAT(STAT_RHIDrawPrimitiveCalls);
-	DEFINE_STAT(STAT_RHITriangles);
-	DEFINE_STAT(STAT_RHILines);
-#endif
 
 #endif
 
diff --git a/Engine/Source/Runtime/RHI/Private/RHIStrings.cpp b/Engine/Source/Runtime/RHI/Private/RHIStrings.cpp
--- a/Engine/Source/Runtime/RHI/Private/RHIStrings.cpp
+++ b/Engine/Source/Runtime/RHI/Private/RHIStrings.cpp
@@ -315,7 +315,7 @@ FString GetBufferUsageFlagsName(EBufferUsageFlags BufferUsage)
 		case BUF_None:					return TEXT("BUF_None");
 		case BUF_Static:				return TEXT("BUF_Static");
 		case BUF_Dynamic:				return TEXT("BUF_Dynamic");
-		case BUF_Volatile:				return TEXT("BUF_Volitile");
+		case BUF_Volatile:				return TEXT("BUF_Volatile");
 		case BUF_UnorderedAccess:		return TEXT("BUF_UnorderedAccess");
 		case BUF_ByteAddressBuffer:		return TEXT("BUF_ByteAddressBuffer");
 		case BUF_SourceCopy:			return TEXT("BUF_SourceCopy");
@@ -377,7 +377,7 @@ FString GetTextureCreateFlagsName(ETextureCreateFlags TextureCreateFlags)
 		case ETextureCreateFlags::NoFastClearFinalize:				return TEXT("NoFastClearFinalize");
 		case ETextureCreateFlags::ReduceMemoryWithTilingMode:		return TEXT("ReduceMemoryWithTilingMode");
 		case ETextureCreateFlags::AtomicCompatible:					return TEXT("AtomicCompatible");
-		case ETextureCreateFlags::Atomic64Compatible:				return TEXT("Aomic64Capable");
+		case ETextureCreateFlags::Atomic64Compatible:				return TEXT("Atomic64Compatible");
 		case ETextureCreateFlags::External:							return TEXT("External");
 		case ETextureCreateFlags::MultiGPUGraphIgnore:				return TEXT("MultiGPUGraphIgnore");
 		case ETextureCreateFlags::ReservedResource:					return TEXT("ReservedResource");
@@ -425,7 +425,9 @@ PRAGMA_DISABLE_DEPRECATION_WARNINGS
 	RHI_RESOURCE_TYPE_DEF(TextureCube),
 PRAGMA_ENABLE_DEPRECATION_WARNINGS
 	RHI_RESOURCE_TYPE_DEF(TextureReference),
+PRAGMA_DISABLE_DEPRECATION_WARNINGS
 	RHI_RESOURCE_TYPE_DEF(TimestampCalibrationQuery),
+PRAGMA_ENABLE_DEPRECATION_WARNINGS
 	RHI_RESOURCE_TYPE_DEF(GPUFence),
 	RHI_RESOURCE_TYPE_DEF(RenderQuery),
 	RHI_RESOURCE_TYPE_DEF(RenderQueryPool),
diff --git a/Engine/Source/Runtime/RHI/Public/GPUProfiler.h b/Engine/Source/Runtime/RHI/Public/GPUProfiler.h
--- a/Engine/Source/Runtime/RHI/Public/GPUProfiler.h
+++ b/Engine/Source/Runtime/RHI/Public/GPUProfiler.h
@@ -17,8 +17,6 @@
 #include "Containers/SpscQueue.h"
 #include "Containers/StaticArray.h"
 
-#if RHI_NEW_GPU_PROFILER
-
 namespace UE::RHI::GPUProfiler
 {
 	struct FQueue
@@ -567,328 +565,6 @@ namespace UE::RHI::GPUProfiler
 	RHI_API bool IsProfiling();
 }
 
-#else
-
-/** Stats for a single perf event node. */
-class FGPUProfilerEventNodeStats : public FRefCountedObject
-{
-public:
-	FGPUProfilerEventNodeStats() :
-		NumDraws(0),
-		NumPrimitives(0),
-		NumVertices(0),
-		NumDispatches(0),
-		GroupCount(FIntVector(0, 0, 0)),
-		NumTotalDispatches(0),
-		NumTotalDraws(0),
-		NumTotalPrimitives(0),
-		NumTotalVertices(0),
-		TimingResult(0),
-		NumEvents(0)
-	{
-	}
-
-	FGPUProfilerEventNodeStats(const FGPUProfilerEventNodeStats& rhs)
-	{
-		NumDraws = rhs.NumDraws;
-		NumPrimitives = rhs.NumPrimitives;
-		NumVertices = rhs.NumVertices;
-		NumDispatches = rhs.NumDispatches;
-		NumTotalDispatches = rhs.NumTotalDispatches;
-		NumTotalDraws = rhs.NumDraws;
-		NumTotalPrimitives = rhs.NumPrimitives;
-		NumTotalVertices = rhs.NumVertices;
-		TimingResult = rhs.TimingResult;
-		NumEvents = rhs.NumEvents;
-	}
-
-	/** Exclusive number of draw calls rendered in this event. */
-	uint32 NumDraws;
-
-	/** Exclusive number of primitives rendered in this event. */
-	uint32 NumPrimitives;
-
-	/** Exclusive number of vertices rendered in this event. */
-	uint32 NumVertices;
-
-	/** Compute stats */
-	uint32 NumDispatches;
-	FIntVector GroupCount;
-	uint32 NumTotalDispatches;
-
-	/** Inclusive number of draw calls rendered in this event and children. */
-	uint32 NumTotalDraws;
-
-	/** Inclusive number of primitives rendered in this event and children. */
-	uint32 NumTotalPrimitives;
-
-	/** Inclusive number of vertices rendered in this event and children. */
-	uint32 NumTotalVertices;
-
-	/** GPU time spent inside the perf event's begin and end, in ms. */
-	float TimingResult;
-
-	/** Inclusive number of other perf events that this is the parent of. */
-	uint32 NumEvents;
-
-	const FGPUProfilerEventNodeStats operator+=(const FGPUProfilerEventNodeStats& rhs)
-	{
-		NumDraws += rhs.NumDraws;
-		NumPrimitives += rhs.NumPrimitives;
-		NumVertices += rhs.NumVertices;
-		NumDispatches += rhs.NumDispatches;
-		NumTotalDispatches += rhs.NumTotalDispatches;
-		NumTotalDraws += rhs.NumDraws;
-		NumTotalPrimitives += rhs.NumPrimitives;
-		NumTotalVertices += rhs.NumVertices;
-		TimingResult += rhs.TimingResult;
-		NumEvents += rhs.NumEvents;
-
-		return *this;
-	}
-};
-
-/** Stats for a single perf event node. */
-class FGPUProfilerEventNode : public FGPUProfilerEventNodeStats
-{
-public:
-	FGPUProfilerEventNode(const TCHAR* InName, FGPUProfilerEventNode* InParent) :
-		FGPUProfilerEventNodeStats(),
-		Name(InName),
-		Parent(InParent)
-	{
-	}
-
-	~FGPUProfilerEventNode() = default;
-
-	FString Name;
-
-	/** Pointer to parent node so we can walk up the tree on appEndDrawEvent. */
-	FGPUProfilerEventNode* Parent;
-
-	/** Children perf event nodes. */
-	TArray<TRefCountPtr<FGPUProfilerEventNode> > Children;
-
-	virtual float GetTiming() { return 0.0f; }
-	virtual void StartTiming() {}
-	virtual void StopTiming() {}
-};
-
-/** An entire frame of perf event nodes, including ancillary timers. */
-struct FGPUProfilerEventNodeFrame
-{
-	virtual ~FGPUProfilerEventNodeFrame() = default;
-
-	/** Root nodes of the perf event tree. */
-	TArray<TRefCountPtr<FGPUProfilerEventNode> > EventTree;
-
-	/** Start this frame of per tracking */
-	virtual void StartFrame() {}
-
-	/** End this frame of per tracking, but do not block yet */
-	virtual void EndFrame() {}
-
-	/** Dumps perf event information, blocking on GPU. */
-	RHI_API void DumpEventTree();
-
-	/** Calculates root timing base frequency (if needed by this RHI) */
-	virtual float GetRootTimingResults() { return 0.0f; }
-
-	/** D3D11 Hack */
-	virtual void LogDisjointQuery() {}
-
-	virtual bool PlatformDisablesVSync() const { return false; }
-};
-
-/**
-* Two timestamps performed on GPU and CPU at nearly the same time.
-* This can be used to visualize GPU and CPU timing events on the same timeline.
-*/
-struct FGPUTimingCalibrationTimestamp
-{
-	uint64 GPUMicroseconds = 0;
-	uint64 CPUMicroseconds = 0;
-};
-
-/**
- * Holds information if this platform's GPU allows timing
- */
-struct FGPUTiming
-{
-public:
-	/**
-	 * Whether GPU timing measurements are supported by the driver.
-	 *
-	 * @return true if GPU timing measurements are supported by the driver.
-	 */
-	static bool IsSupported()
-	{
-		return GIsSupported;
-	}
-
-	/**
-	 * Returns the frequency for the timing values, in number of ticks per seconds.
-	 *
-	 * @return Frequency for the timing values, in number of ticks per seconds, or 0 if the feature isn't supported.
-	 */
-	static uint64 GetTimingFrequency(uint32 GPUIndex = 0)
-	{
-		return GTimingFrequency[GPUIndex];
-	}
-
-	/**
-	* Returns a pair of timestamps performed on GPU and CPU at nearly the same time, in microseconds.
-	*
-	* @return CPU and GPU timestamps, in microseconds. Both are 0 if feature isn't supported.
-	*/
-	static FGPUTimingCalibrationTimestamp GetCalibrationTimestamp(uint32 GPUIndex = 0)
-	{
-		return GCalibrationTimestamp[GPUIndex];
-	}
-
-	typedef void (PlatformStaticInitialize)(void*);
-	static void StaticInitialize(void* UserData, PlatformStaticInitialize* PlatformFunction)
-	{
-		if (!GAreGlobalsInitialized && PlatformFunction)
-		{
-			(*PlatformFunction)(UserData);
-
-			if (GetTimingFrequency() != 0)
-			{
-				GIsSupported = true;
-			}
-			else
-			{
-				GIsSupported = false;
-			}
-
-			GAreGlobalsInitialized = true;
-		}
-	}
-
-protected:
-	/** Whether the static variables have been initialized. */
-	RHI_API static bool		GAreGlobalsInitialized;
-
-	/** Whether GPU timing measurements are supported by the driver. */
-	RHI_API static bool		GIsSupported;
-
-	static void SetTimingFrequency(uint64 TimingFrequency, uint32 GPUIndex = 0)
-	{
-		GTimingFrequency[GPUIndex] = TimingFrequency;
-	}
-
-	static void SetCalibrationTimestamp(FGPUTimingCalibrationTimestamp CalibrationTimestamp, uint32 GPUIndex = 0)
-	{
-		GCalibrationTimestamp[GPUIndex] = CalibrationTimestamp;
-	}
-
-private:
-	/** Frequency for the timing values, in number of ticks per seconds, or 0 if the feature isn't supported. */
-	RHI_API static TStaticArray<uint64, MAX_NUM_GPUS>	GTimingFrequency;
-
-	/**
-	* Two timestamps performed on GPU and CPU at nearly the same time.
-	* This can be used to visualize GPU and CPU timing events on the same timeline.
-	* Both values may be 0 if timer calibration is not available on current platform.
-	*/
-	RHI_API static TStaticArray<FGPUTimingCalibrationTimestamp, MAX_NUM_GPUS> GCalibrationTimestamp;
-};
-
-/** 
- * Encapsulates GPU profiling logic and data. 
- * There's only one global instance of this struct so it should only contain global data, nothing specific to a frame.
- */
-struct FGPUProfiler
-{
-	/** Whether we are currently tracking perf events or not. */
-	bool bTrackingEvents;
-
-	/** Whether we are currently tracking data for gpucrash debugging or not */
-	bool bTrackingGPUCrashData;
-
-	/** A latched version of GTriggerGPUProfile. This is a form of pseudo-thread safety. We read the value once a frame only. */
-	bool bLatchedGProfilingGPU;
-
-	/** A latched version of GTriggerGPUHitchProfile. This is a form of pseudo-thread safety. We read the value once a frame only. */
-	bool bLatchedGProfilingGPUHitches;
-
-	/** The previous latched version of GTriggerGPUHitchProfile.*/
-	bool bPreviousLatchedGProfilingGPUHitches;
-
-	/** Original state of GEmitDrawEvents before it was overridden for profiling. */
-	bool bOriginalGEmitDrawEvents;
-
-	/** GPU hitch profile history debounce...after a hitch, we just ignore frames for a while */
-	int32 GPUHitchDebounce;
-
-	/** scope depth to record crash data depth. to limit perf/mem requirements */
-	int32 GPUCrashDataDepth;
-
-	/** Current perf event node frame. */
-	FGPUProfilerEventNodeFrame* CurrentEventNodeFrame = nullptr;
-
-	/** Current perf event node. */
-	FGPUProfilerEventNode* CurrentEventNode;
-
-	int32 StackDepth;
-
-	FGPUProfiler() :
-		bTrackingEvents(false),
-		bTrackingGPUCrashData(false),
-		bLatchedGProfilingGPU(false),
-		bLatchedGProfilingGPUHitches(false),
-		bPreviousLatchedGProfilingGPUHitches(false),
-		bOriginalGEmitDrawEvents(false),
-		GPUHitchDebounce(0),
-		GPUCrashDataDepth(-1),
-		CurrentEventNodeFrame(NULL),
-		CurrentEventNode(NULL),
-		StackDepth(0)
-	{
-	}
-
-	virtual ~FGPUProfiler() = default;
-
-	void RegisterGPUWork(uint32 NumDraws, uint32 NumPrimitives, uint32 NumVertices)
-	{
-		if (bTrackingEvents && CurrentEventNode)
-		{
-			check(IsInRenderingThread() || IsInRHIThread());
-			CurrentEventNode->NumDraws += NumDraws;
-			CurrentEventNode->NumPrimitives += NumPrimitives;
-			CurrentEventNode->NumVertices += NumVertices;
-		}
-	}
-
-	void RegisterGPUWork(uint32 NumPrimitives = 0, uint32 NumVertices = 0)
-	{
-		RegisterGPUWork(1, NumPrimitives, NumVertices);
-	}
-
-	void RegisterGPUDispatch(FIntVector GroupCount)
-	{
-		if (bTrackingEvents && CurrentEventNode)
-		{
-			check(IsInRenderingThread() || IsInRHIThread());
-			CurrentEventNode->NumDispatches++;
-			CurrentEventNode->GroupCount = GroupCount;
-		}
-	}
-
-	virtual FGPUProfilerEventNode* CreateEventNode(const TCHAR* InName, FGPUProfilerEventNode* InParent)
-	{
-		return new FGPUProfilerEventNode(InName, InParent);
-	}
-
-	RHI_API virtual void PushEvent(const TCHAR* Name, FColor Color);
-	RHI_API virtual void PopEvent();
-
-	bool IsProfilingGPU() const { return bTrackingEvents; }
-};
-
-#endif
-
 //
 // Type used to pipe GPU frame timings from the end-of-pipe / RHI threads up to the game / render threads.
 // Stores a history of GPU frame timings, which can be retrieved by engine code via:
diff --git a/Engine/Source/Runtime/RHI/Public/GpuProfilerTrace.h b/Engine/Source/Runtime/RHI/Public/GpuProfilerTrace.h
--- a/Engine/Source/Runtime/RHI/Public/GpuProfilerTrace.h
+++ b/Engine/Source/Runtime/RHI/Public/GpuProfilerTrace.h
@@ -10,55 +10,9 @@
 
 // Tracing for the new GPU Profiler
 #ifndef UE_TRACE_GPU_PROFILER_ENABLED
-	#define UE_TRACE_GPU_PROFILER_ENABLED UE_TRACE_ENABLED && RHI_NEW_GPU_PROFILER && !UE_BUILD_SHIPPING
+	#define UE_TRACE_GPU_PROFILER_ENABLED UE_TRACE_ENABLED && !UE_BUILD_SHIPPING
 #endif
 
-#if GPUPROFILERTRACE_ENABLED
-
-#if RHI_NEW_GPU_PROFILER
-// Define this structure here when the new GPU profiler is enabled so we can still build the old trace API.
-// @todo - remove this. GPU timestamp calibration is no longer necessary with the new GPU profiler, as the
-// platform RHIs are expected to translate timestamps from GPU to CPU clock domain before they reach the profiler.
-struct FGPUTimingCalibrationTimestamp
-{
-	uint64 GPUMicroseconds = 0;
-	uint64 CPUMicroseconds = 0;
-};
-#endif
-
-class FName;
-
-#if RHI_NEW_GPU_PROFILER
-// Adds a GPUIndex argument to each function in the API without breaking back compat
-#define GPU_TRACE_ARG , uint32 GPUIndex
-#else
-#define GPU_TRACE_ARG
-#endif
-
-struct FGpuProfilerTrace
-{
-	RHI_API static void BeginFrame(struct FGPUTimingCalibrationTimestamp& Calibration GPU_TRACE_ARG);
-	RHI_API static void SpecifyEventByName(const FName& Name GPU_TRACE_ARG);
-	RHI_API static void BeginEventByName(const FName& Name, uint32 FrameNumber, uint64 TimestampMicroseconds GPU_TRACE_ARG);
-	RHI_API static void EndEvent(uint64 TimestampMicroseconds GPU_TRACE_ARG);
-	RHI_API static void EndFrame(uint32 GPUIndex);
-	RHI_API static void Deinitialize();
-};
-
-#undef GPU_TRACE_ARG
-
-#endif // GPUPROFILERTRACE_ENABLED
-
-// Deprecated macros
-#define TRACE_GPUPROFILER_DEFINE_EVENT_TYPE(...)         UE_DEPRECATED_MACRO(5.6, "TRACE_GPUPROFILER_DEFINE_EVENT_TYPE has been deprecated and it is non functional. Use standard RHI breadcrumb events instead.")
-#define TRACE_GPUPROFILER_DECLARE_EVENT_TYPE_EXTERN(...) UE_DEPRECATED_MACRO(5.6, "TRACE_GPUPROFILER_DECLARE_EVENT_TYPE_EXTERN has been deprecated and it is non functional. Use standard RHI breadcrumb events instead.")
-#define TRACE_GPUPROFILER_EVENT_TYPE(...)                UE_DEPRECATED_MACRO(5.6, "TRACE_GPUPROFILER_EVENT_TYPE has been deprecated and it is non functional. Use standard RHI breadcrumb events instead.") nullptr
-#define TRACE_GPUPROFILER_BEGIN_FRAME(...)               UE_DEPRECATED_MACRO(5.6, "TRACE_GPUPROFILER_BEGIN_FRAME has been deprecated and it is non functional. Use standard RHI breadcrumb events instead.")
-#define TRACE_GPUPROFILER_BEGIN_EVENT(...)               UE_DEPRECATED_MACRO(5.6, "TRACE_GPUPROFILER_BEGIN_EVENT has been deprecated and it is non functional. Use standard RHI breadcrumb events instead.")
-#define TRACE_GPUPROFILER_END_EVENT(...)                 UE_DEPRECATED_MACRO(5.6, "TRACE_GPUPROFILER_END_EVENT has been deprecated and it is non functional. Use standard RHI breadcrumb events instead.")
-#define TRACE_GPUPROFILER_END_FRAME(...)                 UE_DEPRECATED_MACRO(5.6, "TRACE_GPUPROFILER_END_FRAME has been deprecated and it is non functional. Use standard RHI breadcrumb events instead.")
-#define TRACE_GPUPROFILER_DEINITIALIZE(...)              UE_DEPRECATED_MACRO(5.6, "TRACE_GPUPROFILER_DEINITIALIZE has been deprecated and it is non functional. Use standard RHI breadcrumb events instead.")
-
 class FCborWriter;
 class FDebugName;
 class FMemoryWriter;
diff --git a/Engine/Source/Runtime/RHI/Public/RHI.h b/Engine/Source/Runtime/RHI/Public/RHI.h
--- a/Engine/Source/Runtime/RHI/Public/RHI.h
+++ b/Engine/Source/Runtime/RHI/Public/RHI.h
@@ -26,19 +26,6 @@ class FResourceArrayInterface;
 class FResourceBulkDataInterface;
 class FRHICommandList;
 
-/**
- * RHI configuration settings.
- */
-
-#if (RHI_NEW_GPU_PROFILER == 0)
-namespace RHIConfig
-{
-	RHI_API bool ShouldSaveScreenshotAfterProfilingGPU();
-	RHI_API bool ShouldShowProfilerAfterProfilingGPU();
-	RHI_API float GetGPUHitchThreshold();
-}
-#endif
-
 /**
  * RHI capabilities.
  */
diff --git a/Engine/Source/Runtime/RHI/Public/RHIBreadcrumbs.h b/Engine/Source/Runtime/RHI/Public/RHIBreadcrumbs.h
--- a/Engine/Source/Runtime/RHI/Public/RHIBreadcrumbs.h
+++ b/Engine/Source/Runtime/RHI/Public/RHIBreadcrumbs.h
@@ -22,7 +22,7 @@
 // (RHI breadcrumb allocators, platform RHI implemntation etc).
 //
 #ifndef WITH_RHI_BREADCRUMBS
-#define WITH_RHI_BREADCRUMBS (UE_BUILD_DEBUG || UE_BUILD_DEVELOPMENT || WITH_PROFILEGPU || (HAS_GPU_STATS && RHI_NEW_GPU_PROFILER))
+#define WITH_RHI_BREADCRUMBS (UE_BUILD_DEBUG || UE_BUILD_DEVELOPMENT || WITH_PROFILEGPU || HAS_GPU_STATS)
 #endif
 
 // Enables all RDG/RHI breadcrumb scopes, and features such as Insights markers.
@@ -35,9 +35,9 @@
 #define RHI_BREADCRUMBS_EMIT_CPU (WITH_RHI_BREADCRUMBS_FULL && CPUPROFILERTRACE_ENABLED && 1)
 
 // Whether to store the filename and line number of each RHI breadcrumb and emit this data to Insights.
-#define RHI_BREADCRUMBS_EMIT_LOCATION (WITH_RHI_BREADCRUMBS_FULL && (CPUPROFILERTRACE_ENABLED || GPUPROFILERTRACE_ENABLED) && 1)
+#define RHI_BREADCRUMBS_EMIT_LOCATION (WITH_RHI_BREADCRUMBS_FULL && CPUPROFILERTRACE_ENABLED && 1)
 
-#if RHI_NEW_GPU_PROFILER && HAS_GPU_STATS
+#if HAS_GPU_STATS
 namespace UE::RHI::GPUProfiler
 {
 	struct FGPUStat;
@@ -65,12 +65,12 @@ namespace UE::RHI::GPUProfiler
 	};
 
 	//
-	// Holds both a stats system ID, and a CSV profiler ID.
+	// Holds the pointer to the GPU profiler stat, when HAS_GPU_STATS is enabled.
 	// The computed stat value is emitted to both "stat gpu" and the CSV profiler.
 	//
 	struct FRHIBreadcrumbData_Stats
 	{
-#if RHI_NEW_GPU_PROFILER && HAS_GPU_STATS
+#if HAS_GPU_STATS
 
 		UE::RHI::GPUProfiler::FGPUStat* GPUStat;
 
@@ -93,58 +93,6 @@ namespace UE::RHI::GPUProfiler
 			return PointerHash(Stats.GPUStat);
 		}
 
-#elif HAS_GPU_STATS
-
-	#if STATS
-		TStatId StatId {};
-	#endif
-	#if CSV_PROFILER_STATS
-		FName CsvStat = NAME_None;
-	#endif
-
-		FRHIBreadcrumbData_Stats(TStatId InStatId, FName InCsvStat)
-		{
-	#if STATS
-			StatId = InStatId;
-	#endif
-	#if CSV_PROFILER_STATS
-			CsvStat = InCsvStat;
-	#endif
-		}
-
-		bool ShouldComputeStat() const
-		{
-	#if STATS
-			return StatId.IsValidStat();
-	#elif CSV_PROFILER_STATS
-			return CsvStat != NAME_None;
-	#else
-			return false;
-	#endif
-		}
-
-		bool operator == (FRHIBreadcrumbData_Stats const& RHS) const
-		{
-	#if STATS
-			return StatId == RHS.StatId;
-	#elif CSV_PROFILER_STATS
-			return CsvStat == RHS.CsvStat;
-	#else
-			return true;
-	#endif
-		}
-
-		friend uint32 GetTypeHash(FRHIBreadcrumbData_Stats const& Stats)
-		{
-	#if STATS
-			return GetTypeHash(Stats.StatId);
-	#elif CSV_PROFILER_STATS
-			return GetTypeHash(Stats.CsvStat);
-	#else
-			return 0;
-	#endif
-		}
-
 #else
 
 		FRHIBreadcrumbData_Stats() = default;
@@ -1349,12 +1297,9 @@ namespace UE::RHI::GPUProfiler
 	#define RHI_BREADCRUMB_DESC_FORWARD_VALUES(StaticName, FormatString, GPUStat) RHI_BREADCRUMB_PRIVATE_DEFINE(StaticName, FormatString, TValueRef , GPUStat)
 	#define RHI_BREADCRUMB_DESC_COPY_VALUES(   StaticName, FormatString, GPUStat) RHI_BREADCRUMB_PRIVATE_DEFINE(StaticName, FormatString, TValueType, GPUStat)
 
-	#if RHI_NEW_GPU_PROFILER && HAS_GPU_STATS	
+	#if HAS_GPU_STATS	
 		#define RHI_GPU_STAT_ARGS(StatName) FRHIBreadcrumbData_Stats(&GPUStat_##StatName)
 		#define RHI_GPU_STAT_ARGS_NONE      FRHIBreadcrumbData_Stats(nullptr)
-	#elif HAS_GPU_STATS
-		#define RHI_GPU_STAT_ARGS(StatName) FRHIBreadcrumbData_Stats(GET_STATID(Stat_GPU_##StatName), CSV_STAT_FNAME(StatName))
-		#define RHI_GPU_STAT_ARGS_NONE      FRHIBreadcrumbData_Stats(TStatId(), NAME_None)
 	#else
 		#define RHI_GPU_STAT_ARGS(StatName) FRHIBreadcrumbData_Stats()
 		#define RHI_GPU_STAT_ARGS_NONE      FRHIBreadcrumbData_Stats()
diff --git a/Engine/Source/Runtime/RHI/Public/RHICommandList.h b/Engine/Source/Runtime/RHI/Public/RHICommandList.h
--- a/Engine/Source/Runtime/RHI/Public/RHICommandList.h
+++ b/Engine/Source/Runtime/RHI/Public/RHICommandList.h
@@ -749,10 +749,6 @@ public:
 	bool IsInsideRenderPass    () const { return PersistentState.bInsideRenderPass;  }
 	bool IsInsideComputePass   () const { return PersistentState.bInsideComputePass; }
 
-#if (RHI_NEW_GPU_PROFILER == 0) && HAS_GPU_STATS
-	RHI_API TOptional<FRHIDrawStatsCategory const*> SetDrawStatsCategory(TOptional<FRHIDrawStatsCategory const*> Category);
-#endif
-
 	RHI_API FGraphEventRef RHIThreadFence(bool bSetLockFence = false);
 
 	inline void* LockBuffer(FRHIBuffer* Buffer, uint32 Offset, uint32 SizeRHI, EResourceLockMode LockMode)
@@ -1269,10 +1265,7 @@ protected:
 		FActivatePipelineCommand* First = nullptr;
 		FActivatePipelineCommand* Prev = nullptr;
 	} ActivatePipelineCommands {};
-#endif
 
-#if (RHI_NEW_GPU_PROFILER == 0) && HAS_GPU_STATS
-	TOptional<FRHIDrawStatsCategory const*> InitialDrawStatsCategory {};
 #endif
 
 	// The values in this struct are preserved when the command list is moved or reset.
@@ -1302,13 +1295,9 @@ protected:
 		FRHIComputeShader* BoundComputeShaderRHI = nullptr;
 		FRHICommandListScopedFence* CurrentFenceScope = nullptr;
 
-#if WITH_RHI_BREADCRUMBS
+	#if WITH_RHI_BREADCRUMBS
 		FRHIBreadcrumbNode* LocalBreadcrumb = FRHIBreadcrumbNode::Sentinel;
-#endif
-
-#if (RHI_NEW_GPU_PROFILER == 0) && HAS_GPU_STATS
-		TOptional<FRHIDrawStatsCategory const*> CurrentDrawStatsCategory {};
-#endif
+	#endif
 
 		TStaticArray<void*, MAX_NUM_GPUS> QueryBatchData_Timestamp { InPlace, nullptr };
 		TStaticArray<void*, MAX_NUM_GPUS> QueryBatchData_Occlusion { InPlace, nullptr };
@@ -1327,10 +1316,6 @@ protected:
 
 	} PersistentState;
 
-#if (RHI_NEW_GPU_PROFILER == 0)
-	FRHIDrawStats DrawStats {};
-#endif
-
 	TArray<FRHIBuffer*> PendingBufferUploads;
 	TArray<FRHITexture*> PendingTextureUploads;
 
@@ -1347,26 +1332,6 @@ public:
 	RHI_API void AttachBreadcrumbSubTree(FRHIBreadcrumbAllocator& Allocator, FRHIBreadcrumbList& Nodes);
 #endif
 
-#if (RHI_NEW_GPU_PROFILER == 0)
-	void Stats_AddDraw()
-	{
-#if HAS_GPU_STATS
-		DrawStats.AddDraw(PersistentState.CurrentGPUMask, PersistentState.CurrentDrawStatsCategory.GetValue());
-#else
-		DrawStats.AddDraw(PersistentState.CurrentGPUMask, nullptr);
-#endif
-	}
-
-	void Stats_AddDrawAndPrimitives(EPrimitiveType PrimitiveType, uint32 NumPrimitives)
-	{
-#if HAS_GPU_STATS
-		DrawStats.AddDrawAndPrimitives(PersistentState.CurrentGPUMask, PersistentState.CurrentDrawStatsCategory.GetValue(), PrimitiveType, NumPrimitives);
-#else
-		DrawStats.AddDrawAndPrimitives(PersistentState.CurrentGPUMask, nullptr, PrimitiveType, NumPrimitives);
-#endif
-	}
-#endif
-
 	TStaticArray<void*, MAX_NUM_GPUS>& GetQueryBatchData(ERenderQueryType QueryType)
 	{
 		switch (QueryType)
@@ -2302,19 +2267,6 @@ FRHICOMMAND_MACRO(FRHICommandEndRenderQuery)
 	RHI_EXECUTE_API void Execute(FRHICommandListBase& CmdList);
 };
 
-#if (RHI_NEW_GPU_PROFILER == 0)
-FRHICOMMAND_MACRO(FRHICommandCalibrateTimers)
-{
-	FRHITimestampCalibrationQuery* CalibrationQuery;
-
-	inline FRHICommandCalibrateTimers(FRHITimestampCalibrationQuery * CalibrationQuery)
-		: CalibrationQuery(CalibrationQuery)
-	{
-	}
-	RHI_EXECUTE_API void Execute(FRHICommandListBase & CmdList);
-};
-#endif
-
 FRHICOMMAND_MACRO(FRHICommandPostExternalCommandsReset)
 {
 	RHI_EXECUTE_API void Execute(FRHICommandListBase& CmdList);
@@ -2773,15 +2725,8 @@ public:
 		ALLOC_COMMAND(FRHICommandClearUAVUint)(UnorderedAccessViewRHI, Values);
 	}
 
-#if WITH_PROFILEGPU && (RHI_NEW_GPU_PROFILER == 0)
-	RHI_API static int32 GetGProfileGPUTransitions();
-#endif
-
 	inline void BeginTransitions(TArrayView<const FRHITransition*> Transitions)
 	{
-#if WITH_PROFILEGPU && (RHI_NEW_GPU_PROFILER == 0)
-		RHI_BREADCRUMB_EVENT_CONDITIONAL(*this, GetGProfileGPUTransitions() != 0, "RHIBeginTransitions");
-#endif
 		if (Bypass())
 		{
 			GetComputeContext().RHIBeginTransitions(Transitions);
@@ -2802,9 +2747,6 @@ public:
 
 	inline void EndTransitions(TArrayView<const FRHITransition*> Transitions)
 	{
-#if WITH_PROFILEGPU && (RHI_NEW_GPU_PROFILER == 0)
-		RHI_BREADCRUMB_EVENT_CONDITIONAL(*this, GetGProfileGPUTransitions() != 0, "RHIEndTransitions");
-#endif
 		if (Bypass())
 		{
 			GetComputeContext().RHIEndTransitions(Transitions);
@@ -4057,18 +3999,6 @@ public:
 		GDynamicRHI->RHIEndRenderQuery_TopOfPipe(*this, RenderQuery);
 	}
 
-#if (RHI_NEW_GPU_PROFILER == 0)
-	inline void CalibrateTimers(FRHITimestampCalibrationQuery* CalibrationQuery)
-	{
-		if (Bypass())
-		{
-			GetContext().RHICalibrateTimers(CalibrationQuery);
-			return;
-		}
-		ALLOC_COMMAND(FRHICommandCalibrateTimers)(CalibrationQuery);
-	}
-#endif
-
 	inline void BeginRenderPass(const FRHIRenderPassInfo& InInfo, const TCHAR* Name)
 	{
 		check(!IsInsideRenderPass());
@@ -4434,13 +4364,9 @@ class FRHICommandListImmediate : public FRHICommandList
 	FRHICommandListImmediate()
 		: FRHICommandList(FRHIGPUMask::All(), true)
 	{
-#if WITH_RHI_BREADCRUMBS
+	#if WITH_RHI_BREADCRUMBS
 		PersistentState.LocalBreadcrumb = nullptr;
-#endif
-
-#if (RHI_NEW_GPU_PROFILER == 0) && HAS_GPU_STATS
-		PersistentState.CurrentDrawStatsCategory = nullptr;
-#endif
+	#endif
 	}
 
 	~FRHICommandListImmediate()
@@ -5235,10 +5161,6 @@ private:
 		bool bUsingSubCmdLists = false;
 		bool bShouldFinalize = true;
 
-#if (RHI_NEW_GPU_PROFILER == 0)
-		FRHIDrawStats DrawStats{};
-#endif
-
 		FTaskPipe* GetTranslateTaskPipe(ENamedThreads::Type& NamedThread);
 		FTaskPipe* EnqueueTranslateTask(FGraphEventArray&& Prereqs, TFunction<void()>&& Lambda);
 
@@ -5259,13 +5181,9 @@ private:
 		bool bAllowSingleParallelCombine = false;
 		bool bAllowParallelTranslate = true;
 
-#if WITH_RHI_BREADCRUMBS
+	#if WITH_RHI_BREADCRUMBS
 		bool bEmitBreadcrumbs = false;
-#endif
-
-#if (RHI_NEW_GPU_PROFILER == 0)
-		FRHIDrawStats DrawStats {};
-#endif
+	#endif
 
 		ERHISubmitFlags SubmitFlags = ERHISubmitFlags::None;
 		TArray<FRHIResource*> ResourcesToDelete {};
@@ -5316,13 +5234,6 @@ private:
 	
 #endif
 
-#if (RHI_NEW_GPU_PROFILER == 0)
-#if HAS_GPU_STATS
-	FRHIDrawStatsCategory const* CurrentDrawStatsCategory = nullptr;
-#endif
-	FRHIDrawStats FrameDrawStats;
-#endif
-
 	// Counts the number of calls to RHIEndFrame, and is used in GPU profiler frame boundary events.
 	uint32 FrameNumber = 0;
 
diff --git a/Engine/Source/Runtime/RHI/Public/RHICommandListCommandExecutes.inl b/Engine/Source/Runtime/RHI/Public/RHICommandListCommandExecutes.inl
--- a/Engine/Source/Runtime/RHI/Public/RHICommandListCommandExecutes.inl
+++ b/Engine/Source/Runtime/RHI/Public/RHICommandListCommandExecutes.inl
@@ -471,14 +471,6 @@ void FRHICommandEndRenderQuery::Execute(FRHICommandListBase& CmdList)
 	INTERNAL_DECORATOR(RHIEndRenderQuery)(RenderQuery);
 }
 
-#if (RHI_NEW_GPU_PROFILER == 0)
-void FRHICommandCalibrateTimers::Execute(FRHICommandListBase& CmdList)
-{
-	RHISTAT(CalibrateTimers);
-	INTERNAL_DECORATOR(RHICalibrateTimers)(CalibrationQuery);
-}
-#endif
-
 void FRHICommandPostExternalCommandsReset::Execute(FRHICommandListBase& CmdList)
 {
 	RHISTAT(PostExternalCommandsReset);
diff --git a/Engine/Source/Runtime/RHI/Public/RHIContext.h b/Engine/Source/Runtime/RHI/Public/RHIContext.h
--- a/Engine/Source/Runtime/RHI/Public/RHIContext.h
+++ b/Engine/Source/Runtime/RHI/Public/RHIContext.h
@@ -38,9 +38,6 @@ enum class EAsyncComputeBudget;
 
 struct FRHIBufferRange;
 struct FRHIPerCategoryDrawStats;
-#if (RHI_NEW_GPU_PROFILER == 0)
-struct FRHIDrawStats;
-#endif
 struct FRHICopyTextureInfo;
 
 
@@ -563,11 +560,9 @@ public:
 #endif
 
 protected:
-#if RHI_NEW_GPU_PROFILER
 	// Used to accumulate draw call and primitive counts,
 	// via the RHI_DRAW_CALL_INC / RHI_DRAW_CALL_STATS macros.
 	UE::RHI::GPUProfiler::FEvent::FStats StatEvent {};
-#endif
 };
 
 /** Context that is used to generate Upload commands. */
@@ -711,13 +706,6 @@ public:
 
 	virtual void RHIEndRenderQuery(FRHIRenderQuery* RenderQuery) = 0;
 
-#if (RHI_NEW_GPU_PROFILER == 0)
-	virtual void RHICalibrateTimers(FRHITimestampCalibrationQuery* CalibrationQuery)
-	{
-		/* empty default implementation */
-	}
-#endif
-
 	// This method is queued with an RHIThread, otherwise it will flush after it is queued; without an RHI thread there is no benefit to queuing this frame advance commands
 	virtual void RHIEndDrawingViewport(FRHIViewport* Viewport, const FRHIPresentArgs& InPresentArgs) = 0;
 
diff --git a/Engine/Source/Runtime/RHI/Public/RHIDefinitions.h b/Engine/Source/Runtime/RHI/Public/RHIDefinitions.h
--- a/Engine/Source/Runtime/RHI/Public/RHIDefinitions.h
+++ b/Engine/Source/Runtime/RHI/Public/RHIDefinitions.h
@@ -65,22 +65,13 @@ static_assert(sizeof(void*) <= SHADER_PARAMETER_POINTER_ALIGNMENT, "The alignmen
 #define RHI_RAYTRACING 0
 #endif
 
-#ifndef RHI_NEW_GPU_PROFILER
-#define RHI_NEW_GPU_PROFILER 1
-#endif
-
-#ifndef GPUPROFILERTRACE_ENABLED
-#if UE_TRACE_ENABLED && !UE_BUILD_SHIPPING
-// The trace implementation for the old GPU Profiler (FGpuProfilerTrace API) is deprecated in UE 5.6.
-// The new GPU Profiler has a different trace implementation (see UE_TRACE_GPU_PROFILER_ENABLED and FGPUProfilerSink_Trace).
-#define GPUPROFILERTRACE_ENABLED !RHI_NEW_GPU_PROFILER
-#else
-#define GPUPROFILERTRACE_ENABLED 0
-#endif
+#ifdef RHI_NEW_GPU_PROFILER
+#undef RHI_NEW_GPU_PROFILER
 #endif
+#define RHI_NEW_GPU_PROFILER UE_DEPRECATED_MACRO(5.8, "The legacy GPU profiler has been removed. RHI_NEW_GPU_PROFILER is deprecated, and always defined to 1. Remove any remaining uses of this macro.") 1
 
 #ifndef HAS_GPU_STATS
-#define HAS_GPU_STATS ((STATS || CSV_PROFILER_STATS || GPUPROFILERTRACE_ENABLED) && (!UE_BUILD_SHIPPING))
+#define HAS_GPU_STATS ((STATS || CSV_PROFILER_STATS) && (!UE_BUILD_SHIPPING))
 #endif
 
 /**
@@ -1041,7 +1032,7 @@ enum ERHIResourceType : uint8
 	RRT_Texture3D      UE_DEPRECATED(5.7,      "RRT_Texture3D is unused. Remove any remaining uses of it. There is a single RHI texture type which always uses RRT_Texture. RHI texture references use RRT_TextureReference."),
 	RRT_TextureCube    UE_DEPRECATED(5.7,    "RRT_TextureCube is unused. Remove any remaining uses of it. There is a single RHI texture type which always uses RRT_Texture. RHI texture references use RRT_TextureReference."),
 	RRT_TextureReference,
-	RRT_TimestampCalibrationQuery,
+	RRT_TimestampCalibrationQuery UE_DEPRECATED(5.8, "RRT_TimestampCalibrationQuery is unused. Remove any remaining uses of it."),
 	RRT_GPUFence,
 	RRT_RenderQuery,
 	RRT_RenderQueryPool,
diff --git a/Engine/Source/Runtime/RHI/Public/RHIFwd.h b/Engine/Source/Runtime/RHI/Public/RHIFwd.h
--- a/Engine/Source/Runtime/RHI/Public/RHIFwd.h
+++ b/Engine/Source/Runtime/RHI/Public/RHIFwd.h
@@ -82,9 +82,6 @@ class FRHIShaderBundle;
 class FRHIStagingBuffer;
 class FRHITexture;
 class FRHITextureReference;
-#if !defined(RHI_NEW_GPU_PROFILER) || (RHI_NEW_GPU_PROFILER == 0)
-class FRHITimestampCalibrationQuery;
-#endif
 class FRHIUniformBuffer;
 class FRHIUnorderedAccessView;
 class FRHIVertexDeclaration;
@@ -129,9 +126,6 @@ using FShaderBundleRHIRef              = TRefCountPtr<FRHIShaderBundle>;
 using FStagingBufferRHIRef             = TRefCountPtr<FRHIStagingBuffer>;
 using FTextureReferenceRHIRef          = TRefCountPtr<FRHITextureReference>;
 using FTextureRHIRef                   = TRefCountPtr<FRHITexture>;
-#if !defined(RHI_NEW_GPU_PROFILER) || (RHI_NEW_GPU_PROFILER == 0)
-using FTimestampCalibrationQueryRHIRef = TRefCountPtr<FRHITimestampCalibrationQuery>;
-#endif
 using FUniformBufferLayoutRHIRef       = TRefCountPtr<const FRHIUniformBufferLayout>;
 using FUniformBufferRHIRef             = TRefCountPtr<FRHIUniformBuffer>;
 using FUnorderedAccessViewRHIRef       = TRefCountPtr<FRHIUnorderedAccessView>;
diff --git a/Engine/Source/Runtime/RHI/Public/RHIGlobals.h b/Engine/Source/Runtime/RHI/Public/RHIGlobals.h
--- a/Engine/Source/Runtime/RHI/Public/RHIGlobals.h
+++ b/Engine/Source/Runtime/RHI/Public/RHIGlobals.h
@@ -369,12 +369,6 @@ struct FRHIGlobals
 	/** true for each VET that is supported. One-to-one mapping with EVertexElementType */
 	FVertexElementTypeSupportInfo VertexElementTypeSupport;
 
-	/** Whether the next frame should profile the GPU. */
-	bool TriggerGPUProfile = false;
-
-	/** Whether we are profiling GPU hitches. */
-	bool TriggerGPUHitchProfile = false;
-
 	/** Whether an intentional GPU crash has been scheduled. */
 	ERequestedGPUCrash TriggerGPUCrash = ERequestedGPUCrash::None;
 
@@ -852,8 +846,6 @@ extern RHI_API FRHIGlobals GRHIGlobals;
 #define GDrawUPVertexCheckCount                                GRHIGlobals.DrawUPVertexCheckCount
 #define GDrawUPIndexCheckCount                                 GRHIGlobals.DrawUPIndexCheckCount
 #define GVertexElementTypeSupport                              GRHIGlobals.VertexElementTypeSupport
-#define GTriggerGPUProfile                                     GRHIGlobals.TriggerGPUProfile
-#define GTriggerGPUHitchProfile                                GRHIGlobals.TriggerGPUHitchProfile
 #define GTriggerGPUCrash                                       GRHIGlobals.TriggerGPUCrash
 #define GGPUTraceFileName                                      GRHIGlobals.GPUTraceFileName
 #define GRHISupportsTextureStreaming                           GRHIGlobals.SupportsTextureStreaming
diff --git a/Engine/Source/Runtime/RHI/Public/RHIResources.h b/Engine/Source/Runtime/RHI/Public/RHIResources.h
--- a/Engine/Source/Runtime/RHI/Public/RHIResources.h
+++ b/Engine/Source/Runtime/RHI/Public/RHIResources.h
@@ -2342,16 +2342,6 @@ private:
 // Misc
 //
 
-#if (RHI_NEW_GPU_PROFILER == 0)
-class FRHITimestampCalibrationQuery : public FRHIResource
-{
-public:
-	FRHITimestampCalibrationQuery() : FRHIResource(RRT_TimestampCalibrationQuery) {}
-	uint64 GPUMicroseconds[MAX_NUM_GPUS] = {};
-	uint64 CPUMicroseconds[MAX_NUM_GPUS] = {};
-};
-#endif
-
 class FRHIGPUFence : public FRHIResource
 {
 public:
diff --git a/Engine/Source/Runtime/RHI/Public/RHIStats.h b/Engine/Source/Runtime/RHI/Public/RHIStats.h
--- a/Engine/Source/Runtime/RHI/Public/RHIStats.h
+++ b/Engine/Source/Runtime/RHI/Public/RHIStats.h
@@ -76,11 +76,6 @@ struct FRHIDrawStatsCategory
 {
 	RHI_API FRHIDrawStatsCategory(FName InName);
 
-#if (RHI_NEW_GPU_PROFILER == 0)
-	RHI_API FRHIDrawStatsCategory();
-	bool ShouldCountDraws() const { return Index != -1; }
-#endif
-
 	FName  const Name;
 	uint32 const Index;
 
@@ -101,17 +96,10 @@ struct FRHIDrawStatsCategory
 	RHI_API static FManager& GetManager();
 };
 
-#if RHI_NEW_GPU_PROFILER
-	DECLARE_DWORD_COUNTER_STAT_EXTERN(TEXT("Draws"     ), STAT_RHIDraws     , STATGROUP_RHI, RHI_API);
-	DECLARE_DWORD_COUNTER_STAT_EXTERN(TEXT("Dispatches"), STAT_RHIDispatches, STATGROUP_RHI, RHI_API);
-	DECLARE_DWORD_COUNTER_STAT_EXTERN(TEXT("Primitives"), STAT_RHIPrimitives, STATGROUP_RHI, RHI_API);
-	DECLARE_DWORD_COUNTER_STAT_EXTERN(TEXT("Vertices"  ), STAT_RHIVertices  , STATGROUP_RHI, RHI_API);
-#else
-	// RHI counter stats.
-	DECLARE_DWORD_COUNTER_STAT_EXTERN(TEXT("DrawPrimitive calls"), STAT_RHIDrawPrimitiveCalls, STATGROUP_RHI, RHI_API);
-	DECLARE_DWORD_COUNTER_STAT_EXTERN(TEXT("Triangles drawn"), STAT_RHITriangles, STATGROUP_RHI, RHI_API);
-	DECLARE_DWORD_COUNTER_STAT_EXTERN(TEXT("Lines drawn"), STAT_RHILines, STATGROUP_RHI, RHI_API);
-#endif
+DECLARE_DWORD_COUNTER_STAT_EXTERN(TEXT("Draws"     ), STAT_RHIDraws     , STATGROUP_RHI, RHI_API);
+DECLARE_DWORD_COUNTER_STAT_EXTERN(TEXT("Dispatches"), STAT_RHIDispatches, STATGROUP_RHI, RHI_API);
+DECLARE_DWORD_COUNTER_STAT_EXTERN(TEXT("Primitives"), STAT_RHIPrimitives, STATGROUP_RHI, RHI_API);
+DECLARE_DWORD_COUNTER_STAT_EXTERN(TEXT("Vertices"  ), STAT_RHIVertices  , STATGROUP_RHI, RHI_API);
 
 #else
 
@@ -124,123 +112,16 @@ struct FRHIDrawStatsCategory
 
 // Macros for use inside RHI context Draw/Dispatch functions.
 // Updates the Stats structure on the executing RHI command list
-#if RHI_NEW_GPU_PROFILER
-	#define RHI_DRAW_CALL_STATS(Type,Verts,Prims,Instances)                \
-	do                                                                     \
-	{                                                                      \
-		StatEvent.NumDraws++;                                              \
-		StatEvent.NumPrimitives += Prims * FMath::Max(1u, Instances);      \
-		StatEvent.NumVertices   += Verts * FMath::Max(1u, Instances);      \
+#define RHI_DRAW_CALL_STATS(Type,Verts,Prims,Instances)					\
+	do																	\
+	{																	\
+		StatEvent.NumDraws++;											\
+		StatEvent.NumPrimitives += Prims * FMath::Max(1u, Instances);	\
+		StatEvent.NumVertices   += Verts * FMath::Max(1u, Instances);	\
 	} while (false)
 
-	#define RHI_DRAW_CALL_INC()                             do { StatEvent.NumDraws++; } while (false)
-	#define RHI_DISPATCH_CALL_INC()			                do { StatEvent.NumDispatches++; } while (false)
-#else
-	#define RHI_DRAW_CALL_INC()                             do { GetExecutingCommandList().Stats_AddDraw(); } while (false)
-	#define RHI_DRAW_CALL_STATS(Type,Verts,Prims,Instances) do { GetExecutingCommandList().Stats_AddDrawAndPrimitives(Type, Prims); } while (false)
-	#define RHI_DISPATCH_CALL_INC()			                do { } while (false)
-
-	struct FRHIDrawStats
-	{
-	#if HAS_GPU_STATS
-		// The +1 is for "uncategorised"
-		static constexpr int32 NumCategories = FRHIDrawStatsCategory::MAX_DRAWCALL_CATEGORY + 1;
-	#else
-		static constexpr int32 NumCategories = 1;
-	#endif
-
-		static constexpr int32 NoCategory = NumCategories - 1;
-
-		struct FPerCategory
-		{
-			uint32 Draws;
-			uint32 Triangles;
-			uint32 Lines;
-			uint32 Quads;
-			uint32 Points;
-			uint32 Rectangles;
-
-			uint32 GetTotalPrimitives() const
-			{
-				return Triangles
-					+ Lines
-					+ Quads
-					+ Points
-					+ Rectangles;
-			}
-
-			FPerCategory& operator += (FPerCategory const& RHS)
-			{
-				Draws      += RHS.Draws;
-				Triangles  += RHS.Triangles;
-				Lines      += RHS.Lines;
-				Quads      += RHS.Quads;
-				Points     += RHS.Points;
-				Rectangles += RHS.Rectangles;
-				return *this;
-			}
-		};
-
-		struct FPerGPU
-		{
-			friend struct FRHIDrawStats;
-			FPerCategory Categories[NumCategories];
-		};
-
-		FPerGPU& GetGPU(uint32 GPUIndex)
-		{
-			checkSlow(GPUIndex < UE_ARRAY_COUNT(GPUs));
-			return GPUs[GPUIndex];
-		}
-
-		FRHIDrawStats()
-		{
-			Reset();
-		}
-
-		void Reset()
-		{
-			FMemory::Memzero(*this);
-		}
-
-		void AddDraw(FRHIGPUMask GPUMask, FRHIDrawStatsCategory const* Category)
-		{
-			uint32 CategoryIndex = Category ? Category->Index : NoCategory;
-			for (uint32 GPUIndex : GPUMask)
-			{
-				FPerCategory& Stats = GPUs[GPUIndex].Categories[CategoryIndex];
-				Stats.Draws++;
-			}
-		}
-
-		void AddDrawAndPrimitives(FRHIGPUMask GPUMask, FRHIDrawStatsCategory const* Category, EPrimitiveType PrimitiveType, uint32 NumPrimitives)
-		{
-			uint32 CategoryIndex = Category ? Category->Index : NoCategory;
-			for (uint32 GPUIndex : GPUMask)
-			{
-				FPerCategory& Stats = GPUs[GPUIndex].Categories[CategoryIndex];
-				Stats.Draws++;
-
-				switch (PrimitiveType)
-				{
-				case PT_TriangleList : Stats.Triangles  += NumPrimitives; break;
-				case PT_TriangleStrip: Stats.Triangles  += NumPrimitives; break;
-				case PT_LineList     : Stats.Lines      += NumPrimitives; break;
-				case PT_QuadList     : Stats.Quads      += NumPrimitives; break;
-				case PT_PointList    : Stats.Points     += NumPrimitives; break;
-				case PT_RectList     : Stats.Rectangles += NumPrimitives; break;
-				}
-			}
-		}
-
-		RHI_API void Accumulate(FRHIDrawStats& RHS);
-		RHI_API void ProcessAsFrameStats();
-
-	private:
-		FPerGPU GPUs[MAX_NUM_GPUS];
-	};
-
-#endif
+#define RHI_DRAW_CALL_INC()			do { StatEvent.NumDraws++;      } while (false)
+#define RHI_DISPATCH_CALL_INC()		do { StatEvent.NumDispatches++; } while (false)
 
 // RHI memory stats.
 DECLARE_MEMORY_STAT_POOL_EXTERN(TEXT("Render Target 2D Memory"), STAT_RenderTargetMemory2D, STATGROUP_RHI, FPlatformMemory::MCR_GPU, RHI_API);
diff --git a/Engine/Source/Runtime/RHI/Public/RHIValidationContext.h b/Engine/Source/Runtime/RHI/Public/RHIValidationContext.h
--- a/Engine/Source/Runtime/RHI/Public/RHIValidationContext.h
+++ b/Engine/Source/Runtime/RHI/Public/RHIValidationContext.h
@@ -758,13 +758,6 @@ public:
 		RHIContext->RHIEndRenderQuery(RenderQuery);
 	}
 
-#if (RHI_NEW_GPU_PROFILER == 0)
-	virtual void RHICalibrateTimers(FRHITimestampCalibrationQuery* CalibrationQuery) override final
-	{
-		RHIContext->RHICalibrateTimers(CalibrationQuery);
-	}
-#endif
-
 	// This method is queued with an RHIThread, otherwise it will flush after it is queued; without an RHI thread there is no benefit to queuing this frame advance commands
 	virtual void RHIEndDrawingViewport(FRHIViewport* Viewport, const FRHIPresentArgs& InPresentArgs) override final
 	{
diff --git a/Engine/Source/Runtime/RenderCore/Private/ProfilingDebugging/RealtimeGPUProfiler.cpp b/Engine/Source/Runtime/RenderCore/Private/ProfilingDebugging/RealtimeGPUProfiler.cpp
deleted file mode 100644
diff --git a/Engine/Source/Runtime/RenderCore/Private/RenderGraphPrivate.cpp b/Engine/Source/Runtime/RenderCore/Private/RenderGraphPrivate.cpp
--- a/Engine/Source/Runtime/RenderCore/Private/RenderGraphPrivate.cpp
+++ b/Engine/Source/Runtime/RenderCore/Private/RenderGraphPrivate.cpp
@@ -659,12 +659,7 @@ FRDGScopeState::FState::FState(bool bInImmediate, bool bInParallelExecute)
 		bool bRDGChannelEnabled = false;
 		IF_RDG_ENABLE_TRACE(bRDGChannelEnabled = UE_TRACE_CHANNELEXPR_IS_ENABLED(RDGChannel));
 
-		if (FRDGBuilder::IsDumpingFrame()
-			|| GTriggerGPUProfile
-		#if RHI_NEW_GPU_PROFILER
-			|| UE::RHI::GPUProfiler::IsProfiling()
-		#endif
-			)
+		if (FRDGBuilder::IsDumpingFrame() || UE::RHI::GPUProfiler::IsProfiling())
 		{
 			// We want all possible scope and pass names in a DumpGPU/profilegpu trace.
 			return ERDGScopeMode::AllEventsAndPassNames;
diff --git a/Engine/Source/Runtime/RenderCore/Private/RenderGraphPrivate.h b/Engine/Source/Runtime/RenderCore/Private/RenderGraphPrivate.h
--- a/Engine/Source/Runtime/RenderCore/Private/RenderGraphPrivate.h
+++ b/Engine/Source/Runtime/RenderCore/Private/RenderGraphPrivate.h
@@ -232,7 +232,7 @@ FORCEINLINE bool IsAsyncComputeSupported(EShaderPlatform ShaderPlatform)
 {
 	// Render pass merging and async compute are mutually exclusive since they are mutually dependent on each other (i.e. we can't
 	// merge render passes that involve async compute fences, since transitions would be relocated to the ends of the merged region).
-	return GRDGAsyncCompute > 0 && !IsImmediateMode() && !IsRenderPassMergeEnabled(ShaderPlatform) && GSupportsEfficientAsyncCompute && GRHISupportsSeparateDepthStencilCopyAccess && !GTriggerGPUProfile;
+	return GRDGAsyncCompute > 0 && !IsImmediateMode() && !IsRenderPassMergeEnabled(ShaderPlatform) && GSupportsEfficientAsyncCompute && GRHISupportsSeparateDepthStencilCopyAccess;
 }
 
 extern bool IsParallelExecuteEnabled(EShaderPlatform ShaderPlatform);
diff --git a/Engine/Source/Runtime/RenderCore/Private/RenderingThread.cpp b/Engine/Source/Runtime/RenderCore/Private/RenderingThread.cpp
--- a/Engine/Source/Runtime/RenderCore/Private/RenderingThread.cpp
+++ b/Engine/Source/Runtime/RenderCore/Private/RenderingThread.cpp
@@ -288,14 +288,10 @@ static void AdvanceRenderingThreadStats(int64 StatsFrame, int32 DisableChangeTag
 	}
 	FThreadStats::AddMessage(FStatConstants::AdvanceFrame.GetEncodedName(), EStatOperation::AdvanceFrameEventRenderThread, Frame);
 
-#if RHI_NEW_GPU_PROFILER
 	// Propagate the stats frame value down to the end-of-pipe thread.
 	UE::Stats::FStats::StatsFrameRT = Frame;
-#else
-	FThreadStats::AddMessage(FStatConstants::AdvanceFrame.GetEncodedName(), EStatOperation::AdvanceFrameEventEndOfPipe, Frame);
-#endif
 
-	if( IsInActualRenderingThread() )
+	if (IsInActualRenderingThread())
 	{
 		FThreadStats::ExplicitFlush();
 	}
diff --git a/Engine/Source/Runtime/RenderCore/Public/ProfilingDebugging/RealtimeGPUProfiler.h b/Engine/Source/Runtime/RenderCore/Public/ProfilingDebugging/RealtimeGPUProfiler.h
--- a/Engine/Source/Runtime/RenderCore/Public/ProfilingDebugging/RealtimeGPUProfiler.h
+++ b/Engine/Source/Runtime/RenderCore/Public/ProfilingDebugging/RealtimeGPUProfiler.h
@@ -151,258 +151,56 @@
 #define SCOPED_RHI_CONDITIONAL_DRAW_EVENT_COLOR(RHICmdContext, Color, Name, Condition)               UE_DEPRECATED_MACRO(5.5, "SCOPED_RHI_CONDITIONAL_DRAW_EVENT_COLOR has been deprecated. Use standard RHI breadcrumb events instead." )
 #define SCOPED_RHI_CONDITIONAL_DRAW_EVENTF_COLOR(RHICmdContext, Color, Name, Condition, Format, ...) UE_DEPRECATED_MACRO(5.5, "SCOPED_RHI_CONDITIONAL_DRAW_EVENTF_COLOR has been deprecated. Use standard RHI breadcrumb events instead.")
 
-#if RHI_NEW_GPU_PROFILER
-
-	#if HAS_GPU_STATS
-
-		#define DECLARE_GPU_STAT_NAME_TYPE(StatName, NameString) \
-			struct TRHIGPUStatNameProvider_##StatName            \
-			{                                                    \
-				static constexpr TCHAR const* GetDisplayName()   \
-				{                                                \
-					return NameString;                           \
-				}                                                \
-				static constexpr TCHAR const* GetStatName()      \
-				{                                                \
-					return TEXT(#StatName);                      \
-				}                                                \
-			}
-
-		// Extern GPU stats are needed where a stat is used in multiple CPPs. Use the DECLARE_GPU_STAT_NAMED_EXTERN in the header and DEFINE_GPU_STAT in the CPPs
-		#define DECLARE_GPU_STAT_NAMED(                StatName, NameString) DECLARE_GPU_STAT_NAME_TYPE(StatName, NameString); static UE::RHI::GPUProfiler::TGPUStat                    <TRHIGPUStatNameProvider_##StatName> GPUStat_##StatName
-		#define DECLARE_GPU_STAT_NAMED_EXTERN(         StatName, NameString) DECLARE_GPU_STAT_NAME_TYPE(StatName, NameString); extern UE::RHI::GPUProfiler::TGPUStat                    <TRHIGPUStatNameProvider_##StatName> GPUStat_##StatName
-		#define DECLARE_GPU_DRAWCALL_STAT_NAMED(       StatName, NameString) DECLARE_GPU_STAT_NAME_TYPE(StatName, NameString); static UE::RHI::GPUProfiler::TGPUStatWithDrawcallCategory<TRHIGPUStatNameProvider_##StatName> GPUStat_##StatName
-		#define DECLARE_GPU_DRAWCALL_STAT_NAMED_EXTERN(StatName, NameString) DECLARE_GPU_STAT_NAME_TYPE(StatName, NameString); extern UE::RHI::GPUProfiler::TGPUStatWithDrawcallCategory<TRHIGPUStatNameProvider_##StatName> GPUStat_##StatName
-
-		#define DECLARE_GPU_STAT(                StatName) DECLARE_GPU_STAT_NAMED(                StatName, TEXT(#StatName))
-		#define DECLARE_GPU_STAT_EXTERN(         StatName) DECLARE_GPU_STAT_NAMED_EXTERN(         StatName, TEXT(#StatName))
-		#define DECLARE_GPU_DRAWCALL_STAT(       StatName) DECLARE_GPU_DRAWCALL_STAT_NAMED(       StatName, TEXT(#StatName))
-		#define DECLARE_GPU_DRAWCALL_STAT_EXTERN(StatName) DECLARE_GPU_DRAWCALL_STAT_NAMED_EXTERN(StatName, TEXT(#StatName))
-
-		#define DEFINE_GPU_STAT(         StatName) UE::RHI::GPUProfiler::TGPUStat                    <TRHIGPUStatNameProvider_##StatName> GPUStat_##StatName
-		#define DEFINE_GPU_DRAWCALL_STAT(StatName) UE::RHI::GPUProfiler::TGPUStatWithDrawcallCategory<TRHIGPUStatNameProvider_##StatName> GPUStat_##StatName
-
-	#else
-
-		#define DECLARE_GPU_STAT_NAMED(...)
-		#define DECLARE_GPU_STAT_NAMED_EXTERN(...)
-		#define DECLARE_GPU_DRAWCALL_STAT_NAMED(...)
-		#define DECLARE_GPU_DRAWCALL_STAT_NAMED_EXTERN(...)
-
-		#define DECLARE_GPU_STAT(...)
-		#define DECLARE_GPU_STAT_EXTERN(...)
-		#define DECLARE_GPU_DRAWCALL_STAT(...)
-		#define DECLARE_GPU_DRAWCALL_STAT_EXTERN(...)
-
-		#define DEFINE_GPU_STAT(...)
-		#define DEFINE_GPU_DRAWCALL_STAT(...)
-
-	#endif
-
-	// Empty when using the new RHI GPU profiler. GPU stats are handled via RHI_BREADCRUMB_EVENT_STAT etc.
-	// @todo deprecate
-	#define SCOPED_GPU_STAT_VERBOSE(...)
-	#define SCOPED_GPU_STAT(...)
-	#define GPU_STATS_BEGINFRAME(...) 
-	#define GPU_STATS_ENDFRAME(...) 
-	#define GPU_STATS_SUSPENDFRAME(...)
-
-#else
-
-class FRealtimeGPUProfiler;
-class FRealtimeGPUProfilerEvent;
-class FRealtimeGPUProfilerFrame;
-
-DECLARE_STATS_GROUP(TEXT("GPU"), STATGROUP_GPU, STATCAT_Advanced);
-CSV_DECLARE_CATEGORY_MODULE_EXTERN(RENDERCORE_API, GPU);
-
 #if HAS_GPU_STATS
 
-	// The DECLARE_GPU_STAT macros both declare and define a stat (for use in a single CPP)
-	#define DECLARE_GPU_STAT(StatName)                                   DECLARE_FLOAT_COUNTER_STAT(TEXT(#StatName)       , Stat_GPU_##StatName, STATGROUP_GPU  ); CSV_DEFINE_STAT(GPU,StatName);         static FRHIDrawStatsCategory DrawcallCountCategory_##StatName;
-	#define DECLARE_GPU_STAT_NAMED(StatName, NameString)                 DECLARE_FLOAT_COUNTER_STAT(NameString            , Stat_GPU_##StatName, STATGROUP_GPU  ); CSV_DEFINE_STAT(GPU,StatName);         static FRHIDrawStatsCategory DrawcallCountCategory_##StatName;
-	#define DECLARE_GPU_DRAWCALL_STAT(StatName)                          DECLARE_FLOAT_COUNTER_STAT(TEXT(#StatName)       , Stat_GPU_##StatName, STATGROUP_GPU  ); CSV_DEFINE_STAT(GPU,StatName);         static FRHIDrawStatsCategory DrawcallCountCategory_##StatName((TCHAR*)TEXT(#StatName));
-	#define DECLARE_GPU_DRAWCALL_STAT_NAMED(StatName, NameString)        DECLARE_FLOAT_COUNTER_STAT(NameString            , Stat_GPU_##StatName, STATGROUP_GPU  ); CSV_DEFINE_STAT(GPU,StatName);         static FRHIDrawStatsCategory DrawcallCountCategory_##StatName((TCHAR*)TEXT(#StatName));
-	#define DECLARE_GPU_DRAWCALL_STAT_EXTERN(StatName)                   DECLARE_FLOAT_COUNTER_STAT_EXTERN(TEXT(#StatName), Stat_GPU_##StatName, STATGROUP_GPU, ); CSV_DECLARE_STAT_EXTERN(GPU,StatName); extern FRHIDrawStatsCategory DrawcallCountCategory_##StatName;
-	#define DECLARE_GPU_DRAWCALL_STAT_NAMED_EXTERN(StatName, NameString) DECLARE_FLOAT_COUNTER_STAT_EXTERN(NameString     , Stat_GPU_##StatName, STATGROUP_GPU, ); CSV_DECLARE_STAT_EXTERN(GPU,StatName); extern FRHIDrawStatsCategory DrawcallCountCategory_##StatName;
+	#define DECLARE_GPU_STAT_NAME_TYPE(StatName, NameString) \
+		struct TRHIGPUStatNameProvider_##StatName            \
+		{                                                    \
+			static constexpr TCHAR const* GetDisplayName()   \
+			{                                                \
+				return NameString;                           \
+			}                                                \
+			static constexpr TCHAR const* GetStatName()      \
+			{                                                \
+				return TEXT(#StatName);                      \
+			}                                                \
+		}
 
 	// Extern GPU stats are needed where a stat is used in multiple CPPs. Use the DECLARE_GPU_STAT_NAMED_EXTERN in the header and DEFINE_GPU_STAT in the CPPs
-	#define DECLARE_GPU_STAT_NAMED_EXTERN(StatName, NameString) DECLARE_FLOAT_COUNTER_STAT_EXTERN(NameString, Stat_GPU_##StatName, STATGROUP_GPU, ); CSV_DECLARE_STAT_EXTERN(GPU,StatName); extern FRHIDrawStatsCategory DrawcallCountCategory_##StatName;
-	#define DEFINE_GPU_STAT(StatName)                           DEFINE_STAT(Stat_GPU_##StatName);                                                    CSV_DEFINE_STAT(GPU,StatName);                FRHIDrawStatsCategory DrawcallCountCategory_##StatName;
-	#define DEFINE_GPU_DRAWCALL_STAT(StatName)                  DEFINE_STAT(Stat_GPU_##StatName);                                                    CSV_DEFINE_STAT(GPU,StatName);                FRHIDrawStatsCategory DrawcallCountCategory_##StatName((TCHAR*)TEXT(#StatName));
+	#define DECLARE_GPU_STAT_NAMED(                StatName, NameString) DECLARE_GPU_STAT_NAME_TYPE(StatName, NameString); static UE::RHI::GPUProfiler::TGPUStat                    <TRHIGPUStatNameProvider_##StatName> GPUStat_##StatName
+	#define DECLARE_GPU_STAT_NAMED_EXTERN(         StatName, NameString) DECLARE_GPU_STAT_NAME_TYPE(StatName, NameString); extern UE::RHI::GPUProfiler::TGPUStat                    <TRHIGPUStatNameProvider_##StatName> GPUStat_##StatName
+	#define DECLARE_GPU_DRAWCALL_STAT_NAMED(       StatName, NameString) DECLARE_GPU_STAT_NAME_TYPE(StatName, NameString); static UE::RHI::GPUProfiler::TGPUStatWithDrawcallCategory<TRHIGPUStatNameProvider_##StatName> GPUStat_##StatName
+	#define DECLARE_GPU_DRAWCALL_STAT_NAMED_EXTERN(StatName, NameString) DECLARE_GPU_STAT_NAME_TYPE(StatName, NameString); extern UE::RHI::GPUProfiler::TGPUStatWithDrawcallCategory<TRHIGPUStatNameProvider_##StatName> GPUStat_##StatName
 
-	#define SCOPED_GPU_STAT_VERBOSE(RHICmdList, StatName, Description)			\
-		FScopedGPUStatEvent UE_JOIN(GPUStatEvent_##StatName,__LINE__)(			\
-			  RHICmdList														\
-			, CSV_STAT_FNAME(StatName)											\
-			, GET_STATID(Stat_GPU_##StatName)									\
-			, Description														\
-		);																		\
-		FScopedDrawStatCategory UE_JOIN(DrawCallScope_##StatName, __LINE__)(	\
-			RHICmdList, DrawcallCountCategory_##StatName						\
-		);
+	#define DECLARE_GPU_STAT(                StatName) DECLARE_GPU_STAT_NAMED(                StatName, TEXT(#StatName))
+	#define DECLARE_GPU_STAT_EXTERN(         StatName) DECLARE_GPU_STAT_NAMED_EXTERN(         StatName, TEXT(#StatName))
+	#define DECLARE_GPU_DRAWCALL_STAT(       StatName) DECLARE_GPU_DRAWCALL_STAT_NAMED(       StatName, TEXT(#StatName))
+	#define DECLARE_GPU_DRAWCALL_STAT_EXTERN(StatName) DECLARE_GPU_DRAWCALL_STAT_NAMED_EXTERN(StatName, TEXT(#StatName))
 
-	#define SCOPED_GPU_STAT(RHICmdList, StatName) SCOPED_GPU_STAT_VERBOSE(RHICmdList, StatName, nullptr)
-
-	#define GPU_STATS_BEGINFRAME(RHICmdList) FRealtimeGPUProfiler::Get()->BeginFrame(RHICmdList);
-	#define GPU_STATS_ENDFRAME(RHICmdList)   FRealtimeGPUProfiler::Get()->EndFrame(RHICmdList);
-	#define GPU_STATS_SUSPENDFRAME()         FRealtimeGPUProfiler::Get()->SuspendFrame();
+	#define DEFINE_GPU_STAT(         StatName) UE::RHI::GPUProfiler::TGPUStat                    <TRHIGPUStatNameProvider_##StatName> GPUStat_##StatName
+	#define DEFINE_GPU_DRAWCALL_STAT(StatName) UE::RHI::GPUProfiler::TGPUStatWithDrawcallCategory<TRHIGPUStatNameProvider_##StatName> GPUStat_##StatName
 
 #else
 
-	#define DECLARE_GPU_STAT(StatName)
-	#define DECLARE_GPU_STAT_NAMED(StatName, NameString)
-	#define DECLARE_GPU_DRAWCALL_STAT(StatName)
-	#define DECLARE_GPU_DRAWCALL_STAT_NAMED(StatName, NameString)
-	#define DECLARE_GPU_DRAWCALL_STAT_EXTERN(StatName)
-	#define DECLARE_GPU_DRAWCALL_STAT_NAMED_EXTERN(StatName, NameString)
-
-	#define DECLARE_GPU_STAT_NAMED_EXTERN(StatName, NameString)
-	#define DEFINE_GPU_STAT(StatName)
-	#define DEFINE_GPU_DRAWCALL_STAT(StatName)
+	#define DECLARE_GPU_STAT_NAMED(...)
+	#define DECLARE_GPU_STAT_NAMED_EXTERN(...)
+	#define DECLARE_GPU_DRAWCALL_STAT_NAMED(...)
+	#define DECLARE_GPU_DRAWCALL_STAT_NAMED_EXTERN(...)
 
-	#define SCOPED_GPU_STAT_VERBOSE(RHICmdList, StatName, Description)
-	#define SCOPED_GPU_STAT(RHICmdList, StatName) 
+	#define DECLARE_GPU_STAT(...)
+	#define DECLARE_GPU_STAT_EXTERN(...)
+	#define DECLARE_GPU_DRAWCALL_STAT(...)
+	#define DECLARE_GPU_DRAWCALL_STAT_EXTERN(...)
 
-	#define GPU_STATS_BEGINFRAME(RHICmdList) 
-	#define GPU_STATS_ENDFRAME(RHICmdList) 
-	#define GPU_STATS_SUSPENDFRAME()
+	#define DEFINE_GPU_STAT(...)
+	#define DEFINE_GPU_DRAWCALL_STAT(...)
 
 #endif
 
-RENDERCORE_API bool AreGPUStatsEnabled();
-
-#if HAS_GPU_STATS
-
-class FRealtimeGPUProfilerQuery
-{
-public:
-	FRealtimeGPUProfilerQuery() = default;
-	FRealtimeGPUProfilerQuery(FRHIGPUMask InGPUMask, FRHIRenderQuery* InQuery, FRealtimeGPUProfilerEvent* Parent)
-		: GPUMask(InGPUMask)
-		, Query(InQuery)
-		, Parent(Parent)
-	{}
-
-	RENDERCORE_API void Submit(FRHICommandList& RHICmdList, bool bBegin) const;
-
-	// RDG might create profiler events that are never submitted due to pass culling etc.
-	// This is called when FRDGScope_GPU instances are destructed, and will mark this query as discarded if it was never submitted.
-	RENDERCORE_API void Discard(bool bBegin);
-
-	operator bool() const { return Query != nullptr; }
-
-private:
-	FRHIGPUMask GPUMask;
-	FRHIRenderQuery* Query{};
-	FRealtimeGPUProfilerEvent* Parent{};
-};
-
-#if GPUPROFILERTRACE_ENABLED
-struct FRealtimeGPUProfilerHistoryItem
-{
-	FRealtimeGPUProfilerHistoryItem();
-
-	static const uint64 HistoryCount = 64;
-
-	// Constructor memsets everything to zero, assuming structure is Plain Old Data.  If any dynamic structures are
-	// added, you'll need a more generalized constructor that zeroes out all the uninitialized data.
-	bool UpdatedThisFrame;
-	FRHIGPUMask LastGPUMask;
-	uint64 NextWriteIndex;
-	uint64 AccumulatedTime;				// Accumulated time could be computed, but may also be useful to inspect in the debugger
-	TStaticArray<uint64, HistoryCount> Times;
-};
-
-struct FRealtimeGPUProfilerHistoryByDescription
-{
-	TMap<FString, FRealtimeGPUProfilerHistoryItem> History;
-	mutable FRWLock Mutex;
-};
-
-struct FRealtimeGPUProfilerDescriptionResult
-{
-	// Times are in microseconds
-	FString Description;
-	FRHIGPUMask GPUMask;
-	uint64 AverageTime;
-	uint64 MinTime;
-	uint64 MaxTime;
-};
-#endif  // GPUPROFILERTRACE_ENABLED
-
-/**
-* FRealtimeGPUProfiler class. This manages recording and reporting all for GPU stats
-*/
-class FRealtimeGPUProfiler
-{
-	static FRealtimeGPUProfiler* Instance;
-public:
-	// Singleton interface
-	static RENDERCORE_API FRealtimeGPUProfiler* Get();
-
-	/** *Safe release of the singleton */
-	static RENDERCORE_API void SafeRelease();
-
-	/** Per-frame update */
-	RENDERCORE_API void BeginFrame(FRHICommandListImmediate& RHICmdList);
-	RENDERCORE_API void EndFrame(FRHICommandListImmediate& RHICmdList);
-	RENDERCORE_API void SuspendFrame();
-
-	/** Push/pop events */
-	RENDERCORE_API FRealtimeGPUProfilerQuery PushEvent(FRHIGPUMask GPUMask, const FName& Name, const TStatId& Stat, const TCHAR* Description);
-	RENDERCORE_API FRealtimeGPUProfilerQuery PopEvent();
-
-	/** Push/pop stats which do additional draw call tracking on top of events. */
-	RENDERCORE_API void PushStat(FRHICommandListImmediate& RHICmdList, const FName& Name, const TStatId& Stat, const TCHAR* Description);
-	RENDERCORE_API void PopStat(FRHICommandListImmediate& RHICmdList);
-
-#if GPUPROFILERTRACE_ENABLED
-	RENDERCORE_API void FetchPerfByDescription(TArray<FRealtimeGPUProfilerDescriptionResult> & OutResults) const;
-#endif
-
-private:
-	FRealtimeGPUProfiler();
-
-	/** Deinitialize of the object*/
-	void Cleanup();
-
-	UE::Tasks::FTask CleanUpTask;
-
-	TUniquePtr<FRealtimeGPUProfilerFrame> ActiveFrame;
-	TQueue<TUniquePtr<FRealtimeGPUProfilerFrame>> PendingFrames;
-
-	FRenderQueryPoolRHIRef RenderQueryPool;
-	bool bLocked = false;
-
-#if GPUPROFILERTRACE_ENABLED
-	FRealtimeGPUProfilerHistoryByDescription HistoryByDescription;
-#endif
-};
-
-/**
-* Class that logs GPU Stat events for the realtime GPU profiler
-*/
-class FScopedGPUStatEvent
-{
-	/** Cmdlist to push onto. */
-	FRHICommandListBase* RHICmdList = nullptr;
-
-public:
-	RENDERCORE_API FScopedGPUStatEvent(FRHICommandListBase& InRHICmdList, const FName& Name, const TStatId& StatId, const TCHAR* Description);
-	RENDERCORE_API ~FScopedGPUStatEvent();
-};
-
-class FScopedDrawStatCategory
-{
-	FRHICommandListBase* const RHICmdList = nullptr;
-	TOptional<FRHIDrawStatsCategory const*> Previous {};
-
-public:
-	RENDERCORE_API FScopedDrawStatCategory(FRHICommandListBase& RHICmdList, FRHIDrawStatsCategory const& Category);
-	RENDERCORE_API ~FScopedDrawStatCategory();
-};
-
-#endif // HAS_GPU_STATS
-
-#endif // (RHI_NEW_GPU_PROFILER == 0)
+// Empty when using the new RHI GPU profiler. GPU stats are handled via RHI_BREADCRUMB_EVENT_STAT etc.
+// @todo deprecate
+#define SCOPED_GPU_STAT_VERBOSE(...)
+#define SCOPED_GPU_STAT(...)
+#define GPU_STATS_BEGINFRAME(...) 
+#define GPU_STATS_ENDFRAME(...) 
+#define GPU_STATS_SUSPENDFRAME(...)
diff --git a/Engine/Source/Runtime/RenderCore/Public/RenderGraphEvent.h b/Engine/Source/Runtime/RenderCore/Public/RenderGraphEvent.h
--- a/Engine/Source/Runtime/RenderCore/Public/RenderGraphEvent.h
+++ b/Engine/Source/Runtime/RenderCore/Public/RenderGraphEvent.h
@@ -85,34 +85,6 @@ enum class ERDGScopeFlags : uint8
 };
 ENUM_CLASS_FLAGS(ERDGScopeFlags);
 
-#if HAS_GPU_STATS && (RHI_NEW_GPU_PROFILER == 0)
-// Scope type for the legacy "realtime" GPU profiler and draw call counter stats
-struct FRDGScope_GPU
-{
-	FRealtimeGPUProfilerQuery StartQuery;
-	FRealtimeGPUProfilerQuery StopQuery;
-
-	FName StatName;
-	TStatId StatId;
-	FString StatDescription;
-
-	TOptional<FRHIDrawStatsCategory const*> PreviousCategory {};
-	FRHIDrawStatsCategory const* CurrentCategory  = nullptr;
-	bool bEmitDuringExecute;
-
-	inline FRDGScope_GPU(FRDGScopeState& State, FRHIGPUMask GPUMask, const FName& CsvStatName, const TStatId& Stat, const TCHAR* Description, FRHIDrawStatsCategory const& Category);
-	inline ~FRDGScope_GPU();
-
-	inline void ImmediateEnd(FRDGScopeState& State);
-
-	inline void BeginCPU(FRHIComputeCommandList& RHICmdList, bool bPreScope);
-	inline void EndCPU  (FRHIComputeCommandList& RHICmdList, bool bPreScope);
-
-	inline void BeginGPU(FRHIComputeCommandList& RHICmdList);
-	inline void EndGPU  (FRHIComputeCommandList& RHICmdList);
-};
-#endif // HAS_GPU_STATS && (RHI_NEW_GPU_PROFILER == 0)
-
 #if CSV_PROFILER_STATS
 struct FRDGScope_CSVExclusive
 {
@@ -287,9 +259,6 @@ struct FRDGScope
 #if RDG_EVENTS
 		, FRDGScope_RHI
 #endif
-#if HAS_GPU_STATS && (RHI_NEW_GPU_PROFILER == 0)
-		, FRDGScope_GPU
-#endif
 #if CSV_PROFILER_STATS
 		, FRDGScope_CSVExclusive
 #endif
@@ -537,13 +506,8 @@ private:
 
 #endif
 
-#if HAS_GPU_STATS && (RHI_NEW_GPU_PROFILER == 0)
-	#define RDG_GPU_STAT_SCOPE(GraphBuilder, StatName)                      TRDGEventScopeGuard<FRDGScope_GPU> UE_JOIN(__RDG_GPUStatEvent_##StatName,__LINE__) ((GraphBuilder), ERDGScopeFlags::AlwaysEnable, (GraphBuilder).RHICmdList.GetGPUMask(), CSV_STAT_FNAME(StatName), GET_STATID(Stat_GPU_##StatName), nullptr    , DrawcallCountCategory_##StatName);
-	#define RDG_GPU_STAT_SCOPE_VERBOSE(GraphBuilder, StatName, Description) TRDGEventScopeGuard<FRDGScope_GPU> UE_JOIN(__RDG_GPUStatEvent_##StatName,__LINE__) ((GraphBuilder), ERDGScopeFlags::AlwaysEnable, (GraphBuilder).RHICmdList.GetGPUMask(), CSV_STAT_FNAME(StatName), GET_STATID(Stat_GPU_##StatName), Description, DrawcallCountCategory_##StatName);
-#else
-	#define RDG_GPU_STAT_SCOPE(GraphBuilder, StatName)
-	#define RDG_GPU_STAT_SCOPE_VERBOSE(GraphBuilder, StatName, Description)
-#endif
+#define RDG_GPU_STAT_SCOPE(GraphBuilder, StatName)
+#define RDG_GPU_STAT_SCOPE_VERBOSE(GraphBuilder, StatName, Description)
 
 #if CSV_PROFILER_STATS
 	#define RDG_CSV_STAT_EXCLUSIVE_SCOPE(GraphBuilder, StatName) TRDGEventScopeGuard<FRDGScope_CSVExclusive> UE_JOIN(__RDG_CSVStat_##StatName,__LINE__) ((GraphBuilder), ERDGScopeFlags::AlwaysEnable, #StatName);
@@ -670,9 +634,6 @@ public:
 	template <typename TScopeType>
 	friend class TRDGEventScopeGuard;
 	friend FRDGScope_Budget;
-#if HAS_GPU_STATS && (RHI_NEW_GPU_PROFILER == 0)
-	friend FRDGScope_GPU;
-#endif
 #if RDG_EVENTS
 	friend FRDGScope_RHI;
 #endif
diff --git a/Engine/Source/Runtime/RenderCore/Public/RenderGraphEvent.inl b/Engine/Source/Runtime/RenderCore/Public/RenderGraphEvent.inl
--- a/Engine/Source/Runtime/RenderCore/Public/RenderGraphEvent.inl
+++ b/Engine/Source/Runtime/RenderCore/Public/RenderGraphEvent.inl
@@ -2,102 +2,6 @@
 
 #pragma once
 
-#if RHI_NEW_GPU_PROFILER
-
-// @todo - new gpu profiler
-
-#else
-
-#if HAS_GPU_STATS
-inline FRDGScope_GPU::FRDGScope_GPU(FRDGScopeState& State, FRHIGPUMask GPUMask, const FName& CsvStatName, const TStatId& Stat, const TCHAR* Description, FRHIDrawStatsCategory const& Category)
-	: CurrentCategory(Category.ShouldCountDraws() ? &Category : nullptr),
-	  bEmitDuringExecute(AreGPUStatsEnabled() && !State.ScopeState.bParallelExecute)
-{
-	if (AreGPUStatsEnabled())
-	{
-		if (bEmitDuringExecute)
-		{ 
-			StatName = CsvStatName;
-			StatId = Stat;
-			StatDescription = FString(Description);
-		}
-		else
-		{
-			StartQuery = FRealtimeGPUProfiler::Get()->PushEvent(GPUMask, CsvStatName, Stat, Description);
-		}
-	}
-}
-
-inline void FRDGScope_GPU::ImmediateEnd(FRDGScopeState&)
-{
-	if (StartQuery)
-	{
-		StopQuery = FRealtimeGPUProfiler::Get()->PopEvent();
-	}
-}
-
-inline FRDGScope_GPU::~FRDGScope_GPU()
-{
-	if (StartQuery)
-	{
-		StartQuery.Discard(true);
-	}
-	if (StopQuery)
-	{
-		StopQuery.Discard(false);
-	}
-}
-
-inline void FRDGScope_GPU::BeginCPU(FRHIComputeCommandList& RHICmdList, bool bPreScope)
-{
-	if (CurrentCategory)
-	{
-		PreviousCategory = RHICmdList.SetDrawStatsCategory(CurrentCategory);
-	}
-}
-
-inline void FRDGScope_GPU::EndCPU(FRHIComputeCommandList& RHICmdList, bool bPreScope)
-{
-	if (CurrentCategory)
-	{
-		RHICmdList.SetDrawStatsCategory(PreviousCategory);
-	}
-}
-
-inline void FRDGScope_GPU::BeginGPU(FRHIComputeCommandList& RHICmdList)
-{
-	if (EnumHasAnyFlags(RHICmdList.GetPipeline(), ERHIPipeline::Graphics))
-	{
-		if (bEmitDuringExecute)
-		{
-			FRealtimeGPUProfiler::Get()->PushStat(RHICmdList.GetAsImmediate(), StatName, StatId, *StatDescription);
-		}
-		else if (StartQuery)
-		{
-			StartQuery.Submit(static_cast<FRHICommandList&>(RHICmdList), true);
-		}
-	}
-}
-
-inline void FRDGScope_GPU::EndGPU(FRHIComputeCommandList& RHICmdList)
-{
-	if (EnumHasAnyFlags(RHICmdList.GetPipeline(), ERHIPipeline::Graphics))
-	{
-		if (bEmitDuringExecute)
-		{
-			FRealtimeGPUProfiler::Get()->PopStat(RHICmdList.GetAsImmediate());
-		}
-		else if (StopQuery)
-		{
-			StopQuery.Submit(static_cast<FRHICommandList&>(RHICmdList), false);
-		}
-	
-	}
-}
-#endif
-
-#endif // (RHI_NEW_GPU_PROFILER == 0)
-
 #if RDG_EVENTS
 
 	template <typename TDesc, typename... TValues>
diff --git a/Engine/Source/Runtime/Renderer/Private/GPUBenchmark.cpp b/Engine/Source/Runtime/Renderer/Private/GPUBenchmark.cpp
--- a/Engine/Source/Runtime/Renderer/Private/GPUBenchmark.cpp
+++ b/Engine/Source/Runtime/Renderer/Private/GPUBenchmark.cpp
@@ -396,17 +396,6 @@ void RendererGPUBenchmark(FRHICommandListImmediate& RHICmdList, FSynthBenchmarkR
 {
 	check(IsInRenderingThread());
 
-#if (RHI_NEW_GPU_PROFILER == 0)
-	// Multi-GPU support : Benchmark needs to be made MGPU-aware.
-	bool bValidGPUTimer = (FGPUTiming::GetTimingFrequency() / (1000 * 1000)) != 0;
-
-	if(!bValidGPUTimer)
-	{
-		UE_LOG(LogSynthBenchmark, Warning, TEXT("RendererGPUBenchmark failed, look for \"GPU Timing Frequency\" in the log"));
-		return;
-	}
-#endif
-
 	RHICmdList.ImmediateFlush(EImmediateFlushType::FlushRHIThread);
 	MeasureLongGPUTaskExecutionTime(RHICmdList);
 	RHICmdList.ImmediateFlush(EImmediateFlushType::FlushRHIThread);	
diff --git a/Engine/Source/Runtime/Windows/D3D11RHI/Private/D3D11Query.cpp b/Engine/Source/Runtime/Windows/D3D11RHI/Private/D3D11Query.cpp
--- a/Engine/Source/Runtime/Windows/D3D11RHI/Private/D3D11Query.cpp
+++ b/Engine/Source/Runtime/Windows/D3D11RHI/Private/D3D11Query.cpp
@@ -207,7 +207,6 @@ bool FD3D11RenderQuery::CacheResult(FD3D11DynamicRHI& RHI, bool bWait)
 
 	case EType::Timestamp:
 		{
-			// GetTimingFrequency is the number of ticks per second
 			uint64 Div = FMath::Max(1llu, RHI.TimestampCalibration->GPUFrequency / (1000 * 1000));
 
 			// convert from GPU specific timestamp to micro sec (1 / 1 000 000 s) which seems a reasonable resolution
