Remove RHIBegin/EndUAVOverlap implementations when RHI validation is disabled in the build
- No platform RHI relies on these functions anymore. They exist only to inform the validation layer that sequential compute shaders dispatches to a given UAV without an interim transition are intentional.

#Synced-CL 48803475

diff --git a/Engine/Source/Runtime/RHI/Public/RHICommandList.h b/Engine/Source/Runtime/RHI/Public/RHICommandList.h
--- a/Engine/Source/Runtime/RHI/Public/RHICommandList.h
+++ b/Engine/Source/Runtime/RHI/Public/RHICommandList.h
@@ -1921,6 +1921,7 @@ FRHICOMMAND_MACRO(FRHICommandSetShaderRootConstants)
 	RHI_EXECUTE_API void Execute(FRHICommandListBase& CmdList);
 };
 
+#if ENABLE_RHI_VALIDATION
 FRHICOMMAND_MACRO(FRHICommandBeginUAVOverlap)
 {
 	RHI_EXECUTE_API void Execute(FRHICommandListBase& CmdList);
@@ -1944,6 +1945,7 @@ FRHICOMMAND_MACRO(FRHICommandEndSpecificUAVOverlap)
 	inline FRHICommandEndSpecificUAVOverlap(TArrayView<FRHIUnorderedAccessView* const> InUAVs) : UAVs(InUAVs) {}
 	RHI_EXECUTE_API void Execute(FRHICommandListBase & CmdList);
 };
+#endif // ENABLE_RHI_VALIDATION
 
 FRHICOMMAND_MACRO(FRHICommandDrawPrimitiveIndirect)
 {
@@ -2940,38 +2942,47 @@ public:
 
 	inline void BeginUAVOverlap()
 	{
+	#if ENABLE_RHI_VALIDATION
 		if (Bypass())
 		{
 			GetComputeContext().RHIBeginUAVOverlap();
 			return;
 		}
 		ALLOC_COMMAND(FRHICommandBeginUAVOverlap)();
+	#endif
 	}
 
 	inline void EndUAVOverlap()
 	{
+	#if ENABLE_RHI_VALIDATION
 		if (Bypass())
 		{
 			GetComputeContext().RHIEndUAVOverlap();
 			return;
 		}
 		ALLOC_COMMAND(FRHICommandEndUAVOverlap)();
+	#endif
 	}
 
 	inline void BeginUAVOverlap(FRHIUnorderedAccessView* UAV)
 	{
+	#if ENABLE_RHI_VALIDATION
 		FRHIUnorderedAccessView* UAVs[1] = { UAV };
 		BeginUAVOverlap(MakeArrayView(UAVs, 1));
+	#endif
 	}
 
 	inline void EndUAVOverlap(FRHIUnorderedAccessView* UAV)
 	{
+	#if ENABLE_RHI_VALIDATION
 		FRHIUnorderedAccessView* UAVs[1] = { UAV };
 		EndUAVOverlap(MakeArrayView(UAVs, 1));
+	#endif
 	}
 
 	inline void BeginUAVOverlap(TArrayView<FRHIUnorderedAccessView* const> UAVs)
 	{
+	#if ENABLE_RHI_VALIDATION
 		if (Bypass())
 		{
 			GetComputeContext().RHIBeginUAVOverlap(UAVs);
@@ -2982,10 +2993,12 @@ public:
 		FRHIUnorderedAccessView** InlineUAVs = (FRHIUnorderedAccessView**)Alloc(AllocSize, alignof(FRHIUnorderedAccessView*));
 		FMemory::Memcpy(InlineUAVs, UAVs.GetData(), AllocSize);
 		ALLOC_COMMAND(FRHICommandBeginSpecificUAVOverlap)(MakeArrayView(InlineUAVs, UAVs.Num()));
+	#endif
 	}
 
 	inline void EndUAVOverlap(TArrayView<FRHIUnorderedAccessView* const> UAVs)
 	{
+	#if ENABLE_RHI_VALIDATION
 		if (Bypass())
 		{
 			GetComputeContext().RHIEndUAVOverlap(UAVs);
@@ -2996,6 +3009,7 @@ public:
 		FRHIUnorderedAccessView** InlineUAVs = (FRHIUnorderedAccessView**)Alloc(AllocSize, alignof(FRHIUnorderedAccessView*));
 		FMemory::Memcpy(InlineUAVs, UAVs.GetData(), AllocSize);
 		ALLOC_COMMAND(FRHICommandEndSpecificUAVOverlap)(MakeArrayView(InlineUAVs, UAVs.Num()));
+	#endif
 	}
 
 #if WITH_RHI_BREADCRUMBS
diff --git a/Engine/Source/Runtime/RHI/Public/RHICommandListCommandExecutes.inl b/Engine/Source/Runtime/RHI/Public/RHICommandListCommandExecutes.inl
--- a/Engine/Source/Runtime/RHI/Public/RHICommandListCommandExecutes.inl
+++ b/Engine/Source/Runtime/RHI/Public/RHICommandListCommandExecutes.inl
@@ -274,6 +274,7 @@ void FRHICommandSetShaderRootConstants::Execute(FRHICommandListBase& CmdList)
 	INTERNAL_DECORATOR_COMPUTE(RHISetShaderRootConstants)(Constants);
 }
 
+#if ENABLE_RHI_VALIDATION
 void FRHICommandBeginUAVOverlap::Execute(FRHICommandListBase& CmdList)
 {
 	RHISTAT(BeginUAVOverlap);
@@ -297,6 +298,7 @@ void FRHICommandEndSpecificUAVOverlap::Execute(FRHICommandListBase& CmdList)
 	RHISTAT(EndSpecificUAVOverlap);
 	INTERNAL_DECORATOR_COMPUTE(RHIEndUAVOverlap)(UAVs);
 }
+#endif // ENABLE_RHI_VALIDATION
 
 void FRHICommandDrawPrimitiveIndirect::Execute(FRHICommandListBase& CmdList)
 {
diff --git a/Engine/Source/Runtime/RHI/Public/RHIContext.h b/Engine/Source/Runtime/RHI/Public/RHIContext.h
--- a/Engine/Source/Runtime/RHI/Public/RHIContext.h
+++ b/Engine/Source/Runtime/RHI/Public/RHIContext.h
@@ -331,11 +331,13 @@ public:
 		TConstArrayView<FRHIShaderBundleGraphicsDispatch> Dispatches,
 		ERHIShaderBundleMethod Method) {}
 
+#if ENABLE_RHI_VALIDATION
 	virtual void RHIBeginUAVOverlap() {}
 	virtual void RHIEndUAVOverlap() {}
 
 	virtual void RHIBeginUAVOverlap(TConstArrayView<FRHIUnorderedAccessView*> UAVs) {}
 	virtual void RHIEndUAVOverlap(TConstArrayView<FRHIUnorderedAccessView*> UAVs) {}
+#endif
 
 	virtual void RHISetShaderParameters(FRHIComputeShader* ComputeShader, TConstArrayView<uint8> InParametersData, TConstArrayView<FRHIShaderParameter> InParameters, TConstArrayView<FRHIShaderParameterResource> InResourceParameters, TConstArrayView<FRHIShaderParameterResource> InBindlessParameters) = 0;
 
diff --git a/Engine/Source/Runtime/Windows/D3D11RHI/Private/D3D11Commands.cpp b/Engine/Source/Runtime/Windows/D3D11RHI/Private/D3D11Commands.cpp
--- a/Engine/Source/Runtime/Windows/D3D11RHI/Private/D3D11Commands.cpp
+++ b/Engine/Source/Runtime/Windows/D3D11RHI/Private/D3D11Commands.cpp
@@ -1533,16 +1533,6 @@ void FD3D11DynamicRHI::RHIEndTransitions(TArrayView<const FRHITransition*> Trans
 	}
 }
 
-void FD3D11DynamicRHI::RHIBeginUAVOverlap()
-{
-	// No need to do anything here. Overlap is always on and the current group is broken up when we see a transition to UAVCompute or UAVGraphics.
-}
-
-void FD3D11DynamicRHI::RHIEndUAVOverlap()
-{
-	// Same as above.
-}
-
 //*********************** StagingBuffer Implementation ***********************//
 
 FStagingBufferRHIRef FD3D11DynamicRHI::RHICreateStagingBuffer()
diff --git a/Engine/Source/Runtime/Windows/D3D11RHI/Private/D3D11RHIPrivate.h b/Engine/Source/Runtime/Windows/D3D11RHI/Private/D3D11RHIPrivate.h
--- a/Engine/Source/Runtime/Windows/D3D11RHI/Private/D3D11RHIPrivate.h
+++ b/Engine/Source/Runtime/Windows/D3D11RHI/Private/D3D11RHIPrivate.h
@@ -316,8 +316,6 @@ public:
 	D3D11RHI_API virtual void RHISetComputeShader(FRHIComputeShader* ComputeShader) final override;
 	D3D11RHI_API virtual void RHIDispatchComputeShader(uint32 ThreadGroupCountX, uint32 ThreadGroupCountY, uint32 ThreadGroupCountZ) final override;
 	D3D11RHI_API virtual void RHIDispatchIndirectComputeShader(FRHIBuffer* ArgumentBuffer, uint32 ArgumentOffset) final override;
-	D3D11RHI_API virtual void RHIBeginUAVOverlap() final override;
-	D3D11RHI_API virtual void RHIEndUAVOverlap() final override;
 	D3D11RHI_API virtual void RHISetMultipleViewports(uint32 Count, const FViewportBounds* Data) final override;
 	D3D11RHI_API virtual void RHIClearUAVFloat(FRHIUnorderedAccessView* UnorderedAccessViewRHI, const FVector4f& Values) final override;
 	D3D11RHI_API virtual void RHIClearUAVUint(FRHIUnorderedAccessView* UnorderedAccessViewRHI, const FUintVector4& Values) final override;
