Add FColor to RHI breadcrumb descriptor
- Calculated via constexpr hash of the static name.
- Allows platform RHIs to emit custom colors for profilers to improve readability of traces

Also changed the colors of various SCOPED_NAMED_EVENTS in RHICommandList.cpp to make them easier to read in traces.

#Synced-CL 48803475

diff --git a/Engine/Source/Runtime/D3D12RHI/Private/D3D12CommandContext.cpp b/Engine/Source/Runtime/D3D12RHI/Private/D3D12CommandContext.cpp
--- a/Engine/Source/Runtime/D3D12RHI/Private/D3D12CommandContext.cpp
+++ b/Engine/Source/Runtime/D3D12RHI/Private/D3D12CommandContext.cpp
@@ -155,7 +155,8 @@ void FD3D12ContextCommon::BindDiagnosticBuffer(FD3D12RootSignature const* RootSi
 		#if USE_PIX
 			if (FD3D12DynamicRHI::GetD3DRHI()->IsPixEventEnabled())
 			{
-				PIXBeginEvent(GraphicsCommandList().Get(), PIX_COLOR(0xff, 0xff, 0xff), TEXT("%s"), GetNameStr());
+				uint32 Color = PIX_COLOR(Breadcrumb->Data.Color.R, Breadcrumb->Data.Color.G, Breadcrumb->Data.Color.B);
+				PIXBeginEvent(GraphicsCommandList().Get(), Color, TEXT("%s"), GetNameStr());
 			}
 		#endif
 		}
diff --git a/Engine/Source/Runtime/Landscape/Public/LandscapeUtils.h b/Engine/Source/Runtime/Landscape/Public/LandscapeUtils.h
--- a/Engine/Source/Runtime/Landscape/Public/LandscapeUtils.h
+++ b/Engine/Source/Runtime/Landscape/Public/LandscapeUtils.h
@@ -191,7 +191,8 @@ public:
 	class FScopedBreadcrumbEvent final
 	{
 	public:
-		FScopedBreadcrumbEvent(FRDGBuilderRecorder& InRecorder, TCHAR const* StaticName, FRDGEventName&& EventName)
+		template <size_t N>
+		FScopedBreadcrumbEvent(FRDGBuilderRecorder& InRecorder, TCHAR const(&StaticName)[N], FRDGEventName&& EventName)
 			: Recorder(&InRecorder)
 			, RDGEvent(MakeShared<TOptional<TRDGEventScopeGuard<FRDGScope_RHI>>>())
 		{
@@ -200,7 +201,7 @@ public:
 			Recorder->EnqueueRDGCommand(
 			[
 				RDGEvent = RDGEvent, 
-				StaticName = StaticName,
+				&StaticName = StaticName,
 				EventName = MoveTemp(EventName)
 			](FRDGBuilder& GraphBuilder) mutable
 			{
diff --git a/Engine/Source/Runtime/RHI/Private/RHICommandList.cpp b/Engine/Source/Runtime/RHI/Private/RHICommandList.cpp
--- a/Engine/Source/Runtime/RHI/Private/RHICommandList.cpp
+++ b/Engine/Source/Runtime/RHI/Private/RHICommandList.cpp
@@ -1000,7 +1000,7 @@ void FRHICommandListExecutor::FSubmitState::Dispatch(FRHICommandListBase* CmdLis
 		CurrentTranslateJob->EnqueueTranslateTask(MoveTemp(Prereqs),
 			[TranslateJob = CurrentTranslateJob, CmdList]()
 			{
-				SCOPED_NAMED_EVENT(RHI_Translate, FColor::White);
+				SCOPED_NAMED_EVENT(RHI_Translate, FColor::Orange);
 				TranslateJob->Translate(CmdList);
 			}
 		);
@@ -1097,7 +1097,7 @@ FGraphEventRef FRHICommandListExecutor::FTranslateState::Finalize()
 	FTaskPipe* Pipe = EnqueueTranslateTask({},
 		[this]()
 		{
-			SCOPED_NAMED_EVENT(RHI_Finalize, FColor::White);
+			SCOPED_NAMED_EVENT(RHI_Finalize, FColor::Magenta);
 		
 			TRHIPipelineArray<IRHIPlatformCommandList*> PlatformCommandLists {InPlace, nullptr};
 			FDynamicRHI::FRHIFinalizeContextArgs FinalizeArgs;
@@ -1205,7 +1205,7 @@ void FRHICommandListExecutor::FSubmitState::Submit(const FSubmitArgs& Args)
 	}
 
 	{
-		SCOPED_NAMED_EVENT(DeleteRHIResources, FColor::Magenta);
+		SCOPED_NAMED_EVENT(DeleteRHIResources, FColor::Red);
 		while (true)
 		{
 			// If the RHI thread will be flushed, keep processing the RHI resource delete queue until it is empty.
@@ -1254,7 +1254,7 @@ void FRHICommandListExecutor::FSubmitState::Submit(const FSubmitArgs& Args)
 RHI_API FGraphEventRef FRHICommandListExecutor::Submit(TConstArrayView<FRHICommandListBase*> AdditionalCommandLists, ERHISubmitFlags SubmitFlags)
 {
 	check(IsInRenderingThread());
-	SCOPED_NAMED_EVENT(RHICmdList_Submit, FColor::White);
+	SCOPED_NAMED_EVENT(RHICmdList_Submit, FColor::Green);
 
 	if (Bypass())
 	{
@@ -1389,7 +1389,7 @@ RHI_API FGraphEventRef FRHICommandListExecutor::Submit(TConstArrayView<FRHIComma
 		EnqueueDispatchTask(MoveTemp(Prereqs),
 			[State = SubmitState, CmdList]()
 			{
-				SCOPED_NAMED_EVENT(RHI_Dispatch, FColor::White);
+				SCOPED_NAMED_EVENT(RHI_Dispatch, FColor::Yellow);
 				State->Dispatch(CmdList);
 			}
 		);
@@ -1415,7 +1415,7 @@ RHI_API FGraphEventRef FRHICommandListExecutor::Submit(TConstArrayView<FRHIComma
 			#endif
 			]() mutable
 			{
-				SCOPED_NAMED_EVENT(RHI_FinalizeAndSubmit, FColor::White);
+				SCOPED_NAMED_EVENT(RHI_FinalizeAndSubmit, FColor::Cyan);
 
 				// Finalize the last translate job
 				State->FinalizeCurrent();
@@ -1435,7 +1435,7 @@ RHI_API FGraphEventRef FRHICommandListExecutor::Submit(TConstArrayView<FRHIComma
 				EnqueueSubmitTask(MoveTemp(State->TranslateEvents),
 					[State, Args]() mutable
 					{
-						SCOPED_NAMED_EVENT(RHI_SubmitToGPU, FColor::White);
+						SCOPED_NAMED_EVENT(RHI_SubmitToGPU, FColor::Purple);
 						State->Submit(Args);
 					}
 				);
diff --git a/Engine/Source/Runtime/RHI/Public/RHIBreadcrumbs.h b/Engine/Source/Runtime/RHI/Public/RHIBreadcrumbs.h
--- a/Engine/Source/Runtime/RHI/Public/RHIBreadcrumbs.h
+++ b/Engine/Source/Runtime/RHI/Public/RHIBreadcrumbs.h
@@ -4,6 +4,7 @@
 
 #include "Misc/Build.h"
 #include "Misc/MemStack.h"
+#include "Hash/Fnv.h"
 #include "HAL/Platform.h"
 #include "Templates/SharedPointer.h"
 #include "ProfilingDebugging/CpuProfilerTrace.h"
@@ -126,10 +127,15 @@ namespace UE::RHI::GPUProfiler
 	public:
 		TCHAR const* const StaticName;
 
-		FRHIBreadcrumbData(TCHAR const* StaticName, ANSICHAR const* File, uint32 Line, FRHIBreadcrumbData_Stats&& Stats)
+		// Random assigned color, based on static name
+		FColor const Color;
+
+		template <size_t N>
+		FRHIBreadcrumbData(TCHAR const(&StaticName)[N], ANSICHAR const* File, uint32 Line, FRHIBreadcrumbData_Stats&& Stats)
 			: FRHIBreadcrumbData_Location(File, Line)
 			, FRHIBreadcrumbData_Stats(MoveTemp(Stats))
 			, StaticName(StaticName)
+			, Color(FLinearColor::IntToDistinctColor(UE::HashStringFNV1a<uint32>(MakeStringView(StaticName, N - 1))).ToFColor(true))
 		{}
 	};
 
