Respect the EnablesHLSL2021ByDefault setting for editor preview modes.
Fixes shader compiler errors with MegaLightsLightComplexity.usf
#rnx
#Synced-CL 48841205

diff --git a/Engine/Source/Runtime/RHI/Private/DataDrivenShaderPlatformInfo.cpp b/Engine/Source/Runtime/RHI/Private/DataDrivenShaderPlatformInfo.cpp
--- a/Engine/Source/Runtime/RHI/Private/DataDrivenShaderPlatformInfo.cpp
+++ b/Engine/Source/Runtime/RHI/Private/DataDrivenShaderPlatformInfo.cpp
@@ -577,6 +577,7 @@ void FGenericDataDrivenShaderPlatformInfo::UpdatePreviewPlatforms()
 				PREVIEW_USE_RUNTIME_VALUE(bSupportsRealTypes);
 				PREVIEW_USE_RUNTIME_VALUE(bSupportsUniformBufferObjects);
 				PREVIEW_USE_RUNTIME_VALUE(bRequiresBindfulUtilityShaders);
+				PREVIEW_USE_RUNTIME_VALUE(EnablesHLSL2021ByDefault);
 
 				PREVIEW_USE_RUNTIME_VALUE(bSupportsInlineRayTracing);
 				PREVIEW_USE_RUNTIME_VALUE(bInlineRayTracingRequiresBindless);
