Fix bug in "stat drawcount" where values would always be zero
- Accumulation of draw counts stats were done in a loop across queues, but were not summed. Only the last value in the loop was kept. This meant that if e.g. the compute queue was the last to be processed, the draw counts would be zero.
- Adjusted the draw counts to sum across the queues. Also moved the logic for updating the display copy into FRHIDrawStatsCategory so its not in the GPU profiler.

Also fixed CSV profiler draw count stats not working in test builds
- The call to the Stats() function in the GPU profiler was wrapped in #if WITH_PROFILEGPU, which meant it wasn't called when profilegpu support is disabled.
- We should always call this function, so that CSV stats continue to work independently.

#Synced-CL 48809907

diff --git a/Engine/Source/Runtime/Engine/Private/UnrealEngine.cpp b/Engine/Source/Runtime/Engine/Private/UnrealEngine.cpp
--- a/Engine/Source/Runtime/Engine/Private/UnrealEngine.cpp
+++ b/Engine/Source/Runtime/Engine/Private/UnrealEngine.cpp
@@ -19200,26 +19200,31 @@ int32 UEngine::RenderStatDrawCount(UWorld* World, FViewport* Viewport, FCanvas*
 	int32 TotalCount[MAX_NUM_GPUS] = { 0 };
 	// Display all the categories of draw counts.
 	FRHIDrawStatsCategory::FManager const& Manager = FRHIDrawStatsCategory::GetManager();
-	for (int32 Index = 0; Index < Manager.NumCategory; ++Index)
-	{
-		for (uint32 GPUIndex : FRHIGPUMask::All())
-		{
-			TotalCount[GPUIndex] += Manager.DisplayCounts[Index][GPUIndex];
-			FRHIDrawStatsCategory* CategoryName = Manager.Array[Index];
-
-			Canvas->DrawShadowedString(
-				X - 100, Y,
-				*FString::Printf(
-					TEXT("%s%s: %i"),
-					*CategoryName->Name.ToString(),
-					GNumExplicitGPUsForRendering > 1 ? *FString::Printf(TEXT(" (GPU%u)"), GPUIndex) : TEXT(""),
-					Manager.DisplayCounts[Index][GPUIndex]
-				),
-				GetSmallFont(),
-				FColor::Green
-			);
 
-			Y += 12;
+	{
+		FScopeLock Lock(&Manager.DisplayCountsCS);
+
+		for (int32 Index = 0; Index < Manager.NumCategory; ++Index)
+		{
+			for (uint32 GPUIndex : FRHIGPUMask::All())
+			{
+				TotalCount[GPUIndex] += Manager.DisplayCounts[Index][GPUIndex];
+				FRHIDrawStatsCategory* CategoryName = Manager.Array[Index];
+
+				Canvas->DrawShadowedString(
+					X - 100, Y,
+					*FString::Printf(
+						TEXT("%s%s: %i"),
+						*CategoryName->Name.ToString(),
+						GNumExplicitGPUsForRendering > 1 ? *FString::Printf(TEXT(" (GPU%u)"), GPUIndex) : TEXT(""),
+						Manager.DisplayCounts[Index][GPUIndex]
+					),
+					GetSmallFont(),
+					FColor::Green
+				);
+
+				Y += 12;
+			}
 		}
 	}
 
diff --git a/Engine/Source/Runtime/RHI/Private/GPUProfiler.cpp b/Engine/Source/Runtime/RHI/Private/GPUProfiler.cpp
--- a/Engine/Source/Runtime/RHI/Private/GPUProfiler.cpp
+++ b/Engine/Source/Runtime/RHI/Private/GPUProfiler.cpp
@@ -1824,14 +1824,12 @@ namespace UE::RHI::GPUProfiler
 					break;
 			#endif // WITH_RHI_BREADCRUMBS
 
-			#if WITH_PROFILEGPU
 				case FEvent::EType::Stats:
 					{
 						check(QueueState.bBusy);
 						QueueState.Stats(Event->Value.Get<FEvent::FStats>());
 					}
 					break;
-			#endif // WITH_PROFILEGPU
 
 				case FEvent::EType::SignalFence:
 					{
@@ -1942,20 +1940,7 @@ namespace UE::RHI::GPUProfiler
 		#endif
 
 		#if HAS_GPU_STATS
-			// Only copy the display counters every half second keep things more stable.
-			constexpr float TimeoutSeconds = 0.5;
-
-			static double LastTime = 0.0;
-			double CurrentTime = FPlatformTime::Seconds();
-
-			bool bCopyDisplayFrames = false;
-			if (CurrentTime - LastTime > TimeoutSeconds)
-			{
-				LastTime = CurrentTime;
-				bCopyDisplayFrames = true;
-			}
-
-			FRHIDrawStatsCategory::FManager& Manager = FRHIDrawStatsCategory::GetManager();
+			FRHIDrawStatsCategory::FManager::FDrawCounts DrawCounts;
 		#endif
 
 			// Summed stats across all GPUs
@@ -2008,12 +1993,12 @@ namespace UE::RHI::GPUProfiler
 					Total                           += Category;
 
 				#if HAS_GPU_STATS
-					if (bCopyDisplayFrames && CategoryIndex < Manager.NumCategory)
+					if (CategoryIndex < DrawCounts.Num())
 					{
-						Manager.DisplayCounts[CategoryIndex][Queue.GPU] = Category.NumDraws;
+						DrawCounts[CategoryIndex][Queue.GPU] += Category.NumDraws;
 					}
-				#endif // HAS_GPU_STATS
-				}				
+				#endif
+				}
 			}
 
 			for (uint32 GPUIndex = 0; GPUIndex < GNumExplicitGPUsForRendering; ++GPUIndex)
@@ -2030,6 +2015,9 @@ namespace UE::RHI::GPUProfiler
 
 	#if HAS_GPU_STATS
 
+			FRHIDrawStatsCategory::FManager& Manager = FRHIDrawStatsCategory::GetManager();
+			Manager.AccumulateFrameStats(DrawCounts);
+
 		#if STATS
 			Stats->AddMessage(GET_STATFNAME(STAT_RHIDraws     ), EStatOperation::Set, uint64(Total.NumDraws     ));
 			Stats->AddMessage(GET_STATFNAME(STAT_RHIDispatches), EStatOperation::Set, uint64(Total.NumDispatches));
diff --git a/Engine/Source/Runtime/RHI/Private/RHI.cpp b/Engine/Source/Runtime/RHI/Private/RHI.cpp
--- a/Engine/Source/Runtime/RHI/Private/RHI.cpp
+++ b/Engine/Source/Runtime/RHI/Private/RHI.cpp
@@ -155,6 +155,39 @@ const FClearValueBinding FClearValueBinding::DefaultNormal8Bit(FLinearColor(128.
 		return Manager;
 	}
 
+	void FRHIDrawStatsCategory::FManager::AccumulateFrameStats(FDrawCounts const& NewCounts)
+	{
+		// Update the current counts using the maximum over the update period
+		for (int32 GPUIndex = 0; GPUIndex < CurrentCounts.Num(); ++GPUIndex)
+		{
+			FDrawCountsPerGPU      & Dst = CurrentCounts[GPUIndex];
+			FDrawCountsPerGPU const& Src = NewCounts[GPUIndex];
+
+			for (int32 CategoryIndex = 0; CategoryIndex < Src.Num(); ++CategoryIndex)
+			{
+				Dst[CategoryIndex] = FMath::Max(Dst[CategoryIndex], Src[CategoryIndex]);
+			}
+		}
+
+		// Only copy the display counters every half second keep things more stable.
+		constexpr float TimeoutSeconds = 0.5;
+		static double LastTime = 0.0;
+		double CurrentTime = FPlatformTime::Seconds();
+
+		if (CurrentTime - LastTime > TimeoutSeconds)
+		{
+			LastTime = CurrentTime;
+
+			{
+				FScopeLock Lock(&DisplayCountsCS);
+				DisplayCounts = CurrentCounts;
+			}
+			
+			// Reset the counts to zero
+			CurrentCounts = {};
+		}
+	}
+
 #endif
 
 TRefCountPtr<FRHITexture> FRHITextureReference::DefaultTexture;
diff --git a/Engine/Source/Runtime/RHI/Public/RHIStats.h b/Engine/Source/Runtime/RHI/Public/RHIStats.h
--- a/Engine/Source/Runtime/RHI/Public/RHIStats.h
+++ b/Engine/Source/Runtime/RHI/Public/RHIStats.h
@@ -86,11 +86,33 @@ struct FRHIDrawStatsCategory
 		TStaticArray<FRHIDrawStatsCategory*, MAX_DRAWCALL_CATEGORY> Array;
 
 		// A backup of the counts that can be used to display on screen to avoid flickering.
-		TStaticArray<TStaticArray<int32, MAX_NUM_GPUS>, MAX_DRAWCALL_CATEGORY> DisplayCounts;
+
+		struct FDrawCountsPerGPU : public TStaticArray<int32, MAX_NUM_GPUS>
+		{
+			FDrawCountsPerGPU()
+				: TStaticArray<int32, MAX_NUM_GPUS>(InPlace, 0)
+			{}
+		};
+
+		struct FDrawCounts : public TStaticArray<FDrawCountsPerGPU, MAX_DRAWCALL_CATEGORY>
+		{
+			FDrawCounts()
+				: TStaticArray<FDrawCountsPerGPU, MAX_DRAWCALL_CATEGORY>(InPlace)
+			{}
+		};
+
+		// Used to track the max number of draws per frame.
+		FDrawCounts CurrentCounts;
+
+		// Copy of CurrentCounts updated periodically to use for displaying the "stat drawcount" text on screen.
+		FDrawCounts DisplayCounts;
+		mutable FCriticalSection DisplayCountsCS;
 
 		int32 NumCategory;
 
 		RHI_API FManager();
+
+		void AccumulateFrameStats(FDrawCounts const& Counts);
 	};
 
 	RHI_API static FManager& GetManager();
