Fix issues with Pipeline File Cache Variable Rate Shading handling added in CL 44399700:
- Increment FPipelineCacheFileFormatCurrentVersion to EPipelineCacheFileFormatVersions::VariableRateShading
- Properly pass the new VRS fields to FGraphicsPipelineStateInitializer in FShaderPipelineCacheTask::Precompile

#Synced-CL 48969735

diff --git a/Engine/Source/Runtime/RHI/Private/PipelineFileCache.cpp b/Engine/Source/Runtime/RHI/Private/PipelineFileCache.cpp
--- a/Engine/Source/Runtime/RHI/Private/PipelineFileCache.cpp
+++ b/Engine/Source/Runtime/RHI/Private/PipelineFileCache.cpp
@@ -91,7 +91,7 @@ enum class EPipelineCacheFileFormatVersions : uint32
 const uint64 FPipelineCacheFileFormatMagic = 0x5049504543414348; // PIPECACH
 const uint64 FPipelineCacheTOCFileFormatMagic = 0x544F435354415232; // TOCSTAR2
 const uint64 FPipelineCacheEOFFileFormatMagic = 0x454F462D4D41524B; // EOF-MARK
-const RHI_API uint32 FPipelineCacheFileFormatCurrentVersion = (uint32)EPipelineCacheFileFormatVersions::AddRTPSOShaderBindingLayout;
+const RHI_API uint32 FPipelineCacheFileFormatCurrentVersion = (uint32)EPipelineCacheFileFormatVersions::VariableRateShading;
 const int32  FPipelineCacheGraphicsDescPartsNum = 67; // parser will expect this number of parts in a description string
 
 /**
diff --git a/Engine/Source/Runtime/RenderCore/Private/ShaderPipelineCache.cpp b/Engine/Source/Runtime/RenderCore/Private/ShaderPipelineCache.cpp
--- a/Engine/Source/Runtime/RenderCore/Private/ShaderPipelineCache.cpp
+++ b/Engine/Source/Runtime/RenderCore/Private/ShaderPipelineCache.cpp
@@ -1083,6 +1083,9 @@ bool FShaderPipelineCacheTask::Precompile(FRHICommandListImmediate& RHICmdList,
 				GraphicsInitializer.MultiViewCount = PSO.GraphicsDesc.MultiViewCount;
 				GraphicsInitializer.bHasFragmentDensityAttachment = PSO.GraphicsDesc.bHasFragmentDensityAttachment;
 
+				GraphicsInitializer.bAllowVariableRateShading = PSO.GraphicsDesc.bAllowVariableRateShading;
+				GraphicsInitializer.ShadingRate = PSO.GraphicsDesc.ShadingRate;
+
 				GraphicsInitializer.bDepthBounds = PSO.GraphicsDesc.bDepthBounds;
 
 				GraphicsInitializer.DepthStencilTargetFormat = PSO.GraphicsDesc.DepthStencilFormat;
