Remove incorrect check in AsyncReallocateTexture2D

#Synced-CL 49559950

diff --git a/Engine/Source/Runtime/RHI/Private/RHICommandList.cpp b/Engine/Source/Runtime/RHI/Private/RHICommandList.cpp
--- a/Engine/Source/Runtime/RHI/Private/RHICommandList.cpp
+++ b/Engine/Source/Runtime/RHI/Private/RHICommandList.cpp
@@ -1889,8 +1889,6 @@ void FRHICommandListImmediate::EndFrame()
 
 FTextureRHIRef FRHICommandListImmediate::AsyncReallocateTexture2D(FRHITexture* OldTexture, int32 NewMipCount, int32 NewSizeX, int32 NewSizeY, FThreadSafeCounter* RequestStatus)
 {
-	check(!GRHISupportsAsyncTextureCreation);
-
 	LLM_SCOPE(ELLMTag::Textures);
 	RHI_BREADCRUMB_EVENT_F(*this, "AsyncReallocateTexture2D", "AsyncReallocateTexture2D - %s", RHI_BREADCRUMB_FIELD("Texture Name", OldTexture->GetName()));
 
