Remove unintended double -> float -> double roundtrip on texture last render times and pessimization in SetLastRenderTime.

#rnx
#tests Not necessarity, this is just a small code cleanup.
#Synced-CL 49648402

diff --git a/Engine/Source/Runtime/RHI/Internal/RHIUniformBufferUtilities.h b/Engine/Source/Runtime/RHI/Internal/RHIUniformBufferUtilities.h
--- a/Engine/Source/Runtime/RHI/Internal/RHIUniformBufferUtilities.h
+++ b/Engine/Source/Runtime/RHI/Internal/RHIUniformBufferUtilities.h
@@ -56,7 +56,7 @@ inline void EnumerateUniformBufferResources(const FRHIUniformBuffer* RESTRICT Bu
 template<typename TBinder>
 void SetAllUniformBufferResourcesForBindless(TBinder&& Binder, const FRHIUniformBuffer* Buffer)
 {
-	float CurrentTimeForTextureTimes = FApp::GetCurrentTime();
+	double CurrentTimeForTextureTimes = FApp::GetCurrentTime();
 
 	for (const TRefCountPtr<FRHIResource>& Resource : Buffer->GetResourceTable())
 	{
@@ -94,7 +94,7 @@ void SetUniformBufferResourcesFromTables(TBinder && Binder, FRHIShader const& Sh
                                          #endif
                                          )
 {
-	float CurrentTimeForTextureTimes = FApp::GetCurrentTime();
+	double CurrentTimeForTextureTimes = FApp::GetCurrentTime();
 	FShaderResourceTable const& SRT = Shader.GetShaderResourceTable();
 
 	// Mask the dirty bits by those buffers from which the shader has bound resources.
diff --git a/Engine/Source/Runtime/RHI/Public/RHIResources.h b/Engine/Source/Runtime/RHI/Public/RHIResources.h
--- a/Engine/Source/Runtime/RHI/Public/RHIResources.h
+++ b/Engine/Source/Runtime/RHI/Public/RHIResources.h
@@ -1649,11 +1649,7 @@ public:
 
 	void SetLastRenderTime(double InLastRenderTime) 
 	{ 
-		// avoid dirty caches from redundant writes
-		if (LastRenderTime != InLastRenderTime)
-		{
-			LastRenderTime = InLastRenderTime;
-		}
+		LastRenderTime = InLastRenderTime;
 	}
 
 private:
