Use approximate time for texture LastRenderTime feedback
- Rendering can be triggered from code which does not have a valid time context, so the accurate FAppTime::GetCurrentTime() function cannot be used.
- Approximate time should be sufficient for the texture streamer.

Removed FLastRenderTimeContainer
- It was just a wrapper around a double.
- Replaced with std::atomic<double> on FRHITexture, which is loaded/stored with relaxed ordering. Multiple threads write this value, but since the time is approximate, we don't need to guarantee ordering.

#Synced-CL 49648710

diff --git a/Engine/Source/Runtime/RHI/Internal/RHIUniformBufferUtilities.h b/Engine/Source/Runtime/RHI/Internal/RHIUniformBufferUtilities.h
--- a/Engine/Source/Runtime/RHI/Internal/RHIUniformBufferUtilities.h
+++ b/Engine/Source/Runtime/RHI/Internal/RHIUniformBufferUtilities.h
@@ -7,6 +7,7 @@
 #include "Misc/App.h"
 #include "RHIDefinitions.h"
 #include "RHIResources.h"
+#include "RHIResourceCollection.h"
 
 namespace UE::RHI::Private
 {
@@ -56,7 +57,7 @@ inline void EnumerateUniformBufferResources(const FRHIUniformBuffer* RESTRICT Bu
 template<typename TBinder>
 void SetAllUniformBufferResourcesForBindless(TBinder&& Binder, const FRHIUniformBuffer* Buffer)
 {
-	double CurrentTimeForTextureTimes = FApp::GetCurrentTime();
+	double CurrentTimeForTextureTimes = FAppTime::GetApproximateCurrentTime();
 
 	for (const TRefCountPtr<FRHIResource>& Resource : Buffer->GetResourceTable())
 	{
@@ -94,7 +95,7 @@ void SetUniformBufferResourcesFromTables(TBinder && Binder, FRHIShader const& Sh
                                          #endif
                                          )
 {
-	double CurrentTimeForTextureTimes = FApp::GetCurrentTime();
+	double CurrentTimeForTextureTimes = FAppTime::GetApproximateCurrentTime();
 	FShaderResourceTable const& SRT = Shader.GetShaderResourceTable();
 
 	// Mask the dirty bits by those buffers from which the shader has bound resources.
diff --git a/Engine/Source/Runtime/RHI/Public/DynamicRHI.h b/Engine/Source/Runtime/RHI/Public/DynamicRHI.h
--- a/Engine/Source/Runtime/RHI/Public/DynamicRHI.h
+++ b/Engine/Source/Runtime/RHI/Public/DynamicRHI.h
@@ -20,7 +20,6 @@ DynamicRHI.h: Dynamically bound Render Hardware Interface definitions.
 
 class FBlendStateInitializerRHI;
 class FGraphicsPipelineStateInitializer;
-class FLastRenderTimeContainer;
 class FReadSurfaceDataFlags;
 class FRHICommandList;
 class FRHICommandListBase;
diff --git a/Engine/Source/Runtime/RHI/Public/RHICommandList.h b/Engine/Source/Runtime/RHI/Public/RHICommandList.h
--- a/Engine/Source/Runtime/RHI/Public/RHICommandList.h
+++ b/Engine/Source/Runtime/RHI/Public/RHICommandList.h
@@ -60,7 +60,6 @@ extern RHI_API ERHIAccess RHIGetDefaultResourceState(EBufferUsageFlags InUsage,
 class FApp;
 class FBlendStateInitializerRHI;
 class FGraphicsPipelineStateInitializer;
-class FLastRenderTimeContainer;
 class FRHICommandListBase;
 class FRHIComputeShader;
 class IRHICommandContext;
diff --git a/Engine/Source/Runtime/RHI/Public/RHIResources.h b/Engine/Source/Runtime/RHI/Public/RHIResources.h
--- a/Engine/Source/Runtime/RHI/Public/RHIResources.h
+++ b/Engine/Source/Runtime/RHI/Public/RHIResources.h
@@ -1640,24 +1640,6 @@ private:
 // Textures
 //
 
-class FLastRenderTimeContainer
-{
-public:
-	FLastRenderTimeContainer() : LastRenderTime(-FLT_MAX) {}
-
-	double GetLastRenderTime() const { return LastRenderTime; }
-
-	void SetLastRenderTime(double InLastRenderTime) 
-	{ 
-		LastRenderTime = InLastRenderTime;
-	}
-
-private:
-	/** The last time the resource was rendered. */
-	double LastRenderTime;
-};
-
-
 /** Descriptor used to create a texture resource */
 struct FRHITextureDesc
 {
@@ -2261,14 +2243,14 @@ public:
 	/// 
 
 	/** sets the last time this texture was cached in a resource table. */
-	inline void SetLastRenderTime(float InLastRenderTime)
+	void SetLastRenderTime(double InLastRenderTime)
 	{
-		LastRenderTime.SetLastRenderTime(InLastRenderTime);
+		LastRenderTime.store(InLastRenderTime, std::memory_order_relaxed);
 	}
 
 	double GetLastRenderTime() const
 	{
-		return LastRenderTime.GetLastRenderTime();
+		return LastRenderTime.load(std::memory_order_relaxed);
 	}
 
 	RHI_API void SetName(FName InName);
@@ -2325,12 +2307,15 @@ public:
 	}
 #endif
 
+	static constexpr double EmptyLastRenderTime = -FLT_MAX;
+
 private:
 	friend class FRHITextureReference;
 
 	FRHITextureDesc TextureDesc;
 
-	FLastRenderTimeContainer LastRenderTime;
+	/** The last time the resource was rendered. */
+	std::atomic<double> LastRenderTime = EmptyLastRenderTime;
 };
 
 //
diff --git a/Engine/Source/Runtime/RenderCore/Private/RenderResource.cpp b/Engine/Source/Runtime/RenderCore/Private/RenderResource.cpp
--- a/Engine/Source/Runtime/RenderCore/Private/RenderResource.cpp
+++ b/Engine/Source/Runtime/RenderCore/Private/RenderResource.cpp
@@ -509,7 +509,7 @@ double FTextureReference::GetLastRenderTime() const
 		return TextureReferenceRHI->GetLastRenderTime();
 	}
 
-	return FLastRenderTimeContainer().GetLastRenderTime();
+	return FRHITexture::EmptyLastRenderTime;
 }
 
 void FTextureReference::InvalidateLastRenderTime()
