Fix pipelining of profilegpu trigger
- The frame index counting could be off by one, meaning the frame that gets profiled might not have had the full RDG event scopes and pass names (UE::RHI::GPUProfiler::IsProfiling() returned false).
- Instead, the latching of the profilegpu command is done in RHIEndFrame_RenderThread, and pipelined through the profiler event stream via the bProfileNextFrame bool on the FFrameBoundary event, to ensure the correct frame is profiled.

Refactored the constructor of the FFrameBoundary GPU profiler event type to reduce the amount of code needed in each RHI's implementation of RHIEndFrame.

#Synced-CL 49648710

diff --git a/Engine/Source/Runtime/Apple/MetalRHI/Private/MetalRHI.cpp b/Engine/Source/Runtime/Apple/MetalRHI/Private/MetalRHI.cpp
--- a/Engine/Source/Runtime/Apple/MetalRHI/Private/MetalRHI.cpp
+++ b/Engine/Source/Runtime/Apple/MetalRHI/Private/MetalRHI.cpp
@@ -1295,15 +1295,7 @@ void FMetalDynamicRHI::RHIEndFrame(const FRHIEndFrameArgs& Args)
 		
 		ERHIPipeline Pipeline = ERHIPipeline::Graphics;
 		
-		Payload.EndFrameEvent = UE::RHI::GPUProfiler::FEvent::FFrameBoundary(
-			0, Args.FrameNumber
-		#if WITH_RHI_BREADCRUMBS
-			, (Pipeline != ERHIPipeline::None) ? Args.GPUBreadcrumbs[Pipeline] : nullptr
-		#endif
-		#if STATS
-			, Args.StatsFrame
-		#endif
-		);
+		Payload.EndFrameEvent = UE::RHI::GPUProfiler::FEvent::FFrameBoundary(Pipeline, Args, 0);
 	});
 	
 	// Pump the interrupt queue to gather completed events
diff --git a/Engine/Source/Runtime/D3D12RHI/Private/D3D12RHI.cpp b/Engine/Source/Runtime/D3D12RHI/Private/D3D12RHI.cpp
--- a/Engine/Source/Runtime/D3D12RHI/Private/D3D12RHI.cpp
+++ b/Engine/Source/Runtime/D3D12RHI/Private/D3D12RHI.cpp
@@ -606,14 +606,7 @@ void FD3D12DynamicRHI::RHIEndFrame(const FRHIEndFrameArgs& Args)
 		}
 
 		// CPU timestamp for the frame boundary event is filled in by the submission thread
-		Payload.EndFrameEvent = UE::RHI::GPUProfiler::FEvent::FFrameBoundary(0, Args.FrameNumber
-		#if WITH_RHI_BREADCRUMBS
-			, (Pipeline != ERHIPipeline::None) ? Args.GPUBreadcrumbs[Pipeline] : nullptr
-		#endif
-		#if STATS
-			, Args.StatsFrame
-		#endif
-		);
+		Payload.EndFrameEvent = UE::RHI::GPUProfiler::FEvent::FFrameBoundary(Pipeline, Args, 0);
 	});
 
 	// Pump the interrupt queue to gather completed events
diff --git a/Engine/Source/Runtime/NullDrv/Private/NullRHI.cpp b/Engine/Source/Runtime/NullDrv/Private/NullRHI.cpp
--- a/Engine/Source/Runtime/NullDrv/Private/NullRHI.cpp
+++ b/Engine/Source/Runtime/NullDrv/Private/NullRHI.cpp
@@ -106,14 +106,7 @@ void FNullDynamicRHI::RHIEndFrame(const FRHIEndFrameArgs& Args)
 	// The RHI must publish the frame boundary events for engine stats to work correctly.
 	UE::RHI::GPUProfiler::FEventStream EventStream(NullRHIGPUProfilerQueue);
 
-	EventStream.Emplace<UE::RHI::GPUProfiler::FEvent::FFrameBoundary>(FPlatformTime::Cycles64(), Args.FrameNumber
-	#if WITH_RHI_BREADCRUMBS
-		, Args.GPUBreadcrumbs[ERHIPipeline::Graphics]
-	#endif
-	#if STATS
-		, Args.StatsFrame
-	#endif
-	);
+	EventStream.Emplace<UE::RHI::GPUProfiler::FEvent::FFrameBoundary>(ERHIPipeline::Graphics, Args, FPlatformTime::Cycles64());
 
 	UE::RHI::GPUProfiler::ProcessEvents(MakeArrayView(&EventStream, 1));
 }
diff --git a/Engine/Source/Runtime/OpenGLDrv/Private/OpenGLDevice.cpp b/Engine/Source/Runtime/OpenGLDrv/Private/OpenGLDevice.cpp
--- a/Engine/Source/Runtime/OpenGLDrv/Private/OpenGLDevice.cpp
+++ b/Engine/Source/Runtime/OpenGLDrv/Private/OpenGLDevice.cpp
@@ -1546,7 +1546,7 @@ void FOpenGLProfiler::Shutdown()
 	FOpenGLRenderQuery::PollQueryResults();
 }
 
-void FOpenGLProfiler::EndFrame(FDynamicRHI::FRHIEndFrameArgs const& Args)
+void FOpenGLProfiler::EndFrame(FRHIEndFrameArgs const& Args)
 {
 	if (bEnabled)
 	{
@@ -1566,13 +1566,10 @@ void FOpenGLProfiler::EndFrame(FDynamicRHI::FRHIEndFrameArgs const& Args)
 	TUniquePtr<FFrame> NewFrame = MakeUnique<FFrame>(*this);
 
 	// Insert frame boundary
-	Current->FrameBoundaryEvent = &EmplaceEvent<UE::RHI::GPUProfiler::FEvent::FFrameBoundary>(NewFrame->Calibration.CPUTimestamp, Args.FrameNumber
-#if WITH_RHI_BREADCRUMBS
-		, Args.GPUBreadcrumbs[ERHIPipeline::Graphics]
-#endif
-#if STATS
-		, Args.StatsFrame
-#endif
+	Current->FrameBoundaryEvent = &EmplaceEvent<UE::RHI::GPUProfiler::FEvent::FFrameBoundary>(
+		ERHIPipeline::Graphics,
+		Args,
+		NewFrame->Calibration.CPUTimestamp
 	);
 
 	Pending.Enqueue(MoveTemp(Current));
diff --git a/Engine/Source/Runtime/OpenGLDrv/Private/OpenGLResources.h b/Engine/Source/Runtime/OpenGLDrv/Private/OpenGLResources.h
--- a/Engine/Source/Runtime/OpenGLDrv/Private/OpenGLResources.h
+++ b/Engine/Source/Runtime/OpenGLDrv/Private/OpenGLResources.h
@@ -1158,7 +1158,7 @@ struct FOpenGLProfiler
 	void Setup();
 	void Shutdown();
 
-	void EndFrame(FDynamicRHI::FRHIEndFrameArgs const& Args);
+	void EndFrame(FRHIEndFrameArgs const& Args);
 
 	template <typename TEventType, typename... TArgs>
 	TEventType& EmplaceEvent(TArgs&&... Args)
diff --git a/Engine/Source/Runtime/RHI/Private/DynamicRHI.cpp b/Engine/Source/Runtime/RHI/Private/DynamicRHI.cpp
--- a/Engine/Source/Runtime/RHI/Private/DynamicRHI.cpp
+++ b/Engine/Source/Runtime/RHI/Private/DynamicRHI.cpp
@@ -461,7 +461,14 @@ FDynamicRHI::~FDynamicRHI() = default;
 
 void FDynamicRHI::RHIEndFrame_RenderThread(FRHICommandListImmediate& RHICmdList)
 {
-	RHICmdList.ImmediateFlush(EImmediateFlushType::DispatchToRHIThread, ERHISubmitFlags::EndFrame);
+	ERHISubmitFlags Flags = ERHISubmitFlags::EndFrame;
+
+	if (UE::RHI::GPUProfiler::ShouldProfileNextFrame())
+	{
+		EnumAddFlags(Flags, ERHISubmitFlags::ProfileNextFrame);
+	}
+
+	RHICmdList.ImmediateFlush(EImmediateFlushType::DispatchToRHIThread, Flags);
 }
 
 // Default fallback; will not work for non-8-bit surfaces and it's extremely slow.
diff --git a/Engine/Source/Runtime/RHI/Private/GPUProfiler.cpp b/Engine/Source/Runtime/RHI/Private/GPUProfiler.cpp
--- a/Engine/Source/Runtime/RHI/Private/GPUProfiler.cpp
+++ b/Engine/Source/Runtime/RHI/Private/GPUProfiler.cpp
@@ -9,6 +9,7 @@
 #include "Misc/WildcardString.h"
 #include "Misc/CommandLine.h"
 #include "RHI.h"
+#include "DynamicRHI.h"
 #include "GpuProfilerTrace.h"
 #include "Containers/AnsiString.h"
 #include "Stats/StatsData.h"
@@ -123,6 +124,19 @@ CSV_DEFINE_CATEGORY(RHI, true);
 
 namespace UE::RHI::GPUProfiler
 {
+	FEvent::FFrameBoundary::FFrameBoundary(ERHIPipeline Pipeline, FRHIEndFrameArgs const& Args, uint64 CPUTimestamp)
+		: CPUTimestamp     (CPUTimestamp)
+		, FrameNumber	   (Args.FrameNumber)
+		, bProfileNextFrame(Args.bProfileNextFrame)
+	#if STATS
+		, bStatsFrameSet   (Args.StatsFrame.IsSet())
+		, StatsFrame       (Args.StatsFrame.IsSet() ? *Args.StatsFrame : 0)
+	#endif
+	#if WITH_RHI_BREADCRUMBS
+		, Breadcrumb	   ((Pipeline != ERHIPipeline::None) ? Args.GPUBreadcrumbs[Pipeline] : nullptr)
+	#endif
+	{}
+
 	TLockFreePointerListUnordered<void, PLATFORM_CACHE_LINE_SIZE> FEventStream::FChunk::MemoryPool;
 
 	FEvent const* FEventStream::Peek() const
@@ -1655,7 +1669,7 @@ namespace UE::RHI::GPUProfiler
 				QueueFrameState.TotalBusyCycles = TotalGPUTime;
 			}
 
-			void FrameBoundary(FEvent::FFrameBoundary const& Event, FFrameState& FrameState, uint32 InProfileFrameNumber)
+			void FrameBoundary(FEvent::FFrameBoundary const& Event, FFrameState& FrameState)
 			{
 				check(!bBusy);
 				QueueFrameState.CPUFrameBoundary = Event.CPUTimestamp;
@@ -1682,7 +1696,7 @@ namespace UE::RHI::GPUProfiler
 			#endif
 
 			#if WITH_PROFILEGPU
-				if (InProfileFrameNumber == Event.FrameNumber + 1)
+				if (Event.bProfileNextFrame)
 				{
 					Profile.bProfileFrame = true;
 
@@ -1709,11 +1723,8 @@ namespace UE::RHI::GPUProfiler
 			}
 		};
 
-		std::atomic<bool> bTriggerProfile{ false };
-		std::atomic<bool> bIsProfiling{ false };
-
-		uint32 ProfileFrameNumber = 0;
-		uint32 MaxFrameNumber = 0;
+		std::atomic<int32> TriggerProfileRefCount { 0 };
+		std::atomic<int32> IsProfilingRefCount    { 0 };
 
 		TMap<uint32, FFrameState> Frames;
 		TMap<FQueue, TUniquePtr<FQueueState>> QueueStates;
@@ -1873,15 +1884,13 @@ namespace UE::RHI::GPUProfiler
 					#endif
 
 					#if WITH_PROFILEGPU
-						// Latch the index of the next frame to profile
-						MaxFrameNumber = FMath::Max(FrameBoundary.FrameNumber, MaxFrameNumber);
-						if (bTriggerProfile.exchange(false))
+						if (QueueState.Profile.bProfileFrame)
 						{
-							ProfileFrameNumber = MaxFrameNumber + 1;
+							IsProfilingRefCount--;
 						}
-					#endif // WITH_PROFILEGPU
+					#endif
 
-						QueueState.FrameBoundary(FrameBoundary, FrameState, ProfileFrameNumber);
+						QueueState.FrameBoundary(FrameBoundary, FrameState);
 
 						if (FrameState.Num() == QueueStates.Num())
 						{
@@ -1904,13 +1913,6 @@ namespace UE::RHI::GPUProfiler
 							ProcessFrame(FrameState);
 
 							Frames.Remove(FrameBoundary.FrameNumber);
-
-						#if WITH_PROFILEGPU
-							if (ProfileFrameNumber == FrameBoundary.FrameNumber)
-							{
-								bIsProfiling = false;
-							}
-						#endif
 						}
 					}
 					break;
@@ -2105,19 +2107,32 @@ namespace UE::RHI::GPUProfiler
 		TEXT("Captures statistics about a frame of GPU work and prints the results to the log."),
 		FConsoleCommandWithArgsDelegate::CreateLambda([](const TArray<FString>& Args)
 		{
-			GGPUProfiler.bIsProfiling = true;
-			GGPUProfiler.bTriggerProfile = true;
+			GGPUProfiler.IsProfilingRefCount += GGPUProfiler.QueueStates.Num();
+			GGPUProfiler.TriggerProfileRefCount++;
 		}));
 #endif // WITH_PROFILEGPU
 
 	RHI_API bool IsProfiling()
 	{
 	#if WITH_PROFILEGPU
-		return GGPUProfiler.bIsProfiling;
+		return GGPUProfiler.IsProfilingRefCount > 0;
 	#else
 		return false;
 	#endif
 	}
+
+	RHI_API bool ShouldProfileNextFrame()
+	{
+	#if WITH_PROFILEGPU
+		if (GGPUProfiler.TriggerProfileRefCount > 0)
+		{
+			GGPUProfiler.TriggerProfileRefCount--;
+			return true;
+		}
+	#endif
+
+		return false;
+	}
 }
 
 RHI_API FRHIGPUFrameTimeHistory GRHIGPUFrameTimeHistory;
diff --git a/Engine/Source/Runtime/RHI/Private/RHICommandList.cpp b/Engine/Source/Runtime/RHI/Private/RHICommandList.cpp
--- a/Engine/Source/Runtime/RHI/Private/RHICommandList.cpp
+++ b/Engine/Source/Runtime/RHI/Private/RHICommandList.cpp
@@ -1235,7 +1235,7 @@ void FRHICommandListExecutor::FSubmitState::Submit(const FSubmitArgs& Args)
 
 	if (EnumHasAnyFlags(SubmitFlags, ERHISubmitFlags::EndFrame))
 	{
-		FDynamicRHI::FRHIEndFrameArgs EndFrameArgs
+		FRHIEndFrameArgs EndFrameArgs
 		{
 			.FrameNumber = GRHICommandList.FrameNumber++,
 		#if WITH_RHI_BREADCRUMBS
@@ -1244,6 +1244,7 @@ void FRHICommandListExecutor::FSubmitState::Submit(const FSubmitArgs& Args)
 		#if STATS
 			.StatsFrame = Args.StatsFrame,
 		#endif
+			.bProfileNextFrame = EnumHasAnyFlags(SubmitFlags, ERHISubmitFlags::ProfileNextFrame)
 		};
 
 		GDynamicRHI->RHIEndFrame(EndFrameArgs);
diff --git a/Engine/Source/Runtime/RHI/Public/DynamicRHI.h b/Engine/Source/Runtime/RHI/Public/DynamicRHI.h
--- a/Engine/Source/Runtime/RHI/Public/DynamicRHI.h
+++ b/Engine/Source/Runtime/RHI/Public/DynamicRHI.h
@@ -200,6 +200,21 @@ struct FRHILockTextureResult
 	uint32 Stride    = 0;
 };
 
+struct FRHIEndFrameArgs
+{
+	// Increments once per call to RHIEndFrame
+	uint32 FrameNumber;
+
+#if WITH_RHI_BREADCRUMBS
+	const TRHIPipelineArray<FRHIBreadcrumbNode*>& GPUBreadcrumbs;
+#endif
+#if STATS
+	TOptional<int64> StatsFrame;
+#endif
+
+	bool bProfileNextFrame = false;
+};
+
 /** The interface which is implemented by the dynamically bound RHI. */
 class FDynamicRHI
 {
@@ -232,19 +247,6 @@ public:
 	/////// RHI Methods
 
 	RHI_API virtual void RHIEndFrame_RenderThread(FRHICommandListImmediate& RHICmdList);
-
-	struct FRHIEndFrameArgs
-	{
-		// Increments once per call to RHIEndFrame
-		uint32 FrameNumber;
-
-	#if WITH_RHI_BREADCRUMBS
-		const TRHIPipelineArray<FRHIBreadcrumbNode*>& GPUBreadcrumbs;
-	#endif
-	#if STATS
-		TOptional<int64> StatsFrame;
-	#endif
-	};
 	virtual void RHIEndFrame(const FRHIEndFrameArgs& Args) = 0;
 
 	// FlushType: Thread safe
diff --git a/Engine/Source/Runtime/RHI/Public/GPUProfiler.h b/Engine/Source/Runtime/RHI/Public/GPUProfiler.h
--- a/Engine/Source/Runtime/RHI/Public/GPUProfiler.h
+++ b/Engine/Source/Runtime/RHI/Public/GPUProfiler.h
@@ -17,6 +17,8 @@
 #include "Containers/SpscQueue.h"
 #include "Containers/StaticArray.h"
 
+struct FRHIEndFrameArgs;
+
 namespace UE::RHI::GPUProfiler
 {
 	struct FQueue
@@ -94,12 +96,15 @@ namespace UE::RHI::GPUProfiler
 
 			// The index of the frame that just ended.
 			// Very first frame of the engine is frame 0 (from boot to first call to RHIEndFrame).
-			uint32 FrameNumber;
+			uint32 const FrameNumber;
+
+			// When true, triggers the profiling of the next frame
+			bool const bProfileNextFrame;
 
 		#if STATS
 			// Should be TOptional<int64> but it is not trivially destructible
-			bool bStatsFrameSet;
-			int64 StatsFrame;
+			bool const bStatsFrameSet;
+			int64 const StatsFrame;
 		#endif
 
 		#if WITH_RHI_BREADCRUMBS
@@ -107,26 +112,7 @@ namespace UE::RHI::GPUProfiler
 			FRHIBreadcrumbNode* const Breadcrumb;
 		#endif
 
-			FFrameBoundary(
-				  uint64 CPUTimestamp
-				, uint32 FrameNumber
-			#if WITH_RHI_BREADCRUMBS
-				, FRHIBreadcrumbNode* Breadcrumb
-			#endif
-			#if STATS
-				, TOptional<int64> StatsFrame
-			#endif
-				)
-				: CPUTimestamp(CPUTimestamp)
-				, FrameNumber(FrameNumber)
-			#if STATS
-				, bStatsFrameSet(StatsFrame.IsSet())
-				, StatsFrame(StatsFrame.IsSet() ? *StatsFrame : 0)
-			#endif
-			#if WITH_RHI_BREADCRUMBS
-				, Breadcrumb(Breadcrumb)
-			#endif
-			{}
+			RHI_API FFrameBoundary(ERHIPipeline Pipeline, FRHIEndFrameArgs const& Args, uint64 CPUTimestamp);
 		};
 
 		// When present in the stream, overrides the total GPU time stat with the value it contains.
@@ -563,6 +549,8 @@ namespace UE::RHI::GPUProfiler
 	};
 
 	RHI_API bool IsProfiling();
+
+	RHI_API bool ShouldProfileNextFrame();
 }
 
 //
diff --git a/Engine/Source/Runtime/RHI/Public/RHICommandList.h b/Engine/Source/Runtime/RHI/Public/RHICommandList.h
--- a/Engine/Source/Runtime/RHI/Public/RHICommandList.h
+++ b/Engine/Source/Runtime/RHI/Public/RHICommandList.h
@@ -4366,15 +4366,18 @@ enum class ERHISubmitFlags
 	// and calls RHIEndFrame for platform RHIs to do various cleanup tasks.
 	EndFrame = 1 << 3,
 
+	// When combined with EndFrame, triggers a GPU profile on the following frame.
+	ProfileNextFrame = 1 << 4,
+
 #if CAN_TOGGLE_COMMAND_LIST_BYPASS
 	// Used when toggling RHI command bypass.
-	EnableBypass  = 1 << 4,
-	DisableBypass = 1 << 5,
+	EnableBypass  = 1 << 5,
+	DisableBypass = 1 << 6,
 #endif
 
 #if WITH_RHI_BREADCRUMBS
-	EnableDrawEvents = 1 << 6,
-	DisableDrawEvents = 1 << 7
+	EnableDrawEvents  = 1 << 7,
+	DisableDrawEvents = 1 << 8
 #endif
 };
 
diff --git a/Engine/Source/Runtime/VulkanRHI/Private/VulkanRHI.cpp b/Engine/Source/Runtime/VulkanRHI/Private/VulkanRHI.cpp
--- a/Engine/Source/Runtime/VulkanRHI/Private/VulkanRHI.cpp
+++ b/Engine/Source/Runtime/VulkanRHI/Private/VulkanRHI.cpp
@@ -1235,14 +1235,7 @@ void FVulkanDynamicRHI::RHIEndFrame(const FRHIEndFrameArgs& Args)
 		}
 
 		// CPU timestamp for the frame boundary event is filled in by the submission thread
-		Payload.EndFrameEvent = UE::RHI::GPUProfiler::FEvent::FFrameBoundary(0, Args.FrameNumber
-		#if WITH_RHI_BREADCRUMBS
-			, (Pipeline != ERHIPipeline::None) ? Args.GPUBreadcrumbs[Pipeline] : nullptr
-		#endif
-		#if STATS
-			, Args.StatsFrame
-		#endif
-		);
+		Payload.EndFrameEvent = UE::RHI::GPUProfiler::FEvent::FFrameBoundary(Pipeline, Args, 0);
 	});
 
 	// Pump the interrupt queue to gather completed events
diff --git a/Engine/Source/Runtime/Windows/D3D11RHI/Private/D3D11RHI.cpp b/Engine/Source/Runtime/Windows/D3D11RHI/Private/D3D11RHI.cpp
--- a/Engine/Source/Runtime/Windows/D3D11RHI/Private/D3D11RHI.cpp
+++ b/Engine/Source/Runtime/Windows/D3D11RHI/Private/D3D11RHI.cpp
@@ -35,14 +35,7 @@ void FD3D11DynamicRHI::RHIEndFrame(const FRHIEndFrameArgs& Args)
 		uint64 Timestamp = FPlatformTime::Cycles64();
 
 		// Insert frame boundary
-		EmplaceProfilerEvent<UE::RHI::GPUProfiler::FEvent::FFrameBoundary>(Timestamp, Args.FrameNumber
-	#if WITH_RHI_BREADCRUMBS
-			, Args.GPUBreadcrumbs[ERHIPipeline::Graphics]
-	#endif
-	#if STATS
-			, Args.StatsFrame
-	#endif
-		);
+		EmplaceProfilerEvent<UE::RHI::GPUProfiler::FEvent::FFrameBoundary>(ERHIPipeline::Graphics, Args, Timestamp);
 
 		Profiler.Pending.Enqueue(MakeUnique<FProfiler::FFrame>(MoveTemp(Profiler.Current)));
 
