Keep precached PSOs in memory for all active materials to fix possible micro hitching when reading from the driver cache - Precached PSOs are refcount tracked for all loaded materials and on material unload all the PSOs precached for that material are decremented in refcount and potentially released (Default disabled)

#Synced-CL 49719249

diff --git a/Engine/Plugins/FX/Niagara/Source/Niagara/Private/NiagaraEmitter.cpp b/Engine/Plugins/FX/Niagara/Source/Niagara/Private/NiagaraEmitter.cpp
--- a/Engine/Plugins/FX/Niagara/Source/Niagara/Private/NiagaraEmitter.cpp
+++ b/Engine/Plugins/FX/Niagara/Source/Niagara/Private/NiagaraEmitter.cpp
@@ -3541,7 +3541,7 @@ FGraphEventArray FVersionedNiagaraEmitterData::PrecacheComputePSOs(const UNiagar
 			}
 
 			FRHIComputeShader* ComputeShader = ShaderRef.GetComputeShader();
-			FPSOPrecacheRequestResult PSOPrecacheRequestResult = PipelineStateCache::PrecacheComputePipelineState(ComputeShader, TEXT("NiagaraEmitter"), true);
+			FPSOPrecacheRequestResult PSOPrecacheRequestResult = PipelineStateCache::PrecacheComputePipelineState(ComputeShader, EPSOPrecacheInMemoryMode::Uncached, TEXT("NiagaraEmitter"), true);
 			if (PSOPrecacheRequestResult.AsyncCompileEvent.IsValid())
 			{
 				PSOPrecacheEvents.Add(PSOPrecacheRequestResult.AsyncCompileEvent);
@@ -3633,7 +3633,7 @@ FGraphEventArray FVersionedNiagaraEmitterData::PrecacheComputePSOs(const UNiagar
 			}
 
 			FRHIComputeShader* ComputeShader = ShaderRef.GetComputeShader();
-			FPSOPrecacheRequestResult PSOPrecacheRequestResult = PipelineStateCache::PrecacheComputePipelineState(ComputeShader, TEXT("NiagaraEmitter"), true);
+			FPSOPrecacheRequestResult PSOPrecacheRequestResult = PipelineStateCache::PrecacheComputePipelineState(ComputeShader, EPSOPrecacheInMemoryMode::Uncached, TEXT("NiagaraEmitter"), true);
 			if (PSOPrecacheRequestResult.AsyncCompileEvent != nullptr && bAddToPSOReadyGraphTasks)
 			{
 				PSOReadyGraphTasks.Add(PSOPrecacheRequestResult.AsyncCompileEvent);
diff --git a/Engine/Plugins/FX/Niagara/Source/NiagaraShader/Private/NiagaraShared.cpp b/Engine/Plugins/FX/Niagara/Source/NiagaraShader/Private/NiagaraShared.cpp
--- a/Engine/Plugins/FX/Niagara/Source/NiagaraShader/Private/NiagaraShared.cpp
+++ b/Engine/Plugins/FX/Niagara/Source/NiagaraShader/Private/NiagaraShared.cpp
@@ -553,7 +553,7 @@ void FNiagaraShaderScript::UpdateCachedData_PostCompile(bool bCalledFromSerializ
 					check(NiagaraShader->GetFrequency() == SF_Compute);
 					FRHIShader* RHIShader = GameThreadShaderMap->GetResource()->GetShader(Shader->GetResourceIndex());
 					FRHIComputeShader* RHIComputeShader = static_cast<FRHIComputeShader*>(RHIShader);
-					PipelineStateCache::PrecacheComputePipelineState(RHIComputeShader, TEXT("NiagaraCompute"));
+					PipelineStateCache::PrecacheComputePipelineState(RHIComputeShader, EPSOPrecacheInMemoryMode::Uncached, TEXT("NiagaraCompute"));
 				}
 			}
 		}
diff --git a/Engine/Source/Runtime/Engine/Private/PSOPrecache.cpp b/Engine/Source/Runtime/Engine/Private/PSOPrecache.cpp
--- a/Engine/Source/Runtime/Engine/Private/PSOPrecache.cpp
+++ b/Engine/Source/Runtime/Engine/Private/PSOPrecache.cpp
@@ -211,7 +211,7 @@ FPSOPrecacheRequestResultArray RequestPrecachePSOs(EPSOPrecacheType PSOPrecacheT
 				}
 				#endif // PSO_PRECACHING_VALIDATE			
 
-				FPSOPrecacheRequestResult PSOPrecacheResult = PipelineStateCache::PrecacheGraphicsPipelineState(PrecacheData.GraphicsPSOInitializer);
+				FPSOPrecacheRequestResult PSOPrecacheResult = PipelineStateCache::PrecacheGraphicsPipelineState(PrecacheData.GraphicsPSOInitializer, PrecacheData.InMemoryMode);
 				if (PSOPrecacheResult.IsValid() && PrecacheData.bRequired)
 				{
 					Results.AddUnique(PSOPrecacheResult);
@@ -227,7 +227,7 @@ FPSOPrecacheRequestResultArray RequestPrecachePSOs(EPSOPrecacheType PSOPrecacheT
 				}
 				#endif // PSO_PRECACHING_VALIDATE
 			
-				FPSOPrecacheRequestResult PSOPrecacheResult = PipelineStateCache::PrecacheComputePipelineState(PrecacheData.ComputeInitializer);
+				FPSOPrecacheRequestResult PSOPrecacheResult = PipelineStateCache::PrecacheComputePipelineState(PrecacheData.ComputeInitializer, PrecacheData.InMemoryMode);
 				if (PSOPrecacheResult.IsValid() && PrecacheData.bRequired)
 				{
 					Results.AddUnique(PSOPrecacheResult);
diff --git a/Engine/Source/Runtime/Engine/Private/PSOPrecacheMaterial.cpp b/Engine/Source/Runtime/Engine/Private/PSOPrecacheMaterial.cpp
--- a/Engine/Source/Runtime/Engine/Private/PSOPrecacheMaterial.cpp
+++ b/Engine/Source/Runtime/Engine/Private/PSOPrecacheMaterial.cpp
@@ -35,6 +35,13 @@ static FAutoConsoleVariableRef CVarShaderPreloadFilterUniqueRequest(
 	ECVF_Default
 );
 
+int32 GPSOPrecacheKeepInMemoryForActiveMaterials = 0;
+static FAutoConsoleVariableRef CVarPSOPrecacheKeepInMemoryForActiveMaterials(
+	TEXT("r.PSOPrecache.KeepInMemoryForActiveMaterials"),
+	GPSOPrecacheKeepInMemoryForActiveMaterials,
+	TEXT("Keep all precached PSOs in memory for loaded materials - if the material is released then all the PSOs unique for this material are also released"),
+	ECVF_ReadOnly);
+
 FPSOCollectorCreateManager::FPSOCollectorData FPSOCollectorCreateManager::PSOCollectors[(int32)EShadingPath::Num][FPSOCollectorCreateManager::MaxPSOCollectorCount] = {};
 
 int32 FPSOCollectorCreateManager::GetIndex(EShadingPath ShadingPath, const TCHAR* Name)
@@ -106,6 +113,7 @@ public:
 	}
 
 	void DoTask(ENamedThreads::Type CurrentThread, const FGraphEventRef& MyCompletionGraphEvent);
+	static FPSOPrecacheRequestResultArray CollectPrecacheDataAndStartCompiles(const FMaterialPSOPrecacheParams& PrecacheParams, uint32 LifecycleID);
 
 	~FMaterialPSOPrecacheCollectionTask()
 	{
@@ -292,14 +300,10 @@ public:
 			{
 				// Collect pso data. Note we don't explicitly collect and preload shaders here since we're not using background tasks
 				// and doing so in separate phases wouldn't benefit anything.
-				FPSOPrecacheDataArray PSOPrecacheData = Params.Material->GetGameThreadShaderMap()->CollectPSOPrecacheData(Params);
 
 				// Start the async compiles
-				FPSOPrecacheRequestResultArray PrecacheResults = RequestPrecachePSOs(EPSOPrecacheType::MeshPass, PSOPrecacheData);
-
-				// Mark collection complete
-				MarkCollectionComplete(Params, PSOPrecacheData, PrecacheResults, LifecycleID);
-
+				FPSOPrecacheRequestResultArray PrecacheResults = FMaterialPSOPrecacheCollectionTask::CollectPrecacheDataAndStartCompiles(Params, LifecycleID);
+				
 				// Add the graph events to wait for
 				for (FPSOPrecacheRequestResult& Result : PrecacheResults)
 				{
@@ -371,11 +375,18 @@ public:
 		FPrecacheData* FindResult = MaterialPSORequestData.Find(Params);
 		check(FindResult->State == EState::Collecting);
 		check(FindResult->ActivePSOPrecacheRequests.IsEmpty());
-		FindResult->ActivePSOPrecacheRequests = PrecacheRequestResults;
+		for (const FPSOPrecacheRequestResult& RequestResult : PrecacheRequestResults)
+		{
+			if (RequestResult.AsyncCompileEvent)
+			{
+				FindResult->ActivePSOPrecacheRequests.Add(RequestResult);
+			}
+			else
+			{
+				FindResult->CachedPSOPrecacheRequests.Add(RequestResult.RequestID);
+			}
+		}
 #if PSO_PRECACHING_TRACKING
-		// TODO: remove PSOPrecachaData from FPrecacheData because it's stored already in the tracked data
-		// 	     but need to be sure that every material name is unique then.
-		// 		 Currently checked by pointer of material and not just name - but pointer can't be stored for tracked because material could be deleted already
 		FindResult->PSOPrecachaData = PrecacheData;
 #if PSO_PRECACHING_VALIDATE
 		AddTrackedMaterialData(Params, PrecacheData);
@@ -398,7 +409,7 @@ public:
 		}
 	}
 
-	void ReleasePrecacheData(FMaterialPSOPrecacheRequestID MaterialPSORequestID)
+	void ReleaseMaterialPrecacheData(FMaterialPSOPrecacheRequestID MaterialPSORequestID)
 	{
 		LLM_SCOPE(ELLMTag::PSO);
 		LLM_TAGSET_SCOPE_CLEAR(ELLMTagSet::Assets);
@@ -407,7 +418,37 @@ public:
 		check(MaterialPSORequestID != INDEX_NONE);
 
 		FRWScopeLock WriteLock(RWLock, SLT_Write);
+		if ((int32)MaterialPSORequestID >= MaterialPSORequests.Num())
+		{
+			return;
+		}
+
 		const FMaterialPSOPrecacheParams& Params = MaterialPSORequests[MaterialPSORequestID];
+				
+		if (GPSOPrecacheKeepInMemoryForActiveMaterials)
+		{
+			FPrecacheData* FindResult = MaterialPSORequestData.Find(Params);
+			check(FindResult && FindResult->State != EState::Collecting); //< Shouldn't be collecting because material is kept alive then
+			if (FindResult)
+			{
+				if (FindResult->State == EState::Compiling)
+				{
+					check(!FindResult->ActivePSOPrecacheRequests.IsEmpty());
+
+					// Request release for all active and finished PSOs
+					for (const FPSOPrecacheRequestResult& RequestResult : FindResult->ActivePSOPrecacheRequests)
+					{
+						PipelineStateCache::ReleasePrecachedPSO(RequestResult.RequestID);
+					}
+				}
+
+				// Release cached PSOs
+				for (const FPSOPrecacheRequestID& RequestID : FindResult->CachedPSOPrecacheRequests)
+				{
+					PipelineStateCache::ReleasePrecachedPSO(RequestID);
+				}
+			}
+		}
 
 		// Mark invalid & remove from from map (could reused IDs with free list)
 		verify(MaterialPSORequestData.Remove(Params) == 1);
@@ -493,11 +534,27 @@ public:
 		LifecycleID++;
 
 		TSet<FMaterial*> Materials;
-		for (FMaterialPSOPrecacheParams& Params : MaterialPSORequests)
+		for (auto Iter : MaterialPSORequestData)
 		{
+			// Collect all materials
+			FMaterialPSOPrecacheParams& Params = Iter.Key;
 			if (Params.Material)
 			{
 				Materials.Add(Params.Material);
+			}	
+			
+			if (GPSOPrecacheKeepInMemoryForActiveMaterials)
+			{
+				// Request release for all active and finished PSOs
+				FPrecacheData& PrecachData = Iter.Value;
+				for (const FPSOPrecacheRequestResult& RequestResult : PrecachData.ActivePSOPrecacheRequests)
+				{
+					PipelineStateCache::ReleasePrecachedPSO(RequestResult.RequestID);
+				}
+				for (const FPSOPrecacheRequestID& RequestID : PrecachData.CachedPSOPrecacheRequests)
+				{
+					PipelineStateCache::ReleasePrecachedPSO(RequestID);
+				}
 			}
 		}
 
@@ -723,6 +780,7 @@ private:
 		FGraphEventRef CollectionGraphEvent;
 		FPSOPrecacheRequestResultArray ActivePSOPrecacheRequests;
 		EPSOPrecachePriority Priority;
+		TArray<FPSOPrecacheRequestID> CachedPSOPrecacheRequests;
 #if PSO_PRECACHING_TRACKING
 		FPSOPrecacheDataArray PSOPrecachaData;
 #endif // PSO_PRECACHING_TRACKING
@@ -740,6 +798,7 @@ private:
 				FPSOPrecacheRequestResult& RequestResult = PrecacheData.ActivePSOPrecacheRequests[Index];
 				if (!PipelineStateCache::IsPrecaching(RequestResult.RequestID))
 				{
+					PrecacheData.CachedPSOPrecacheRequests.Add(RequestResult.RequestID);
 					PrecacheData.ActivePSOPrecacheRequests.RemoveAtSwap(Index);
 					Index--;
 				}
@@ -827,6 +886,31 @@ private:
 // The global request manager - only used locally in a few global function to precache, release or boost PSO precache requests
 FMaterialPSORequestManager GMaterialPSORequestManager;
 
+FPSOPrecacheRequestResultArray FMaterialPSOPrecacheCollectionTask::CollectPrecacheDataAndStartCompiles(const FMaterialPSOPrecacheParams& PrecacheParams, uint32 LifecycleID)
+{
+	// Collect pso data
+	FPSOPrecacheDataArray PSOPrecacheData;
+	if (PrecacheParams.Material->GetGameThreadShaderMap())
+	{
+		PSOPrecacheData = PrecacheParams.Material->GetGameThreadShaderMap()->CollectPSOPrecacheData(PrecacheParams);
+	}
+
+	// Set the correct in memory mode for all the requested PSOs
+	for (FPSOPrecacheData& PrecacheData : PSOPrecacheData)
+	{
+		// Also exclude default material from in memory?
+		PrecacheData.InMemoryMode = GPSOPrecacheKeepInMemoryForActiveMaterials ? PrecacheData.InMemoryMode : EPSOPrecacheInMemoryMode::Uncached;
+	}
+
+	// Start the async compiles
+	FPSOPrecacheRequestResultArray PrecacheResults = RequestPrecachePSOs(EPSOPrecacheType::MeshPass, PSOPrecacheData);
+
+	// Mark collection complete
+	GMaterialPSORequestManager.MarkCollectionComplete(PrecacheParams, PSOPrecacheData, PrecacheResults, LifecycleID);
+
+	return MoveTemp(PrecacheResults);
+}
+
 void FMaterialPSOPrecacheCollectionTask::DoTask(ENamedThreads::Type CurrentThread, const FGraphEventRef& MyCompletionGraphEvent)
 {
 	LLM_SCOPE(ELLMTag::PSO);
@@ -849,18 +933,7 @@ void FMaterialPSOPrecacheCollectionTask::DoTask(ENamedThreads::Type CurrentThrea
 
 	FTaskTagScope ParallelGTScope(ETaskTag::EParallelGameThread);
 
-	// Collect pso data
-	FPSOPrecacheDataArray PSOPrecacheData;
-	if (PrecacheParams.Material->GetGameThreadShaderMap())
-	{
-		PSOPrecacheData = PrecacheParams.Material->GetGameThreadShaderMap()->CollectPSOPrecacheData(PrecacheParams);
-	}
-
-	// Start the async compiles
-	FPSOPrecacheRequestResultArray PrecacheResults = RequestPrecachePSOs(EPSOPrecacheType::MeshPass, PSOPrecacheData);
-
-	// Mark collection complete
-	GMaterialPSORequestManager.MarkCollectionComplete(PrecacheParams, PSOPrecacheData, PrecacheResults, RequestLifecycleID);
+	FPSOPrecacheRequestResultArray PrecacheResults = CollectPrecacheDataAndStartCompiles(PrecacheParams, RequestLifecycleID);
 
 	// Won't touch the material interface anymore - PSO compile jobs take refs to all RHI resources while creating the task
 	MaterialInterface.Reset();
@@ -970,7 +1043,7 @@ void ReleasePSOPrecacheData(const TArray<FMaterialPSOPrecacheRequestID>& Materia
 {
 	for (FMaterialPSOPrecacheRequestID RequestID : MaterialPSORequestIDs)
 	{
-		GMaterialPSORequestManager.ReleasePrecacheData(RequestID);
+		GMaterialPSORequestManager.ReleaseMaterialPrecacheData(RequestID);
 	}
 }
 
diff --git a/Engine/Source/Runtime/Engine/Private/ShaderCompiler/ShaderCompiler.cpp b/Engine/Source/Runtime/Engine/Private/ShaderCompiler/ShaderCompiler.cpp
--- a/Engine/Source/Runtime/Engine/Private/ShaderCompiler/ShaderCompiler.cpp
+++ b/Engine/Source/Runtime/Engine/Private/ShaderCompiler/ShaderCompiler.cpp
@@ -5101,7 +5101,7 @@ void PrecacheComputePipelineStatesForGlobalShaders(ERHIFeatureLevel::Type Featur
 							if (RHIComputeShader)
 							{
 								const TCHAR* TypeName = GlobalShader.GetType()->GetName();
-								PipelineStateCache::PrecacheComputePipelineState(RHIComputeShader, TypeName);
+								PipelineStateCache::PrecacheComputePipelineState(RHIComputeShader, EPSOPrecacheInMemoryMode::Uncached, TypeName);
 							}
 						}
 					));
@@ -5125,7 +5125,7 @@ void PrecacheComputePipelineStatesForGlobalShaders(ERHIFeatureLevel::Type Featur
 							if (RHIComputeShader)
 							{
 								const TCHAR* TypeName = GlobalShader.GetType()->GetName();
-								PipelineStateCache::PrecacheComputePipelineState(RHIComputeShader, TypeName);
+								PipelineStateCache::PrecacheComputePipelineState(RHIComputeShader, EPSOPrecacheInMemoryMode::Uncached, TypeName);
 							}
 						}
 					});
@@ -5161,6 +5161,12 @@ void PrecacheComputePipelineStatesForGlobalShaders(ERHIFeatureLevel::Type Featur
 		}
 	}
 
+	// Don't keep global PSOs in memory
+	for (FPSOPrecacheData& PrecacheData : GlobalPSOInitializers)
+	{
+		PrecacheData.InMemoryMode = EPSOPrecacheInMemoryMode::Uncached;
+	}
+
 	RequestPrecachePSOs(EPSOPrecacheType::Global, GlobalPSOInitializers);
 }
 
diff --git a/Engine/Source/Runtime/Engine/Public/PSOPrecache.h b/Engine/Source/Runtime/Engine/Public/PSOPrecache.h
--- a/Engine/Source/Runtime/Engine/Public/PSOPrecache.h
+++ b/Engine/Source/Runtime/Engine/Public/PSOPrecache.h
@@ -235,7 +235,7 @@ extern ENGINE_API EPSOPrecacheMode GetPSOPrecacheMode();
  */
 struct FPSOPrecacheData
 {
-	FPSOPrecacheData() : Type(EType::Graphics), bRequired(true)
+	FPSOPrecacheData() : Type(EType::Graphics), bRequired(true), InMemoryMode(EPSOPrecacheInMemoryMode::Cached)
 #if PSO_PRECACHING_VALIDATE
 		, PSOCollectorIndex(INDEX_NONE)
 		, bDefaultMaterial(0)
@@ -255,6 +255,9 @@ struct FPSOPrecacheData
 	// (proxy creation won't wait for these PSOs if enabled)
 	bool bRequired;
 
+	// Keep the PSO in memory until the owning material is request deleted
+	EPSOPrecacheInMemoryMode InMemoryMode;
+
 	FGraphicsPipelineStateInitializer GraphicsPSOInitializer;
 	FComputePipelineStateInitializer ComputeInitializer;
 
diff --git a/Engine/Source/Runtime/RHI/Private/PipelineStateCache.cpp b/Engine/Source/Runtime/RHI/Private/PipelineStateCache.cpp
--- a/Engine/Source/Runtime/RHI/Private/PipelineStateCache.cpp
+++ b/Engine/Source/Runtime/RHI/Private/PipelineStateCache.cpp
@@ -56,7 +56,19 @@ DECLARE_DWORD_ACCUMULATOR_STAT(TEXT("High Priority Compute PSO Precache Requests
 
 DECLARE_DWORD_ACCUMULATOR_STAT(TEXT("Highest Priority Graphics PSO Precache Requests"), STAT_HighestPriorityGraphicsPSOPrecacheRequests, STATGROUP_PipelineStateCache);
 DECLARE_DWORD_ACCUMULATOR_STAT(TEXT("Highest Priority Compute PSO Precache Requests"), STAT_HighestPriorityComputePSOPrecacheRequests, STATGROUP_PipelineStateCache);
-DECLARE_DWORD_ACCUMULATOR_STAT(TEXT("Precached PSOs Kept In Memory"), STAT_InMemoryPrecachedPSOCount, STATGROUP_PipelineStateCache);
+
+DECLARE_DWORD_ACCUMULATOR_STAT(TEXT("Total Precached Graphics PSOs"), STAT_TotalGraphicsPrecachedPSOs, STATGROUP_PipelineStateCache);
+DECLARE_DWORD_ACCUMULATOR_STAT(TEXT("Total Precached Compute PSOs"), STAT_TotalComputePrecachedPSOs, STATGROUP_PipelineStateCache);
+
+DECLARE_STATS_GROUP(TEXT("PSO InMemory"), STATGROUP_PSOInMemory, STATCAT_Advanced);
+
+DECLARE_DWORD_ACCUMULATOR_STAT(TEXT("Precached Graphics PSOs InMemory"), STAT_PrecachedGraphicsPrecachedPSOsInMemory, STATGROUP_PSOInMemory);
+DECLARE_DWORD_ACCUMULATOR_STAT(TEXT("Precached Compute PSOs InMemory"), STAT_PrecachedComputePrecachedPSOsInMemory, STATGROUP_PSOInMemory);
+
+DECLARE_DWORD_ACCUMULATOR_STAT(TEXT("Runtime Graphics PSOs"), STAT_ActiveRuntimeGraphicsPSOs, STATGROUP_PSOInMemory);
+DECLARE_DWORD_ACCUMULATOR_STAT(TEXT("Runtime Compute PSOs"), STAT_ActiveRuntimeComputePSOs, STATGROUP_PSOInMemory);
+DECLARE_DWORD_ACCUMULATOR_STAT(TEXT("Runtime Graphics PSOs From InMemoryCache"), STAT_RuntimeGraphicsPSOsFromInMemoryPSOCache, STATGROUP_PSOInMemory);
+DECLARE_DWORD_ACCUMULATOR_STAT(TEXT("Runtime Compute PSOs From InMemoryCache"), STAT_RuntimeComputePSOsFromInMemoryPSOCache, STATGROUP_PSOInMemory);
 
 static inline uint32 GetTypeHash(const FBoundShaderStateInput& Input)
 {
@@ -1161,19 +1173,42 @@ public:
 	}
 };
 
+
 /* State for compute  */
 class FComputePipelineState : public FPipelineStateAsync
 {
 public:
-	FComputePipelineState(FRHIComputeShader* InComputeShader)
+	FComputePipelineState(FRHIComputeShader* InComputeShader, bool bInRuntimePSO, bool bInFromInMemoryPSOCache)
 		: ComputeShader(InComputeShader)
 	{
 		ComputeShader->AddRef();
+#if STATS
+		bRuntimePSO = bInRuntimePSO; 
+		bFromInMemoryPSOCache = bInFromInMemoryPSOCache;
+		if (bRuntimePSO)
+		{
+			INC_DWORD_STAT(STAT_ActiveRuntimeComputePSOs);
+			if (bFromInMemoryPSOCache)
+			{
+				INC_DWORD_STAT(STAT_RuntimeComputePSOsFromInMemoryPSOCache);
+			}
+		}
+#endif // STATS
 	}
 	
 	~FComputePipelineState()
 	{
 		ComputeShader->Release();
+#if STATS		
+		if (bRuntimePSO)
+		{
+			DEC_DWORD_STAT(STAT_ActiveRuntimeComputePSOs);
+			if (bFromInMemoryPSOCache)
+			{
+				DEC_DWORD_STAT(STAT_RuntimeComputePSOsFromInMemoryPSOCache);
+			}
+		}
+#endif // STATS
 	}
 
 	virtual bool IsCompute() const
@@ -1206,6 +1241,10 @@ public:
 
 	FRHIComputeShader* ComputeShader;
 	TRefCountPtr<FRHIComputePipelineState> RHIPipeline;
+#if STATS
+	bool bRuntimePSO = true;
+	bool bFromInMemoryPSOCache = false;
+#endif // STATS
 #if PIPELINESTATECACHE_VERIFYTHREADSAFE
 	FThreadSafeCounter InUseCount;
 #endif
@@ -1267,12 +1306,39 @@ public:
 #endif
 };
 
+
 /* State for graphics */
 class FGraphicsPipelineState : public FPipelineStateAsync
 {
 public:
-	FGraphicsPipelineState() 
+	FGraphicsPipelineState(bool bInRuntimePSO, bool bInFromInMemoryPSOCache)
 	{
+	#if STATS
+		bRuntimePSO = bInRuntimePSO;
+		bFromInMemoryPSOCache = bInFromInMemoryPSOCache;
+		if (bRuntimePSO)
+		{
+			INC_DWORD_STAT(STAT_ActiveRuntimeGraphicsPSOs);
+			if (bFromInMemoryPSOCache)
+			{
+				INC_DWORD_STAT(STAT_RuntimeGraphicsPSOsFromInMemoryPSOCache);
+			}
+		}
+	#endif // STATS
+	}
+
+	~FGraphicsPipelineState()
+	{
+	#if STATS
+		if (bRuntimePSO)
+		{
+			DEC_DWORD_STAT(STAT_ActiveRuntimeGraphicsPSOs);
+			if (bFromInMemoryPSOCache)
+			{
+				DEC_DWORD_STAT(STAT_RuntimeGraphicsPSOsFromInMemoryPSOCache);
+			}
+		}
+	#endif // STATS
 	}
 
 	virtual bool IsCompute() const
@@ -1308,6 +1374,10 @@ public:
 	FThreadSafeCounter InUseCount;
 #endif
 	uint64 SortKey = 0;
+#if STATS
+	bool bRuntimePSO = true;
+	bool bFromInMemoryPSOCache = false;
+#endif // STATS
 };
 
 FRHIComputePipelineState* GetRHIComputePipelineState(FComputePipelineState* PipelineState)
@@ -2313,11 +2383,12 @@ enum class EPSOPrecacheStateMask : uint8
 {
 	None = 0,
 	Compiling = 1 << 0, // once the start is scheduled
-	Succeeded = 1 << 1, // once compilation is finished
-	Failed = 1 << 2,    // once compilation is finished
-	Boosted = 1 << 3,
-	HighestPri = 1 << 4,
-	UsedForRendering = 1 << 5, // Set once this precached PSO is actually needed and used for rendering.
+	Cached = 1 << 1, 		// PSO is cached
+	Uncached = 1 << 2,  	// PSO is released
+	Failed = 1 << 3,    	// PSO failed to compile
+	Boosted = 1 << 4,		// PSO has been boosted to high
+	HighestPri = 1 << 5,	// PSO has been boosted to highest
+	UsedForRendering = 1 << 6, // Set once this precached PSO is actually needed and used for rendering.
 };
 
 ENUM_CLASS_FLAGS(EPSOPrecacheStateMask)
@@ -2377,7 +2448,7 @@ protected:
 
 		check(!EnumHasAnyFlags(PrevState, EPSOPrecacheStateMask::HighestPri));
 
-		bool bCompleted = EnumHasAnyFlags(PrevState, EPSOPrecacheStateMask::Failed | EPSOPrecacheStateMask::Succeeded);
+		bool bCompleted = EnumHasAnyFlags(PrevState, EPSOPrecacheStateMask::Failed | EPSOPrecacheStateMask::Cached);
 		UE_CLOG(bCompleted, LogRHI, Error, TEXT("pso request has completed? prev %x, new %x"), (uint32)PrevState, (uint32)NewState);
 
 		if (FPSOPrecacheThreadPool::UsePool())
@@ -2409,15 +2480,18 @@ protected:
 		}
 	}
 
-	FPSOPrecacheRequestResult TryAddNewState(const TPrecachedPSOInitializer& Initializer, const FString& PSOCompilationEventName, bool bDoAsyncCompile)
+	FPSOPrecacheRequestResult TryAddNewState(const TPrecachedPSOInitializer& Initializer, EPSOPrecacheInMemoryMode InMemoryMode, const FString& PSOCompilationEventName, bool bDoAsyncCompile)
 	{
+		// Refcounted and max in memory PSOs can't be used together
+		check(InMemoryMode == EPSOPrecacheInMemoryMode::Uncached || !ShouldKeepPrecachedPSOsInMemory());
+
 		FPSOPrecacheRequestResult Result;
 		uint64 InitializerHash = TPrecachePipelineCacheDerived::PipelineStateInitializerHash(Initializer);
 
 		// Fast check first with read lock
 		{
 			FRWScopeLock ReadLock(PrecachePSOsRWLock, SLT_ReadOnly);
-			if (HasPSOBeenRequested(Initializer, InitializerHash, Result))
+			if (IsCachedOrCompiling(Initializer, InitializerHash, InMemoryMode, Result))
 			{
 				return Result;
 			}
@@ -2427,24 +2501,76 @@ protected:
 		TPipelineState* NewPipelineState = nullptr;
 		{
 			FRWScopeLock WriteLock(PrecachePSOsRWLock, SLT_Write);
-			if (HasPSOBeenRequested(Initializer, InitializerHash, Result))
+			if (IsCachedOrCompiling(Initializer, InitializerHash, InMemoryMode, Result))
 			{
 				return Result;
 			}
 
-			// Add to array to get the new RequestID
-			Result.RequestID.Type = (uint32)PSOType;
-			Result.RequestID.RequestID = PrecachedPSOInitializers.Add(InitializerHash);
+			// Check if it's uncached
+			FPrecacheTask* FindResult = PrecachedPSOInitializerData.Find(InitializerHash);
+			if (FindResult)
+			{			
+				check(!ShouldKeepPrecachedPSOsInMemory());
 
-			// create new graphics state
-			NewPipelineState = TPrecachePipelineCacheDerived::CreateNewPSO(Initializer);
+				check(InMemoryMode == EPSOPrecacheInMemoryMode::Cached);				
+				check(FindResult->RefCount == 0);
+
+				FPlatformAtomics::InterlockedIncrement(&FindResult->RefCount);									
+				Result.RequestID = FindResult->RequestID;
+				check(Result.RequestID.IsValid());
+
+				// If it's not uncached then it's still async compiling but requested for cleanup when done
+				// Remove from the cleanup list again.
+				EPSOPrecacheStateMask CurrentState = FindResult->ReadPSOPrecacheState();
+				if (!EnumHasAnyFlags(CurrentState, EPSOPrecacheStateMask::Uncached))
+				{	
+					verify(PrecachedPSOsToCleanup.Remove(InitializerHash) == 1);	
+
+					// Add async compile event if it's still compiling
+					if (!IsCompilationDone(CurrentState))
+					{
+						Result.AsyncCompileEvent = FindResult->PipelineState->GetCompletionEvent();
+					}
+					return Result;
+				}
+
+				// Uncached - reset the state and kick the compilation again
+				FindResult->ResetStateMask();
 
-			// Add data to map
-			FPrecacheTask PrecacheTask;
-			PrecacheTask.PipelineState = NewPipelineState;
-			PrecacheTask.RequestID = Result.RequestID;
-			PrecachedPSOInitializerData.Add(InitializerHash, PrecacheTask);
+				check(!PrecachedPSOsToCleanup.Contains(InitializerHash));
+			}
+			else
+			{
+				// Add to array to get the new RequestID
+				Result.RequestID.Type = (uint32)PSOType;
+				Result.RequestID.RequestID = PrecachedPSOInitializers.Add(InitializerHash);
 
+				// Add data to map
+				FPrecacheTask PrecacheTask;
+				PrecacheTask.RequestID = Result.RequestID;
+				PrecacheTask.RefCount = ShouldKeepPrecachedPSOsInMemory() ? 0 : 1;
+				PrecachedPSOInitializerData.Add(InitializerHash, PrecacheTask);
+
+				FindResult = PrecachedPSOInitializerData.Find(InitializerHash);
+				check(FindResult);
+
+				// Increment total count
+				if (PSOType == FPSOPrecacheRequestID::EType::Graphics)
+				{
+					INC_DWORD_STAT(STAT_TotalGraphicsPrecachedPSOs);
+				}
+				else
+				{
+					INC_DWORD_STAT(STAT_TotalComputePrecachedPSOs);
+				}
+			}
+						
+			// create new graphics state
+			NewPipelineState = TPrecachePipelineCacheDerived::CreateNewPSO(Initializer);			
+			check(FindResult->PipelineState == nullptr);
+			FindResult->PipelineState = NewPipelineState;
+			FindResult->InMemoryMode = InMemoryMode;
+						
 			if (bDoAsyncCompile)
 			{
 				// Assign the event at this point because we need to release the lock before calling OnNewPipelineStateCreated which
@@ -2474,7 +2600,6 @@ protected:
 					
 					InMemoryPSOIndices.PushLast(Result.RequestID.RequestID);
 				}
-				INC_DWORD_STAT(STAT_InMemoryPrecachedPSOCount);
 			}
 		}
 
@@ -2540,7 +2665,7 @@ public:
 					}
 					else if (PrecacheTask.PipelineState)
 					{
-						check(EnumHasAnyFlags(PrecacheTask.ReadPSOPrecacheState(), (EPSOPrecacheStateMask::Succeeded | EPSOPrecacheStateMask::Failed)));
+						check(EnumHasAnyFlags(PrecacheTask.ReadPSOPrecacheState(), (EPSOPrecacheStateMask::Cached | EPSOPrecacheStateMask::Failed)));
 						delete PrecacheTask.PipelineState;
 						PrecacheTask.PipelineState = nullptr;
 					}
@@ -2559,6 +2684,34 @@ public:
 		}
 	}
 
+	bool ReleasePrecachedPSO(const FPSOPrecacheRequestID& RequestID)
+	{
+		check(RequestID.GetType() == PSOType);
+		check(!ShouldKeepPrecachedPSOsInMemory());
+
+		FRWScopeLock WriteLock(PrecachePSOsRWLock, SLT_Write);		
+		if ((int32)RequestID.RequestID >= PrecachedPSOInitializers.Num())
+		{
+			return false;
+		}
+
+		uint64 InitializerHash = PrecachedPSOInitializers[RequestID.RequestID];
+		FPrecacheTask* FindResult = PrecachedPSOInitializerData.Find(InitializerHash);
+		check(FindResult);
+		if (FindResult->InMemoryMode == EPSOPrecacheInMemoryMode::Cached)
+		{
+			check(FindResult->RefCount > 0);
+			int32 NewRefCount = FPlatformAtomics::InterlockedDecrement(&FindResult->RefCount);
+			if (NewRefCount == 0)
+			{
+				PrecachedPSOsToCleanup.Add(InitializerHash);
+				return true;
+			}
+		}
+		
+		return false;
+	}
+
 	EPSOPrecacheResult GetPrecachingState(const FPSOPrecacheRequestID& RequestID)
 	{
 		check(RequestID.GetType() == PSOType);
@@ -2572,7 +2725,8 @@ public:
 		bool bOutMarkedUsed = false;
 		FGraphEventRef CompletionEvent;
 		FPSOPrecacheRequestID PSOPrecacheRequestID;
-		return GetPrecachingStateInternal(InitializerHash, bMarkUsed, bOutMarkedUsed, CompletionEvent, PSOPrecacheRequestID);
+		EPSOPrecacheInMemoryMode InMemoryMode = EPSOPrecacheInMemoryMode::Uncached;
+		return GetPrecachingStateInternal(InitializerHash, bMarkUsed, bOutMarkedUsed, CompletionEvent, PSOPrecacheRequestID, InMemoryMode);
 	}
 
 	EPSOPrecacheResult GetPrecachingState(const TPrecachedPSOInitializer& Initializer, FGraphEventRef& OutCompletionEvent, FPSOPrecacheRequestID& OutPSOPrecacheRequestID)
@@ -2580,7 +2734,8 @@ public:
 		uint64 InitializerHash = TPrecachePipelineCacheDerived::PipelineStateInitializerHash(Initializer);
 		bool bMarkUsed = false;
 		bool bOutMarkedUsed = false;
-		return GetPrecachingStateInternal(InitializerHash, bMarkUsed, bOutMarkedUsed, OutCompletionEvent, OutPSOPrecacheRequestID);
+		EPSOPrecacheInMemoryMode InMemoryMode = EPSOPrecacheInMemoryMode::Uncached;
+		return GetPrecachingStateInternal(InitializerHash, bMarkUsed, bOutMarkedUsed, OutCompletionEvent, OutPSOPrecacheRequestID, InMemoryMode);
 	}
 
 	uint64 GetPSOCompilationTime(uint64 PSOPrecacheHash)
@@ -2601,7 +2756,18 @@ public:
 		bool bMarkUsed = true;
 		FGraphEventRef CompletionEvent;		
 		FPSOPrecacheRequestID PSOPrecacheRequestID;
-		return GetPrecachingStateInternal(InitializerHash, bMarkUsed, bOutMarkedUsed, CompletionEvent, PSOPrecacheRequestID);
+		EPSOPrecacheInMemoryMode InMemoryMode = EPSOPrecacheInMemoryMode::Uncached;
+		return GetPrecachingStateInternal(InitializerHash, bMarkUsed, bOutMarkedUsed, CompletionEvent, PSOPrecacheRequestID, InMemoryMode);
+	}
+
+	EPSOPrecacheResult GetPrecachingStateAndInMemoryMode(const TPrecachedPSOInitializer& Initializer, EPSOPrecacheInMemoryMode& OutInMemoryMode)
+	{
+		uint64 InitializerHash = TPrecachePipelineCacheDerived::PipelineStateInitializerHash(Initializer);
+		bool bMarkUsed = false;
+		bool bOutMarkedUsed = false;
+		FGraphEventRef CompletionEvent;		
+		FPSOPrecacheRequestID PSOPrecacheRequestID;
+		return GetPrecachingStateInternal(InitializerHash, bMarkUsed, bOutMarkedUsed, CompletionEvent, PSOPrecacheRequestID, OutInMemoryMode);
 	}
 
 	bool IsPrecaching()
@@ -2650,21 +2816,39 @@ public:
 			check(FindResult);
 
 #if !UE_BUILD_SHIPPING
-			// Store the compilation time for stat tracking
-			FindResult->PSOCompilationTime = PSOCompilationTime;
+			// Store the compilation time for stat tracking (only store the first compilation time - otherwise coming from driver cache)
+			if (FindResult->PSOCompilationTime == 0)
+			{
+				FindResult->PSOCompilationTime = PSOCompilationTime;
+			}
 #endif
 
 			// We still add the 'compiling' bit here because if the task is fast enough, we can get here before the end of TryAddNewState
-			const EPSOPrecacheStateMask CompleteStateMask = bValid ? (EPSOPrecacheStateMask::Succeeded | EPSOPrecacheStateMask::Compiling) : (EPSOPrecacheStateMask::Failed | EPSOPrecacheStateMask::Compiling);
+			const EPSOPrecacheStateMask CompleteStateMask = bValid ? (EPSOPrecacheStateMask::Cached | EPSOPrecacheStateMask::Compiling) : (EPSOPrecacheStateMask::Failed | EPSOPrecacheStateMask::Compiling);
 			PreviousStateMask = FindResult->AddPSOPrecacheState(CompleteStateMask);
 
-			// Add to array of precached PSOs so it can be cleaned up
-			if (!ShouldKeepPrecachedPSOsInMemory())
+			// Release immediately if not kept in memory
+			if (FindResult->InMemoryMode != EPSOPrecacheInMemoryMode::Cached && !ShouldKeepPrecachedPSOsInMemory())
 			{
+				if (!ShouldKeepPrecachedPSOsInMemory())
+				{
+					check(FindResult->RefCount == 1);
+					FPlatformAtomics::InterlockedDecrement(&FindResult->RefCount);
+				}
 				PrecachedPSOsToCleanup.Add(InitializerHash);
 			}
 		}
 
+		// Mark as in memory now - if it's not cached then it was added to list for cleanup and will be subtracted there
+		if (PSOType == FPSOPrecacheRequestID::EType::Graphics)
+		{
+			INC_DWORD_STAT(STAT_PrecachedGraphicsPrecachedPSOsInMemory);
+		}
+		else
+		{
+			INC_DWORD_STAT(STAT_PrecachedComputePrecachedPSOsInMemory);
+		}
+
         // Need to ensure that the boost request was actually executed: if only it was asked by BoostPriority, but not requested (ie TryAddNewState has not set the Compiling bit
         // yet) then we must ignore the request
 		if (EnumHasAllFlags(PreviousStateMask, EPSOPrecacheStateMask::Boosted | EPSOPrecacheStateMask::Compiling))
@@ -2683,7 +2867,8 @@ public:
 
 	static bool IsCompilationDone(EPSOPrecacheStateMask StateMask)
 	{
-		return EnumHasAnyFlags(StateMask, EPSOPrecacheStateMask::Succeeded | EPSOPrecacheStateMask::Failed);
+		// TODO: Is Uncached also compilation done?
+		return EnumHasAnyFlags(StateMask, EPSOPrecacheStateMask::Cached | EPSOPrecacheStateMask::Uncached | EPSOPrecacheStateMask::Failed);
 	}
 
 	void ProcessDelayedCleanup()
@@ -2697,7 +2882,7 @@ public:
 			uint64 InitializerHash = PrecachedPSOsToCleanup[Index];
 
 			FPrecacheTask* FindResult = PrecachedPSOInitializerData.Find(InitializerHash);
-			check(FindResult && (ShouldKeepPrecachedPSOsInMemory() || IsCompilationDone((FindResult->ReadPSOPrecacheState()))));
+			check(FindResult && FindResult->RefCount == 0);
 			if (!FindResult || !FindResult->PipelineState)
 			{
 				// Was already cleaned up (can happen if it was marked as used).
@@ -2711,14 +2896,27 @@ public:
 				delete FindResult->PipelineState;
 				FindResult->PipelineState = nullptr;
 
-				if (ShouldKeepPrecachedPSOsInMemory())
+				// Mark as uncached
+				FindResult->AddPSOPrecacheState(EPSOPrecacheStateMask::Uncached);
+				FindResult->RemovePSOPrecacheState(EPSOPrecacheStateMask::Cached);
+
+				// Update in memory stats
+				if (PSOType == FPSOPrecacheRequestID::EType::Graphics)
 				{
-					DEC_DWORD_STAT(STAT_InMemoryPrecachedPSOCount);
+					DEC_DWORD_STAT(STAT_PrecachedGraphicsPrecachedPSOsInMemory);
+				}
+				else
+				{
+					DEC_DWORD_STAT(STAT_PrecachedComputePrecachedPSOsInMemory);
 				}
 
 				PrecachedPSOsToCleanup.RemoveAtSwap(Index);
 				Index--;
 			}
+			else
+			{
+				// Compilation still busy - try and cancel?
+			}
 		}
 	}
 
@@ -2735,31 +2933,74 @@ public:
 
 		FRWScopeLock WriteLock(PrecachePSOsRWLock, SLT_Write);
 		PrecachedPSOsToCleanup.Add(InitializerHash);
+
+		// Should also remove the index from InMemoryPSOIndices but it's a deque so can't be done fast
 	}
 
 protected:
-	bool HasPSOBeenRequested(const TPrecachedPSOInitializer& Initializer, uint64 InitializerHash, FPSOPrecacheRequestResult& Result)
+	bool IsCachedOrCompiling(const TPrecachedPSOInitializer& Initializer, uint64 InitializerHash, EPSOPrecacheInMemoryMode InMemoryMode, FPSOPrecacheRequestResult& Result)
 	{
 		FPrecacheTask* FindResult = PrecachedPSOInitializerData.Find(InitializerHash);
+
+		// Found and not uncached?
 		if (FindResult)
 		{			
-			// If not compiled yet, then return the request ID so the caller can check the state
-			if (!IsCompilationDone(FindResult->ReadPSOPrecacheState()))
+			if (ShouldKeepPrecachedPSOsInMemory())
 			{
+				// If not compiled yet, then return the request ID so the caller can check the state
+				if (!IsCompilationDone(FindResult->ReadPSOPrecacheState()))
+				{
+					Result.RequestID = FindResult->RequestID;
+					Result.AsyncCompileEvent = FindResult->PipelineState->GetCompletionEvent();
+					check(Result.RequestID.IsValid());
+				}
+			}
+			else
+			{
+				check(FindResult->InMemoryMode == InMemoryMode);
+
+				// Special ref count tracking when kept in memory
+				if (InMemoryMode == EPSOPrecacheInMemoryMode::Cached)
+				{
+					// If it's uncached then return false because it should be compiled again
+					if (EnumHasAnyFlags(FindResult->ReadPSOPrecacheState(), EPSOPrecacheStateMask::Uncached))
+					{
+						return false;
+					}
+
+					// If refcount 0 then it means it's in the cleanup list and needs to be removed there with write lock (only have read lock here)
+					// only location where ref count is incremented without write lock is below
+					if (FindResult->RefCount == 0)
+					{
+						return false;
+					}
+
+					// Find in cache and still referenced by something else
+					FPlatformAtomics::InterlockedIncrement(&FindResult->RefCount);
+					check(!PrecachedPSOsToCleanup.Contains(InitializerHash));
+				}
+
 				Result.RequestID = FindResult->RequestID;
-				Result.AsyncCompileEvent = FindResult->PipelineState->GetCompletionEvent();
 				check(Result.RequestID.IsValid());
+
+				// Add async compile event if it's still compiling
+				if (!IsCompilationDone(FindResult->ReadPSOPrecacheState()))
+				{
+					Result.AsyncCompileEvent = FindResult->PipelineState->GetCompletionEvent();
+				}
 			}
+
 			return true;
 		}
 
 		return false;
 	}
 
-	EPSOPrecacheResult GetPrecachingStateInternal(uint64 InitializerHash, bool bTryMarkUsed, bool& bOutMarkedUsed, FGraphEventRef& OutCompletionEvent, FPSOPrecacheRequestID& OutPSOPrecacheRequestID)
+	EPSOPrecacheResult GetPrecachingStateInternal(uint64 InitializerHash, bool bTryMarkUsed, bool& bOutMarkedUsed, FGraphEventRef& OutCompletionEvent, FPSOPrecacheRequestID& OutPSOPrecacheRequestID, EPSOPrecacheInMemoryMode& OutInMemoryMode)
 	{
 		EPSOPrecacheStateMask CompilationState = EPSOPrecacheStateMask::None;
 		FGraphEventRef AsyncCompileEvent;
+		OutInMemoryMode = EPSOPrecacheInMemoryMode::Uncached;
 
 		{
 			FRWScopeLock ReadLock(PrecachePSOsRWLock, SLT_ReadOnly);
@@ -2780,6 +3021,8 @@ protected:
 			{
 				CompilationState = FindResult->ReadPSOPrecacheState();
 			}
+
+			OutInMemoryMode = FindResult->InMemoryMode;
 		}
 
 		// We successfully marked it as used if nobody else did concurrently.
@@ -2794,7 +3037,7 @@ protected:
 		check(AsyncCompileEvent == nullptr || AsyncCompileEvent->IsComplete());
 
 		// check we only set 1 completion bit
-		const EPSOPrecacheStateMask CompletionMask = EPSOPrecacheStateMask::Succeeded | EPSOPrecacheStateMask::Failed;
+		const EPSOPrecacheStateMask CompletionMask = EPSOPrecacheStateMask::Cached | EPSOPrecacheStateMask::Failed | EPSOPrecacheStateMask::Uncached;
 		check(EnumHasAnyFlags(CompilationState, CompletionMask) && !EnumHasAllFlags(CompilationState, CompletionMask));
 
 		return (EnumHasAnyFlags(CompilationState, EPSOPrecacheStateMask::Failed)) ? EPSOPrecacheResult::NotSupported : EPSOPrecacheResult::Complete;
@@ -2852,6 +3095,8 @@ protected:
 	{
 		TPipelineState* PipelineState = nullptr;
 		FPSOPrecacheRequestID RequestID;
+		volatile int32 RefCount = 0;
+		EPSOPrecacheInMemoryMode InMemoryMode = EPSOPrecacheInMemoryMode::Uncached;
 #if !UE_BUILD_SHIPPING
 		uint64 PSOCompilationTime = 0;
 #endif
@@ -2862,6 +3107,18 @@ protected:
 			return (EPSOPrecacheStateMask)FPlatformAtomics::InterlockedOr((volatile int8*)&StateMask, (int8)DesiredState);
 		}
 
+		EPSOPrecacheStateMask RemovePSOPrecacheState(EPSOPrecacheStateMask StateToRemove)
+		{
+			static_assert(sizeof(EPSOPrecacheStateMask) == sizeof(int8), "Fix the cast below");
+			return (EPSOPrecacheStateMask)FPlatformAtomics::InterlockedAnd((volatile int8*)&StateMask, (int8)~StateToRemove);
+		}
+
+		EPSOPrecacheStateMask ResetStateMask()
+		{
+			static_assert(sizeof(EPSOPrecacheStateMask) == sizeof(int8), "Fix the cast below");
+			return (EPSOPrecacheStateMask)FPlatformAtomics::InterlockedExchange((volatile int8*)&StateMask, (int8)EPSOPrecacheStateMask::None);
+		}
+
 		inline EPSOPrecacheStateMask ReadPSOPrecacheState()
 		{
 			static_assert(sizeof(EPSOPrecacheStateMask) == sizeof(int8), "Fix the cast below");
@@ -2911,11 +3168,13 @@ class FPrecacheComputePipelineCache : public TPrecachePipelineCacheBase<FPrecach
 public:
 	static FComputePipelineState* CreateNewPSO(const FComputePipelineStateInitializer& ComputeShaderInitializer)
 	{
-		return new FComputePipelineState(ComputeShaderInitializer.ComputeShader);
+		bool bRuntimePSO = false;
+		bool bFromInMemoryPSOCache = false;
+		return new FComputePipelineState(ComputeShaderInitializer.ComputeShader, bRuntimePSO, bFromInMemoryPSOCache);
 	}
 
 	FPrecacheComputePipelineCache(uint32 InMemoryPSOsMaxNum) : TPrecachePipelineCacheBase(FPSOPrecacheRequestID::EType::Compute, InMemoryPSOsMaxNum) {}
-	FPSOPrecacheRequestResult PrecacheComputePipelineState(const FComputePipelineStateInitializer& ComputeInitializer, const TCHAR* Name, bool bForcePrecache);
+	FPSOPrecacheRequestResult PrecacheComputePipelineState(const FComputePipelineStateInitializer& ComputeInitializer, EPSOPrecacheInMemoryMode InMemoryMode, const TCHAR* Name, bool bForcePrecache);
 	static void OnNewPipelineStateCreated(const FComputePipelineStateInitializer& ComputeInitializer, FComputePipelineState* NewComputePipelineState, const FString& PSOCompilationEventName, bool bDoAsyncCompile);
  
 	static const FName GetActiveCompileStatName()
@@ -2943,7 +3202,9 @@ public:
 
 	static FGraphicsPipelineState* CreateNewPSO(const FGraphicsPipelineStateInitializer& Initializer)
 	{
-		return new FGraphicsPipelineState;
+		bool bRuntimePSO = false;
+		bool bFromInMemoryPSOCache = false;
+		return new FGraphicsPipelineState(bRuntimePSO, bFromInMemoryPSOCache);
 	}
 
 	static FORCEINLINE uint64 PipelineStateInitializerHash(const FGraphicsPipelineStateInitializer& Key)
@@ -2964,7 +3225,7 @@ public:
 		return GET_STATFNAME(STAT_HighestPriorityGraphicsPSOPrecacheRequests);
 	}
 	FPrecacheGraphicsPipelineCache(uint32 InMemoryPSOsMaxNum) : TPrecachePipelineCacheBase(FPSOPrecacheRequestID::EType::Graphics, InMemoryPSOsMaxNum) {}
-	FPSOPrecacheRequestResult PrecacheGraphicsPipelineState(const FGraphicsPipelineStateInitializer& Initializer);
+	FPSOPrecacheRequestResult PrecacheGraphicsPipelineState(const FGraphicsPipelineStateInitializer& Initializer, EPSOPrecacheInMemoryMode InMemoryMode);
 
 	static void OnNewPipelineStateCreated(const FGraphicsPipelineStateInitializer& Initializer, FGraphicsPipelineState* NewGraphicsPipelineState, const FString& PSOCompilationEventName, bool bDoAsyncCompile);
 
@@ -3812,6 +4073,7 @@ FComputePipelineState* PipelineStateCache::GetAndOrCreateComputePipelineState(FR
 	if (bWasFound == false)
 	{		
 		EPSOPrecacheResult PSOPrecacheResult = EPSOPrecacheResult::Unknown;
+		bool bFromInMemoryPSOCache = false;
 		bool bPrecachedPSOMarkedUsed = false;
 		if (IsPSOPrecachingEnabled())
 		{
@@ -3821,10 +4083,11 @@ FComputePipelineState* PipelineStateCache::GetAndOrCreateComputePipelineState(FR
 			}
 			else
 			{
-				FGraphEventRef CompletionEvent;				
-				FPSOPrecacheRequestID PSOPrecacheRequestID;
-				PSOPrecacheResult = GPrecacheComputePipelineCache->GetPrecachingState(ComputeInitializer, CompletionEvent, PSOPrecacheRequestID);
+				EPSOPrecacheInMemoryMode InMemoryMode = EPSOPrecacheInMemoryMode::Uncached;
+				PSOPrecacheResult = GPrecacheComputePipelineCache->GetPrecachingStateAndInMemoryMode(ComputeInitializer, InMemoryMode);
+				bFromInMemoryPSOCache = InMemoryMode == EPSOPrecacheInMemoryMode::Cached;
 			}
+
 		}
 
 		bool bWasPSOPrecached = PSOPrecacheResult == EPSOPrecacheResult::Active || PSOPrecacheResult == EPSOPrecacheResult::Complete;
@@ -3833,7 +4096,8 @@ FComputePipelineState* PipelineStateCache::GetAndOrCreateComputePipelineState(FR
 		FPipelineFileCacheManager::CacheComputePSO(GetTypeHash(ComputeInitializer), ComputeInitializer, bWasPSOPrecached, &Stats);
 
 		// create new compute state
-		OutCachedState = new FComputePipelineState(ComputeInitializer.ComputeShader);
+		bool bRuntimePSO = true;
+		OutCachedState = new FComputePipelineState(ComputeInitializer.ComputeShader, bRuntimePSO, bFromInMemoryPSOCache);
 		OutCachedState->Stats = Stats;
 		if (DoAsyncCompile)
 		{
@@ -4338,6 +4602,7 @@ FGraphicsPipelineState* PipelineStateCache::GetAndOrCreateGraphicsPipelineState(
 	if (bWasFound == false)
 	{
 		EPSOPrecacheResult PSOPrecacheResult = EPSOPrecacheResult::Unknown;
+		bool bFromInMemoryPSOCache = false;
 		bool bPrecachedPSOMarkedUsed = false;
 		if (IsPSOPrecachingEnabled())
 		{
@@ -4353,9 +4618,9 @@ FGraphicsPipelineState* PipelineStateCache::GetAndOrCreateGraphicsPipelineState(
 			}
 			else
 			{
-				FGraphEventRef CompletionEvent;				
-				FPSOPrecacheRequestID PSOPrecacheRequestID;
-				PSOPrecacheResult = GPrecacheGraphicsPipelineCache->GetPrecachingState(PSOPrecacheInitializer, CompletionEvent, PSOPrecacheRequestID);
+				EPSOPrecacheInMemoryMode InMemoryMode = EPSOPrecacheInMemoryMode::Uncached;
+				PSOPrecacheResult = GPrecacheGraphicsPipelineCache->GetPrecachingStateAndInMemoryMode(PSOPrecacheInitializer, InMemoryMode);
+				bFromInMemoryPSOCache = InMemoryMode == EPSOPrecacheInMemoryMode::Cached;
 			}
 		}
 
@@ -4367,7 +4632,8 @@ FGraphicsPipelineState* PipelineStateCache::GetAndOrCreateGraphicsPipelineState(
 		FPipelineFileCacheManager::CacheGraphicsPSO(GetTypeHash(Initializer), Initializer, bWasPSOPrecached, &Stats);
 
 		// create new graphics state
-		OutCachedState = new FGraphicsPipelineState();
+		bool bRuntimePSO = true;
+		OutCachedState = new FGraphicsPipelineState(bRuntimePSO, bFromInMemoryPSOCache);
 		OutCachedState->Stats = Stats;
 		if (DoAsyncCompile)
 		{
@@ -4513,7 +4779,7 @@ bool PipelineStateCache::IsPSOPrecachingEnabled()
 #endif // WITH_EDITOR
 }
 
-FPSOPrecacheRequestResult FPrecacheComputePipelineCache::PrecacheComputePipelineState(const FComputePipelineStateInitializer& ComputeInitializer, const TCHAR* Name, bool bForcePrecache)
+FPSOPrecacheRequestResult FPrecacheComputePipelineCache::PrecacheComputePipelineState(const FComputePipelineStateInitializer& ComputeInitializer, EPSOPrecacheInMemoryMode InMemoryMode, const TCHAR* Name, bool bForcePrecache)
 {
 	FPSOPrecacheRequestResult Result;
 	if (!PipelineStateCache::IsPSOPrecachingEnabled() && !bForcePrecache)
@@ -4531,7 +4797,7 @@ FPSOPrecacheRequestResult FPrecacheComputePipelineCache::PrecacheComputePipeline
 	FComputePipelineStateInitializer PrecacheComputeInitializer(ComputeInitializer);
 	PrecacheComputeInitializer.bPSOPrecache = true;
 	
-	return TryAddNewState(PrecacheComputeInitializer, Name, bDoAsyncCompile);
+	return TryAddNewState(PrecacheComputeInitializer, InMemoryMode, Name, bDoAsyncCompile);
 }
 
 void FPrecacheComputePipelineCache::OnNewPipelineStateCreated(const FComputePipelineStateInitializer& ComputeInitializer, FComputePipelineState* CachedState, const FString& PSOCompilationEventName, bool bDoAsyncCompile)
@@ -4545,12 +4811,12 @@ void FPrecacheComputePipelineCache::OnNewPipelineStateCreated(const FComputePipe
 	InternalCreateComputePipelineState(ComputeInitializer, bDoAsyncCompile, EPSOPrecacheResult::Active, false, CachedState, PSOCompilationDebugData, false);
 }
 
-FPSOPrecacheRequestResult PipelineStateCache::PrecacheComputePipelineState(const FComputePipelineStateInitializer& ComputeInitializer, const TCHAR* Name, bool bForcePrecache)
+FPSOPrecacheRequestResult PipelineStateCache::PrecacheComputePipelineState(const FComputePipelineStateInitializer& ComputeInitializer, EPSOPrecacheInMemoryMode InMemoryMode, const TCHAR* Name, bool bForcePrecache)
 {
-	return GPrecacheComputePipelineCache->PrecacheComputePipelineState(ComputeInitializer, Name, bForcePrecache);
+	return GPrecacheComputePipelineCache->PrecacheComputePipelineState(ComputeInitializer, InMemoryMode, Name, bForcePrecache);
 }
 
-FPSOPrecacheRequestResult FPrecacheGraphicsPipelineCache::PrecacheGraphicsPipelineState(const FGraphicsPipelineStateInitializer& Initializer)
+FPSOPrecacheRequestResult FPrecacheGraphicsPipelineCache::PrecacheGraphicsPipelineState(const FGraphicsPipelineStateInitializer& Initializer, EPSOPrecacheInMemoryMode InMemoryMode)
 {
 	FPSOPrecacheRequestResult Result;
 	if (!PipelineStateCache::IsPSOPrecachingEnabled())
@@ -4566,7 +4832,7 @@ FPSOPrecacheRequestResult FPrecacheGraphicsPipelineCache::PrecacheGraphicsPipeli
 	FString PSOCompilationEventName;
 
 	// try and create new graphics state
-	return TryAddNewState(Initializer, PSOCompilationEventName, bDoAsyncCompile);
+	return TryAddNewState(Initializer, InMemoryMode, PSOCompilationEventName, bDoAsyncCompile);
 }
 	
 void FPrecacheGraphicsPipelineCache::OnNewPipelineStateCreated(const FGraphicsPipelineStateInitializer & Initializer, FGraphicsPipelineState * NewGraphicsPipelineState, const FString& PSOCompilationEventName, bool bDoAsyncCompile)
@@ -4585,9 +4851,9 @@ void FPrecacheGraphicsPipelineCache::OnNewPipelineStateCreated(const FGraphicsPi
 	InternalCreateGraphicsPipelineState(InitializerCopy, EPSOPrecacheResult::Active, false, bDoAsyncCompile, NewGraphicsPipelineState, PSOCompilationDebugData, false);
 }
 
-FPSOPrecacheRequestResult PipelineStateCache::PrecacheGraphicsPipelineState(const FGraphicsPipelineStateInitializer& Initializer)
+FPSOPrecacheRequestResult PipelineStateCache::PrecacheGraphicsPipelineState(const FGraphicsPipelineStateInitializer& Initializer, EPSOPrecacheInMemoryMode InMemoryMode)
 {
-	return GPrecacheGraphicsPipelineCache->PrecacheGraphicsPipelineState(Initializer);
+	return GPrecacheGraphicsPipelineCache->PrecacheGraphicsPipelineState(Initializer, InMemoryMode);
 }
 
 EPSOPrecacheResult PipelineStateCache::CheckPipelineStateInCache(const FGraphicsPipelineStateInitializer& PipelineStateInitializer, FGraphEventRef& OutCompletionEvent, FPSOPrecacheRequestID& OutPSOPrecacheRequestID)
@@ -4620,6 +4886,25 @@ uint64 PipelineStateCache::GetPrecachePSOCompilationTime(bool bGraphics, uint64
 	return bGraphics ? GPrecacheGraphicsPipelineCache->GetPSOCompilationTime(PSOPrecacheHash) : GPrecacheComputePipelineCache->GetPSOCompilationTime(PSOPrecacheHash);
 }
 
+bool PipelineStateCache::ReleasePrecachedPSO(const FPSOPrecacheRequestID& PSOPrecacheRequestID)
+{
+	if (!IsPSOPrecachingEnabled())
+	{
+		return false;
+	}
+
+	bool bReleased = false;
+	if (PSOPrecacheRequestID.GetType() == FPSOPrecacheRequestID::EType::Graphics)
+	{
+		bReleased = GPrecacheGraphicsPipelineCache->ReleasePrecachedPSO(PSOPrecacheRequestID);
+	}
+	else
+	{
+		bReleased = GPrecacheComputePipelineCache->ReleasePrecachedPSO(PSOPrecacheRequestID);
+	}
+	return bReleased;
+}
+
 bool PipelineStateCache::IsPrecaching(const FPSOPrecacheRequestID& PSOPrecacheRequestID)
 {
 	if (!IsPSOPrecachingEnabled())
@@ -4657,7 +4942,7 @@ bool PipelineStateCache::IsPrecaching(FRHIComputeShader* ComputeShader)
 	{
 		return false;
 	}
-		
+
 	FGraphEventRef CompletionEvent;	
 	FPSOPrecacheRequestID PSOPrecacheRequestID;
 	return GPrecacheComputePipelineCache->GetPrecachingState(FComputePipelineStateInitializer(ComputeShader), CompletionEvent, PSOPrecacheRequestID) == EPSOPrecacheResult::Active;
diff --git a/Engine/Source/Runtime/RHI/Public/PipelineStateCache.h b/Engine/Source/Runtime/RHI/Public/PipelineStateCache.h
--- a/Engine/Source/Runtime/RHI/Public/PipelineStateCache.h
+++ b/Engine/Source/Runtime/RHI/Public/PipelineStateCache.h
@@ -36,6 +36,15 @@ enum class EPSOPrecachePriority : uint8
 	Highest,
 };
 
+/**
+ * PSO Precache InMemory mode
+ */
+enum class EPSOPrecacheInMemoryMode : uint8
+{
+	Cached,		//< Keep precached PSO in memory
+	Uncached,	//< Don't keep precached PSO in memory (release after async compile)
+};
+
 enum class ERayTracingPipelineCacheFlags : uint8
 {
 	// Query the pipeline cache, create pipeline if necessary.
@@ -214,10 +223,10 @@ namespace PipelineStateCache
 	extern RHI_API bool						IsPSOPrecachingEnabled();
 
 	/* Precache the compute shader and return a request ID if precached async */
-	extern RHI_API FPSOPrecacheRequestResult PrecacheComputePipelineState(const FComputePipelineStateInitializer& ComputeInitializer, const TCHAR* Name = nullptr, bool bForcePrecache = false);
+	extern RHI_API FPSOPrecacheRequestResult PrecacheComputePipelineState(const FComputePipelineStateInitializer& ComputeInitializer, EPSOPrecacheInMemoryMode InMemoryMode, const TCHAR* Name = nullptr, bool bForcePrecache = false);
 
-	/* Precache the graphic PSO and return an optional graph event if precached async */
-	extern RHI_API FPSOPrecacheRequestResult PrecacheGraphicsPipelineState(const FGraphicsPipelineStateInitializer& PipelineStateInitializer);
+	/* Precache the graphics PSO and return an optional graph event if precached async */
+	extern RHI_API FPSOPrecacheRequestResult PrecacheGraphicsPipelineState(const FGraphicsPipelineStateInitializer& PipelineStateInitializer, EPSOPrecacheInMemoryMode InMemoryMode);
 
 	/* Retrieve the current PSO precache result state (slightly slower than IsPrecaching) */
 	extern RHI_API EPSOPrecacheResult		CheckPipelineStateInCache(const FGraphicsPipelineStateInitializer& PipelineStateInitializer, FGraphEventRef& OutCompletionEvent, FPSOPrecacheRequestID& OutPSOPrecacheRequestID);	
@@ -236,7 +245,10 @@ namespace PipelineStateCache
 		FPSOPrecacheRequestID PSOPrecacheRequestID;
 		return CheckPipelineStateInCache(ComputeShader, CompletionEvent, PSOPrecacheRequestID);
 	}
-
+		
+	/* Release the precached PSOs - returns true when PSO is queued for deletion, otherwise only the refcount was decremented */
+	extern RHI_API bool						ReleasePrecachedPSO(const FPSOPrecacheRequestID& PSORequest);
+		
 	/* Retrieve the the time it took to compile this precache PSO */
 	extern RHI_API uint64					GetPrecachePSOCompilationTime(bool bGraphics, uint64 PSOPrecacheHash);
 
diff --git a/Engine/Source/Runtime/SlateRHIRenderer/Private/SlateRHIRenderingPolicy.cpp b/Engine/Source/Runtime/SlateRHIRenderer/Private/SlateRHIRenderingPolicy.cpp
--- a/Engine/Source/Runtime/SlateRHIRenderer/Private/SlateRHIRenderingPolicy.cpp
+++ b/Engine/Source/Runtime/SlateRHIRenderer/Private/SlateRHIRenderingPolicy.cpp
@@ -1845,6 +1845,7 @@ void AddSlatePSOInitializer(
 	PSOPrecacheData.bRequired = true;
 	PSOPrecacheData.Type = FPSOPrecacheData::EType::Graphics;
 	PSOPrecacheData.GraphicsPSOInitializer = GraphicsPSOInit;
+	PSOPrecacheData.InMemoryMode = EPSOPrecacheInMemoryMode::Uncached; //< Don't cache Slate PSOs in memory - we have way too many unused ones which would eat up memory (should be make this an option?)
 #if PSO_PRECACHING_VALIDATE
 	PSOPrecacheData.PSOCollectorIndex = PSOCollectorIndex;
 	PSOPrecacheData.VertexFactoryType = nullptr;
