Remove invalid check because cached flag is set while the task is running and thus the graph event might not have completed again which tracks the compilation task

#Synced-CL 49719249

diff --git a/Engine/Source/Runtime/RHI/Private/PipelineStateCache.cpp b/Engine/Source/Runtime/RHI/Private/PipelineStateCache.cpp
--- a/Engine/Source/Runtime/RHI/Private/PipelineStateCache.cpp
+++ b/Engine/Source/Runtime/RHI/Private/PipelineStateCache.cpp
@@ -2867,7 +2867,7 @@ public:
 
 	static bool IsCompilationDone(EPSOPrecacheStateMask StateMask)
 	{
-		// TODO: Is Uncached also compilation done?
+		// When marked uncached then compilation is done as well but released again so valid in the driver cache but not runtime cached
 		return EnumHasAnyFlags(StateMask, EPSOPrecacheStateMask::Cached | EPSOPrecacheStateMask::Uncached | EPSOPrecacheStateMask::Failed);
 	}
 
@@ -3034,8 +3034,6 @@ protected:
 			return EPSOPrecacheResult::Active;
 		}
 
-		check(AsyncCompileEvent == nullptr || AsyncCompileEvent->IsComplete());
-
 		// check we only set 1 completion bit
 		const EPSOPrecacheStateMask CompletionMask = EPSOPrecacheStateMask::Cached | EPSOPrecacheStateMask::Failed | EPSOPrecacheStateMask::Uncached;
 		check(EnumHasAnyFlags(CompilationState, CompletionMask) && !EnumHasAllFlags(CompilationState, CompletionMask));
