Updated CSVProfiling metadata to include the GPU and GPU driver version to assist performance testing triage.

#Synced-CL 49818450

#Synced-CL 49818450

diff --git a/Engine/Source/Runtime/RHI/Private/DynamicRHI.cpp b/Engine/Source/Runtime/RHI/Private/DynamicRHI.cpp
--- a/Engine/Source/Runtime/RHI/Private/DynamicRHI.cpp
+++ b/Engine/Source/Runtime/RHI/Private/DynamicRHI.cpp
@@ -413,6 +413,8 @@ void RHIInit(bool bHasEditorToken)
 	CSV_METADATA(TEXT("RHIName"), GMaxRHIFeatureLevel == ERHIFeatureLevel::ES3_1 ? *(FString(GDynamicRHI->GetName()) + TEXT("_ES31")) : GDynamicRHI->GetName());
 	CSV_METADATA(TEXT("RHIFeatureLevel"), *LexToString(GMaxRHIFeatureLevel));
 	CSV_METADATA(TEXT("ShaderPlatform"), *LexToString(GetFeatureLevelShaderPlatform(GMaxRHIFeatureLevel)));
+	CSV_METADATA(TEXT("GPU"), *GRHIGlobals.GpuInfo.AdapterName.TrimStartAndEnd());
+	CSV_METADATA(TEXT("GPUDriver"), *GRHIGlobals.GpuInfo.AdapterUserDriverVersion.TrimStartAndEnd());
 }
 
 void RHIPostInit(const TArray<uint32>& InPixelFormatByteWidth)
