Fixes breadcrumb events to work properly when pausing trace or disabling the cpu channel.
#Synced-CL 49071325

diff --git a/Engine/Source/Runtime/RHI/Public/RHIBreadcrumbs.h b/Engine/Source/Runtime/RHI/Public/RHIBreadcrumbs.h
--- a/Engine/Source/Runtime/RHI/Public/RHIBreadcrumbs.h
+++ b/Engine/Source/Runtime/RHI/Public/RHIBreadcrumbs.h
@@ -1112,7 +1112,7 @@ namespace UE::RHI::GPUProfiler
 			uint32 GetTraceCpuSpec() const
 			{
 #if CPUPROFILERTRACE_ENABLED
-				if (TraceCpuSpecId == 0 && UE::RHI::GPUProfiler::FGpuProfilerTrace::IsAvailable())
+				if (TraceCpuSpecId == 0 && bool(CpuChannel))
 				{
 					FScopeLock Lock(&Cs);
 
@@ -1168,14 +1168,17 @@ namespace UE::RHI::GPUProfiler
 				, Desc(Desc)
 			{
 #if RHI_BREADCRUMBS_EMIT_CPU
-				TraceCpuSpecId = Desc.GetTraceCpuSpec();
-				if (TraceCpuSpecId)
+				if (bool(CpuChannel))
 				{
-					if (Desc.NumValues > 0)
+					TraceCpuSpecId = Desc.GetTraceCpuSpec();
+					if (TraceCpuSpecId)
 					{
-						UE::RHI::GPUProfiler::FMetadataSerializer Serializer;
-						Desc.SerializeValues(Serializer, *this);
-						TraceCpuMetadataId = FCpuProfilerTrace::OutputMetadata(TraceCpuSpecId, Serializer.GetData());
+						if (Desc.NumValues > 0)
+						{
+							UE::RHI::GPUProfiler::FMetadataSerializer Serializer;
+							Desc.SerializeValues(Serializer, *this);
+							TraceCpuMetadataId = FCpuProfilerTrace::OutputMetadata(TraceCpuSpecId, Serializer.GetData());
+						}
 					}
 				}
 #endif
