Allow disallowing some GPUs from using D3D11 SM5

diff --git a/Engine/Source/Runtime/RHI/Private/Windows/WindowsDynamicRHI.cpp b/Engine/Source/Runtime/RHI/Private/Windows/WindowsDynamicRHI.cpp
--- a/Engine/Source/Runtime/RHI/Private/Windows/WindowsDynamicRHI.cpp
+++ b/Engine/Source/Runtime/RHI/Private/Windows/WindowsDynamicRHI.cpp
@@ -384,80 +384,6 @@
 	return DefaultRHI;
 }
 
-static TArray<EShaderPlatform> ParseShaderPlatformsConfig(const TCHAR* InSettingName)
-{
-	TArray<FString> TargetedShaderFormats;
-	GConfig->GetArray(TEXT("/Script/WindowsTargetPlatform.WindowsTargetSettings"), InSettingName, TargetedShaderFormats, GEngineIni);
-
-	TArray<EShaderPlatform> ShaderPlatforms;
-	ShaderPlatforms.Reserve(TargetedShaderFormats.Num());
-
-	for (const FString& ShaderFormat : TargetedShaderFormats)
-	{
-		ShaderPlatforms.AddUnique(ShaderFormatToLegacyShaderPlatform(FName(*ShaderFormat)));
-	}
-
-	ShaderPlatforms.RemoveAll(
-		[](EShaderPlatform ShaderPlatform)
-		{
-			return GIsEditor && (FDataDrivenShaderPlatformInfo::GetMaxFeatureLevel(ShaderPlatform) < ERHIFeatureLevel::SM5);
-		});
-
-	return ShaderPlatforms;
-}
-
-static TArray<ERHIFeatureLevel::Type> FeatureLevelsFromShaderPlatforms(const TArray<EShaderPlatform>& InShaderPlatforms)
-{
-	TArray<ERHIFeatureLevel::Type> FeatureLevels;
-	FeatureLevels.Reserve(InShaderPlatforms.Num());
-
-	for (EShaderPlatform ShaderPlatform : InShaderPlatforms)
-	{
-		FeatureLevels.AddUnique(FDataDrivenShaderPlatformInfo::GetMaxFeatureLevel(ShaderPlatform));
-	}
-
-	return FeatureLevels;
-}
-
-static FWindowsRHIConfig ParseWindowsRHIConfig(const TCHAR* ShaderFormatsSettingName)
-{
-	FWindowsRHIConfig Config;
-
-	Config.ShaderPlatforms = ParseShaderPlatformsConfig(ShaderFormatsSettingName);
-	Config.FeatureLevels = FeatureLevelsFromShaderPlatforms(Config.ShaderPlatforms);
-
-	return Config;
-}
-
-FParsedWindowsDynamicRHIConfig ParseWindowsDynamicRHIConfig()
-{
-	FParsedWindowsDynamicRHIConfig Config;
-
-	Config.DefaultRHI = ParseDefaultWindowsRHI();
-
-	Config.RHIConfigs[(int32)EWindowsRHI::D3D11]  = ParseWindowsRHIConfig(TEXT("D3D11TargetedShaderFormats"));
-	Config.RHIConfigs[(int32)EWindowsRHI::D3D12]  = ParseWindowsRHIConfig(TEXT("D3D12TargetedShaderFormats"));
-	Config.RHIConfigs[(int32)EWindowsRHI::Vulkan] = ParseWindowsRHIConfig(TEXT("VulkanTargetedShaderFormats"));
-
-	// Only add OpenGL support to non-client programs.
-	if (!FPlatformProperties::RequiresCookedData())
-	{
-		Config.RHIConfigs[(int32)EWindowsRHI::OpenGL].ShaderPlatforms.Add(SP_OPENGL_PCES3_1);
-		Config.RHIConfigs[(int32)EWindowsRHI::OpenGL].FeatureLevels.Add(ERHIFeatureLevel::ES3_1);
-	}
-
-	if (FWindowsRHIConfig DeprecatedConfig = ParseWindowsRHIConfig(TEXT("TargetedRHIs")); !DeprecatedConfig.ShaderPlatforms.IsEmpty())
-	{
-		// Since we don't have context here, we have to add this old config setting to every potential RHI.
-		Config.MergeConfig(EWindowsRHI::D3D11,  DeprecatedConfig);
-		Config.MergeConfig(EWindowsRHI::D3D12,  DeprecatedConfig);
-		Config.MergeConfig(EWindowsRHI::Vulkan, DeprecatedConfig);
-		Config.MergeConfig(EWindowsRHI::OpenGL, DeprecatedConfig);
-	}
-
-	return Config;
-}
-
 struct FWindowsGPUInfo
 {
 	uint32 VendorId = 0;
@@ -499,7 +425,7 @@
 	}
 }
 
-static bool IsGPUInHardwareListForD3D11_ES31(const FWindowsGPUInfo& BestGPUInfo, const TCHAR* const ConfigKey)
+static bool IsGPUInHardwareListForRHI(const FWindowsGPUInfo& BestGPUInfo, const TCHAR* const ConfigKey, const TCHAR* const MatchingRHIName)
 {
 	TArray<FString> DeviceRHIList;
 	GConfig->GetArray(TEXT("Devices"), ConfigKey, DeviceRHIList, GHardwareIni);
@@ -517,7 +443,7 @@
 		FString DeviceName;
 		FParse::Value(Line + 1, TEXT("DeviceName="), DeviceName);
 
-		if (RHIName.Compare("D3D11_ES31", ESearchCase::IgnoreCase) == 0 && GPUBrand.Compare(DeviceName, ESearchCase::IgnoreCase) == 0)
+		if (RHIName.Compare(MatchingRHIName, ESearchCase::IgnoreCase) == 0 && GPUBrand.Compare(DeviceName, ESearchCase::IgnoreCase) == 0)
 		{
 			UE_LOG(LogRHI, Log, TEXT("%s: Found RHIName %s with DeviceName %s"), ConfigKey, *RHIName, *DeviceName);
 			return true;
@@ -528,21 +454,143 @@
 		uint32 VendorIdInt = FParse::HexNumber(*VendorId);
 
 		FString DeviceId;
-		FParse::Value(Line + 1, TEXT("DeviceId="), DeviceId);
+		const bool bSpecificDevice = FParse::Value(Line + 1, TEXT("DeviceId="), DeviceId);
 		uint32 DeviceIdInt = FParse::HexNumber(*DeviceId);
 
-		if (BestGPUInfo.VendorId && BestGPUInfo.DeviceId &&
-			BestGPUInfo.VendorId == VendorIdInt && BestGPUInfo.DeviceId == DeviceIdInt &&
-			RHIName.Compare("D3D11_ES31", ESearchCase::IgnoreCase) == 0)
+		if (BestGPUInfo.VendorId && (BestGPUInfo.VendorId == VendorIdInt))
 		{
-			UE_LOG(LogRHI, Log, TEXT("%s: Found RHIName %s with DeviceName %s, VendorId %s, and DeviceId %s"), ConfigKey, *RHIName, *DeviceName, *VendorId, *DeviceId);
-			return true;
+			if (!bSpecificDevice || (BestGPUInfo.DeviceId && (BestGPUInfo.DeviceId == DeviceIdInt)))
+			{
+				if (RHIName.Compare(MatchingRHIName, ESearchCase::IgnoreCase) == 0)
+				{
+					UE_LOG(LogRHI, Log, TEXT("%s: Found RHIName %s with DeviceName %s, VendorId %s, and DeviceId %s"), ConfigKey, *RHIName, *DeviceName, *VendorId, *DeviceId);
+					return true;
+				}
+			}
 		}
 	}
 
 	return false;
 }
 
+static bool ShouldAvoidD3D11_SM5()
+{
+	bool bAvoidD3D11SM5 = false;
+
+	static TOptional<bool> AvoidD3D11SM5;
+	if (AvoidD3D11SM5.IsSet())
+	{
+		bAvoidD3D11SM5 = AvoidD3D11SM5.GetValue();
+	}
+	else
+	{
+		FWindowsGPUInfo BestGPUInfo;
+		GetGPUInfo(BestGPUInfo);
+
+		bAvoidD3D11SM5 = IsGPUInHardwareListForRHI(BestGPUInfo, TEXT("DeviceAvoidRHIList"), TEXT("D3D11_SM5"));
+		if (bAvoidD3D11SM5)
+		{
+			UE_LOG(LogRHI, Log, TEXT("Avoiding feature level D3D11 %s"), GetLogName(ERHIFeatureLevel::SM5));
+		}
+
+		AvoidD3D11SM5 = bAvoidD3D11SM5;
+	}
+
+	return bAvoidD3D11SM5;
+}
+
+static TArray<EShaderPlatform> ParseShaderPlatformsConfig(const TCHAR* InSettingName)
+{
+	TArray<FString> TargetedShaderFormats;
+	GConfig->GetArray(TEXT("/Script/WindowsTargetPlatform.WindowsTargetSettings"), InSettingName, TargetedShaderFormats, GEngineIni);
+
+	TArray<EShaderPlatform> ShaderPlatforms;
+	ShaderPlatforms.Reserve(TargetedShaderFormats.Num());
+
+	for (const FString& ShaderFormat : TargetedShaderFormats)
+	{
+		ShaderPlatforms.AddUnique(ShaderFormatToLegacyShaderPlatform(FName(*ShaderFormat)));
+	}
+
+	if (GIsEditor)
+	{
+		ShaderPlatforms.RemoveAll(
+			[](EShaderPlatform ShaderPlatform)
+			{
+				return FDataDrivenShaderPlatformInfo::GetMaxFeatureLevel(ShaderPlatform) < ERHIFeatureLevel::SM5;
+			});
+	}
+	else if (FCString::Strcmp(InSettingName, TEXT("D3D11TargetedShaderFormats")) == 0)
+	{
+		bool bUseSM5AvoidGPUList = false;
+		if (GConfig->GetBool(TEXT("D3D11_SM5"), TEXT("bUseSM5AvoidGPUList"), bUseSM5AvoidGPUList, GEngineIni))
+		{
+			if (bUseSM5AvoidGPUList && ShouldAvoidD3D11_SM5())
+			{
+				ShaderPlatforms.RemoveAll(
+					[](EShaderPlatform ShaderPlatform)
+					{
+						return FDataDrivenShaderPlatformInfo::GetMaxFeatureLevel(ShaderPlatform) == ERHIFeatureLevel::SM5;
+					});
+			}
+		}
+	}
+
+	return ShaderPlatforms;
+}
+
+static TArray<ERHIFeatureLevel::Type> FeatureLevelsFromShaderPlatforms(const TArray<EShaderPlatform>& InShaderPlatforms)
+{
+	TArray<ERHIFeatureLevel::Type> FeatureLevels;
+	FeatureLevels.Reserve(InShaderPlatforms.Num());
+
+	for (EShaderPlatform ShaderPlatform : InShaderPlatforms)
+	{
+		FeatureLevels.AddUnique(FDataDrivenShaderPlatformInfo::GetMaxFeatureLevel(ShaderPlatform));
+	}
+
+	return FeatureLevels;
+}
+
+static FWindowsRHIConfig ParseWindowsRHIConfig(const TCHAR* ShaderFormatsSettingName)
+{
+	FWindowsRHIConfig Config;
+
+	Config.ShaderPlatforms = ParseShaderPlatformsConfig(ShaderFormatsSettingName);
+	Config.FeatureLevels = FeatureLevelsFromShaderPlatforms(Config.ShaderPlatforms);
+
+	return Config;
+}
+
+FParsedWindowsDynamicRHIConfig ParseWindowsDynamicRHIConfig()
+{
+	FParsedWindowsDynamicRHIConfig Config;
+
+	Config.DefaultRHI = ParseDefaultWindowsRHI();
+
+	Config.RHIConfigs[(int32)EWindowsRHI::D3D11]  = ParseWindowsRHIConfig(TEXT("D3D11TargetedShaderFormats"));
+	Config.RHIConfigs[(int32)EWindowsRHI::D3D12]  = ParseWindowsRHIConfig(TEXT("D3D12TargetedShaderFormats"));
+	Config.RHIConfigs[(int32)EWindowsRHI::Vulkan] = ParseWindowsRHIConfig(TEXT("VulkanTargetedShaderFormats"));
+
+	// Only add OpenGL support to non-client programs.
+	if (!FPlatformProperties::RequiresCookedData())
+	{
+		Config.RHIConfigs[(int32)EWindowsRHI::OpenGL].ShaderPlatforms.Add(SP_OPENGL_PCES3_1);
+		Config.RHIConfigs[(int32)EWindowsRHI::OpenGL].FeatureLevels.Add(ERHIFeatureLevel::ES3_1);
+	}
+
+	if (FWindowsRHIConfig DeprecatedConfig = ParseWindowsRHIConfig(TEXT("TargetedRHIs")); !DeprecatedConfig.ShaderPlatforms.IsEmpty())
+	{
+		// Since we don't have context here, we have to add this old config setting to every potential RHI.
+		Config.MergeConfig(EWindowsRHI::D3D11,  DeprecatedConfig);
+		Config.MergeConfig(EWindowsRHI::D3D12,  DeprecatedConfig);
+		Config.MergeConfig(EWindowsRHI::Vulkan, DeprecatedConfig);
+		Config.MergeConfig(EWindowsRHI::OpenGL, DeprecatedConfig);
+	}
+
+	return Config;
+}
+
 static bool IsForcedToD3D11_ES31()
 {
 	bool bRequireD3D11ES31 = false;
@@ -557,7 +605,7 @@
 		FWindowsGPUInfo BestGPUInfo;
 		GetGPUInfo(BestGPUInfo);
 
-		bRequireD3D11ES31 = IsGPUInHardwareListForD3D11_ES31(BestGPUInfo, TEXT("DeviceForceRHIList"));
+		bRequireD3D11ES31 = IsGPUInHardwareListForRHI(BestGPUInfo, TEXT("DeviceForceRHIList"), TEXT("D3D11_ES31"));
 		if (bRequireD3D11ES31)
 		{
 			UE_LOG(LogRHI, Log, TEXT("Forced feature level to D3D11 %s"), GetLogName(ERHIFeatureLevel::ES3_1));
@@ -584,7 +632,7 @@
 	FWindowsGPUInfo BestGPUInfo;
 	GetGPUInfo(BestGPUInfo);
 
-	if (IsGPUInHardwareListForD3D11_ES31(BestGPUInfo, TEXT("DeviceDefaultRHIList")))
+	if (IsGPUInHardwareListForRHI(BestGPUInfo, TEXT("DeviceDefaultRHIList"), TEXT("D3D11_ES31")))
 	{
 		UE_LOG(LogRHI, Log, TEXT("Default feature level to %s"), GetLogName(ERHIFeatureLevel::ES3_1));
 
