Skip D3D11 SM5 block if any related command-line arguments are specified

diff --git a/Engine/Source/Runtime/RHI/Private/Windows/WindowsDynamicRHI.cpp b/Engine/Source/Runtime/RHI/Private/Windows/WindowsDynamicRHI.cpp
--- a/Engine/Source/Runtime/RHI/Private/Windows/WindowsDynamicRHI.cpp
+++ b/Engine/Source/Runtime/RHI/Private/Windows/WindowsDynamicRHI.cpp
@@ -482,7 +482,9 @@ static bool ShouldAvoidD3D11_SM5()
 	{
 		bAvoidD3D11SM5 = AvoidD3D11SM5.GetValue();
 	}
-	else
+	else if (!FParse::Param(FCommandLine::Get(), TEXT("d3d11")) &&
+		!FParse::Param(FCommandLine::Get(), TEXT("dx11")) &&
+		!FParse::Param(FCommandLine::Get(), TEXT("sm5")))
 	{
 		FWindowsGPUInfo BestGPUInfo;
 		GetGPUInfo(BestGPUInfo);
