set r.PSOPrecache.KeepInMemoryUntilUsed=2 on windows platforms only.
- reclaims some memory on QC based Android systems.

diff --git a/Engine/Config/Windows/BaseWindowsEngine.ini b/Engine/Config/Windows/BaseWindowsEngine.ini
--- a/Engine/Config/Windows/BaseWindowsEngine.ini
+++ b/Engine/Config/Windows/BaseWindowsEngine.ini
@@ -68,3 +68,4 @@ MaxNumThreadsWithTlsSlots=512
 
 [ConsoleVariables]
 WindowsApplication.UseWorkerThreadForRawInput=False
+r.PSOPrecache.KeepInMemoryUntilUsed=2
diff --git a/Engine/Source/Runtime/RHI/Private/PipelineStateCache.cpp b/Engine/Source/Runtime/RHI/Private/PipelineStateCache.cpp
--- a/Engine/Source/Runtime/RHI/Private/PipelineStateCache.cpp
+++ b/Engine/Source/Runtime/RHI/Private/PipelineStateCache.cpp
@@ -401,7 +401,7 @@ static FAutoConsoleVariableRef CVarPrecompileThreadPoolThreadPriority(
 	TEXT("Thread priority for the PSO precompile pool"),
 	ECVF_RenderThreadSafe);
 
-int32 GPSOPrecacheKeepInMemoryUntilUsed = 2;
+int32 GPSOPrecacheKeepInMemoryUntilUsed = 0;
 static FAutoConsoleVariableRef CVarPSOPrecacheKeepInMemoryUntilUsed(
 	TEXT("r.PSOPrecache.KeepInMemoryUntilUsed"),
 	GPSOPrecacheKeepInMemoryUntilUsed,
