More compact log formatting option for ProfileGPU
* Added r.ProfileGPU.TableFormatting, which can be disabled to get indentation based formatting that's easier to read
Fixed r.profilegpu.thresholdpercent being relative to the current root, it's supposed to be relative to the total frame time
Fixed 'Other Children' node which shows hidden exclusive time

#Synced-CL 50418018

diff --git a/Engine/Source/Runtime/RHI/Private/GPUProfiler.cpp b/Engine/Source/Runtime/RHI/Private/GPUProfiler.cpp
--- a/Engine/Source/Runtime/RHI/Private/GPUProfiler.cpp
+++ b/Engine/Source/Runtime/RHI/Private/GPUProfiler.cpp
@@ -54,6 +54,12 @@ static TAutoConsoleVariable<float> GCVarProfileGPU_ThresholdPercent(
 	TEXT("Percent of the total execution duration the event needs to be larger than to be printed."),
 	ECVF_Default);
 
+static TAutoConsoleVariable<bool> GCVarProfileGPU_TableFormatting(
+	TEXT("r.ProfileGPU.TableFormatting"),
+	true,
+	TEXT("When enabled, the output results will be formatted in a table with many secondary stats. When disabled, only inclusive times and event names are printed in an indented list for compactness."),
+	ECVF_Default);
+
 static TAutoConsoleVariable<bool> GCVarProfileGPU_UnicodeOutput(
 	TEXT("r.ProfileGPU.UnicodeOutput"),
 	true,
@@ -291,8 +297,6 @@ namespace UE::RHI::GPUProfiler
 	struct FTable
 	{
 		bool const bUnicodeOutput;
-		bool const bShowHeader;
-		bool const bShowEmptyQueues;
 		bool const bShowStats;
 		bool const bShowPercent;
 		bool const bShowInclusive;
@@ -300,8 +304,6 @@ namespace UE::RHI::GPUProfiler
 
 		FTable()
 			: bUnicodeOutput  (GCVarProfileGPU_UnicodeOutput    .GetValueOnAnyThread())
-			, bShowHeader     (GCVarProfileGPU_ShowHeader       .GetValueOnAnyThread())
-			, bShowEmptyQueues(GCVarProfileGPU_ShowEmptyQueues  .GetValueOnAnyThread())
 			, bShowStats      (GCVarProfileGPU_ShowStats        .GetValueOnAnyThread())
 			, bShowPercent    (GCVarProfileGPU_ShowPercentColumn.GetValueOnAnyThread())
 			, bShowInclusive  (GCVarProfileGPU_ShowInclusive    .GetValueOnAnyThread())
@@ -1200,8 +1202,6 @@ namespace UE::RHI::GPUProfiler
 
 				void LogTree(FQueueState const& QueueState, uint32 FrameNumber) const
 				{
-					FTable Table;
-
 					EGPUProfileSortMode SortMode = (EGPUProfileSortMode)FMath::Clamp(GCVarProfileGPU_Sort.GetValueOnAnyThread(), 0, ((int32)EGPUProfileSortMode::Max - 1));
 					FWildcardString RootWildcard(GCVarProfileGPU_Root.GetValueOnAnyThread());
 					const bool bShowEmptyNodes = GCVarProfileGPU_ShowLeafEvents.GetValueOnAnyThread();
@@ -1224,10 +1224,17 @@ namespace UE::RHI::GPUProfiler
 						}
 					}
 
+					bool bHasRows = false;
+					const bool bTableFormatting = GCVarProfileGPU_TableFormatting.GetValueOnAnyThread() != 0;
+
+					FTable Table;
+					FString NodeOutput;
+					double FrameInclusiveTime = First ? First->Inclusive.GetBusyMilliseconds() : 0.0;
+
 					auto Recurse = [&](auto& Recurse, FNode* Root, FNode* CurrentNode, bool bParentMatchedFilter, int32 Level) -> bool
 					{
 						// Percent that this node was of the total frame time
-						const double Percent = Root ? (CurrentNode->Inclusive.GetBusyMilliseconds() / Root->Inclusive.GetBusyMilliseconds()) * 100.0 : 100.0;
+						const double Percent = (CurrentNode->Inclusive.GetBusyMilliseconds() / FrameInclusiveTime) * 100.0;
 
 						// Filter nodes according to cvar settings
 						const bool bAboveThreshold = Percent >= PercentThreshold;
@@ -1235,7 +1242,7 @@ namespace UE::RHI::GPUProfiler
 						const bool bHasWork = bShowEmptyNodes || CurrentNode->Inclusive.HasWork();
 
 						const bool bDisplayEvent = bNameMatches && bHasWork && bAboveThreshold;
-			
+
 						if (bDisplayEvent)
 						{
 							if (Root == nullptr)
@@ -1243,13 +1250,23 @@ namespace UE::RHI::GPUProfiler
 								Root = CurrentNode;
 							}
 
-							Table.AddRow(
-								Root,
-								CurrentNode->Inclusive,
-								CurrentNode->Exclusive,
-								CurrentNode->Name,
-								Level
-							);
+							if (bTableFormatting)
+							{
+								Table.AddRow(
+									Root,
+									CurrentNode->Inclusive,
+									CurrentNode->Exclusive,
+									CurrentNode->Name,
+									Level);
+							}
+							else
+							{
+								NodeOutput += FString::Printf(TEXT("%s%4.1f%% %5.2fms   %s\n"),
+									*FString().LeftPad(Level * 3),
+									Percent,
+									CurrentNode->Inclusive.GetBusyMilliseconds(),
+									*CurrentNode->Name);
+							}
 						}
 
 						FNode::FStats OtherChildrenInclusive;
@@ -1270,19 +1287,32 @@ namespace UE::RHI::GPUProfiler
 
 						if (bDisplayEvent && NumHiddenChildren > 0)
 						{
-							// Don't show the "other children" node if their total inclusive time is still below the percent threshold
-							if ((double(OtherChildrenInclusive.BusyCycles) / Root->Inclusive.BusyCycles) >= PercentThreshold)
+							const double OtherChildrenInclusiveMs = FPlatformTime::ToMilliseconds64(OtherChildrenInclusive.BusyCycles);
+							const double OtherChildrenPercent = (OtherChildrenInclusiveMs / FrameInclusiveTime) * 100.0;
+							// Show the "other children" node if their total inclusive time is above the threshold
+							if (OtherChildrenPercent >= PercentThreshold)
 							{
-								Table.AddRow(
-									Root,
-									OtherChildrenInclusive,
-									OtherChildrenExclusive,
-									FString::Printf(TEXT("%d Other %s"), NumHiddenChildren, NumHiddenChildren >= 2 ? TEXT("Children") : TEXT("Child")),
-									Level + 1
-								);
+								if (bTableFormatting)
+								{
+									Table.AddRow(
+										Root,
+										OtherChildrenInclusive,
+										OtherChildrenExclusive,
+										FString::Printf(TEXT("%d Other %s"), NumHiddenChildren, NumHiddenChildren >= 2 ? TEXT("Children") : TEXT("Child")),
+										Level + 1
+									);
+								}
+								else
+								{
+									NodeOutput += FString::Printf(TEXT("%s%4.1f%% %5.2fms   %d Other Children\n"),
+										*FString().LeftPad(Level * 3),
+										OtherChildrenPercent,
+										OtherChildrenInclusiveMs,
+										NumHiddenChildren);
+								}
 							}
 						}
-			
+
 						return bDisplayEvent;
 					};
 
@@ -1293,7 +1323,7 @@ namespace UE::RHI::GPUProfiler
 					}
 
 					FString Header;
-					if (Table.bShowHeader)
+					if (GCVarProfileGPU_ShowHeader.GetValueOnAnyThread() != 0)
 					{
 						Header = FString::Printf(
 							TEXT("    - %-30s: %.2fms\n")
@@ -1319,6 +1349,18 @@ namespace UE::RHI::GPUProfiler
 						, QueueState.Queue.GPU
 					);
 
+					bool bShowQueue = false;
+
+					if (bTableFormatting)
+					{
+						NodeOutput = Table.HasRows() ? *Table.ToString() : TEXT("    No recorded work for this queue.\n");
+						bShowQueue = Table.HasRows()  || GCVarProfileGPU_ShowEmptyQueues.GetValueOnAnyThread() != 0;
+					}
+					else
+					{
+						bShowQueue = NodeOutput.Len() > 0;
+					}
+
 					FString Final = FString::Printf(
 						TEXT("\n")
 						TEXT("GPU Profile for %s\n")
@@ -1327,10 +1369,10 @@ namespace UE::RHI::GPUProfiler
 						TEXT("%s")
 						, *GPUString
 						, *Header
-						, Table.HasRows() ? *Table.ToString() : TEXT("    No recorded work for this queue.\n")
+						, *NodeOutput
 					);
 
-					if (Table.HasRows() || Table.bShowEmptyQueues)
+					if (bShowQueue)
 					{
 						TArray<FString> Lines;
 						Final.ParseIntoArrayLines(Lines, false);
@@ -1343,7 +1385,7 @@ namespace UE::RHI::GPUProfiler
 
 				#if !UE_BUILD_SHIPPING
 					// Create and display profile visualizer data for the graphics queue
-					if (GCVarProfileGPU_ShowUI.GetValueOnAnyThread() && QueueState.Queue.Type == FQueue::EType::Graphics && Table.HasRows())
+					if (GCVarProfileGPU_ShowUI.GetValueOnAnyThread() && QueueState.Queue.Type == FQueue::EType::Graphics && bShowQueue)
 					{
 						// Count the total number of exclusive cycles in the frame. Needed to draw the horizontal bars in the visualizer.
 						uint64 TotalBusyCycles = 0;
