Change of we correlate FrameCounter to PresentCounter in order to have more robust Input Latency computation on all supported Platforms

#Synced-CL 50477113
#Synced-CL 50477113

diff --git a/Engine/Source/Runtime/D3D12RHI/Private/D3D12Viewport.cpp b/Engine/Source/Runtime/D3D12RHI/Private/D3D12Viewport.cpp
--- a/Engine/Source/Runtime/D3D12RHI/Private/D3D12Viewport.cpp
+++ b/Engine/Source/Runtime/D3D12RHI/Private/D3D12Viewport.cpp
@@ -595,7 +595,7 @@ bool FD3D12Viewport::SchedulePresent(FD3D12CommandContextBase* BaseContext, cons
 #if CPUPROFILERTRACE_ENABLED
 	UE_TRACE_LOG_SCOPED_T(Cpu, SchedulePresent, CpuChannel)
 		<< SchedulePresent.BackBuffer(CurrentBackBufferIndex_RHIThread)
-		<< SchedulePresent.PresentCounter(GRHIPresentCounter+1)
+		<< SchedulePresent.PresentCounter(GRHIPresentCounter)
 		<< SchedulePresent.FrameCounter(InPresentArgs.FrameCounter)
 		;
 #endif // CPUPROFILERTRACE_ENABLED
@@ -627,7 +627,7 @@ void FD3D12Viewport::PresentOnSubmissionThread(FD3D12PresentArgs& PresentArgs)
 #if CPUPROFILERTRACE_ENABLED
 	UE_TRACE_LOG_SCOPED_T(Cpu, Present_SubmissionThread, CpuChannel)
 		<< Present_SubmissionThread.BackBuffer(CurrentBackBufferIndex_RHIThread)
-		<< Present_SubmissionThread.PresentCounter(GRHIPresentCounter+1)
+		<< Present_SubmissionThread.PresentCounter(GRHIPresentCounter)
 		<< Present_SubmissionThread.FrameCounter(PresentArgs.FrameCounter)
 		;
 #endif // CPUPROFILERTRACE_ENABLED
@@ -699,13 +699,6 @@ void FD3D12Viewport::WaitForLastPresent() const
 
 bool FD3D12Viewport::EndRendering(FD3D12CommandContextBase& ContextBase, const FRHIPresentArgs& InPresentArgs)
 {
-#if CPUPROFILERTRACE_ENABLED
-	UE_TRACE_LOG_SCOPED_T(Cpu, End_Rendering, CpuChannel)
-		<< End_Rendering.BackBuffer(CurrentBackBufferIndex_RHIThread)
-		<< End_Rendering.PresentCounter(GRHIPresentCounter+1)
-		<< End_Rendering.FrameCounter(InPresentArgs.FrameCounter);
-#endif // CPUPROFILERTRACE_ENABLED
-	
 	if (!IsPresentAllowed() || !InPresentArgs.bPresent )
 	{
 		bLastFramePresented = false;
@@ -714,6 +707,13 @@ bool FD3D12Viewport::EndRendering(FD3D12CommandContextBase& ContextBase, const F
 	
 	bLastFramePresented = true;
 
+#if CPUPROFILERTRACE_ENABLED
+	UE_TRACE_LOG_SCOPED_T(Cpu, End_Rendering, CpuChannel)
+		<< End_Rendering.BackBuffer(CurrentBackBufferIndex_RHIThread)
+		<< End_Rendering.PresentCounter(GRHIPresentCounter)
+		<< End_Rendering.FrameCounter(InPresentArgs.FrameCounter);
+#endif // CPUPROFILERTRACE_ENABLED
+
 	check(ContextBase.GetParentAdapter() == GetParentAdapter());
 
 	const bool bCustomPresent = IsValidRef(CustomPresent);
diff --git a/Engine/Source/Runtime/D3D12RHI/Private/Windows/WindowsD3D12Viewport.cpp b/Engine/Source/Runtime/D3D12RHI/Private/Windows/WindowsD3D12Viewport.cpp
--- a/Engine/Source/Runtime/D3D12RHI/Private/Windows/WindowsD3D12Viewport.cpp
+++ b/Engine/Source/Runtime/D3D12RHI/Private/Windows/WindowsD3D12Viewport.cpp
@@ -364,7 +364,7 @@ HRESULT FD3D12Viewport::PresentInternal(int32 SyncInterval)
 	{
 		// Ignore time spent waiting in Present. This function blocks based on GPU progress and space in the swap chain.
 		FRenderThreadIdleScope Scope(ERenderThreadIdleTypes::WaitingForGPUPresent);
-#if !UE_BUILD_SHIPPING && PLATFORM_SUPPORTS_FLIP_TRACKING
+#if UE_INPUT_LATENCY_TRACKING
 		static FRHIFlipDetails LastFlipFrame;
 		static DXGI_FRAME_STATISTICS LastStats = { 0 };
 		HRESULT GetStatHR;
@@ -375,7 +375,7 @@ HRESULT FD3D12Viewport::PresentInternal(int32 SyncInterval)
 			NewFlipFrame.PresentIndex = stats.PresentCount;
 			NewFlipFrame.VBlankTimeInCycles = stats.SyncQPCTime.QuadPart;
 
-			RHISetVsyncDebugInfo(NewFlipFrame);
+			RHIFrameInfoSetVsyncDetails(NewFlipFrame);
 
 			LastFlipFrame = NewFlipFrame;
 			LastStats = stats;
@@ -391,7 +391,8 @@ HRESULT FD3D12Viewport::PresentInternal(int32 SyncInterval)
 		UINT PresentID;
 		if (SUCCEEDED(SwapChain1->GetLastPresentCount(&PresentID)))
 		{
-			GRHIPresentCounter = PresentID;
+			// Assuming Next PresentID == Current PresentID + 1 , this might be not true in some cases
+			GRHIPresentCounter = PresentID + 1;
 		}
 		else
 		{
diff --git a/Engine/Source/Runtime/Engine/Private/UnrealEngine.cpp b/Engine/Source/Runtime/Engine/Private/UnrealEngine.cpp
--- a/Engine/Source/Runtime/Engine/Private/UnrealEngine.cpp
+++ b/Engine/Source/Runtime/Engine/Private/UnrealEngine.cpp
@@ -2841,6 +2841,8 @@ void UEngine::SetInputSampleLatencyMarker(uint64 FrameNumber)
 {
 	TArray<ILatencyMarkerModule*> LatencyMarkerModules = IModularFeatures::Get().GetModularFeatureImplementations<ILatencyMarkerModule>(ILatencyMarkerModule::GetModularFeatureName());
 
+	RHIFrameInfoSetInputSampleMarker(FrameNumber);
+
 	for (ILatencyMarkerModule* LatencyMarkerModule : LatencyMarkerModules)
 	{
 		LatencyMarkerModule->SetInputSampleLatencyMarker(FrameNumber);
diff --git a/Engine/Source/Runtime/RHI/Private/RHICommandList.cpp b/Engine/Source/Runtime/RHI/Private/RHICommandList.cpp
--- a/Engine/Source/Runtime/RHI/Private/RHICommandList.cpp
+++ b/Engine/Source/Runtime/RHI/Private/RHICommandList.cpp
@@ -15,6 +15,7 @@
 #include "RHIContext.h"
 #include "RHIFwd.h"
 #include "RHITransition.h"
+#include "RHIUtilities.h"
 #include "Stats/StatsSystem.h"
 #include "Stats/StatsTrace.h"
 #include "Stats/ThreadIdleStats.h"
@@ -1862,6 +1863,10 @@ void FRHICommandListImmediate::UnStallRHIThread()
 void FRHICommandListImmediate::EndDrawingViewport(FRHIViewport* Viewport, const FRHIPresentArgs& InPresentArgs)
 {
 	RHI_BREADCRUMB_EVENT(*this, "EndDrawingViewport");
+	
+	// Not all platforms have control on the Present Counter, correlate the information FrameCounter -> PresentCounter before Presenting the frame
+	RHIFrameInfoSetPresentCounter(InPresentArgs.FrameCounter, GRHIPresentCounter);
+	
 	GDynamicRHI->RHIEndDrawingViewport(*this, Viewport, InPresentArgs);
 
 	// Dispatch to ensure any swapchain present commands enqueued by the platform RHI above are executed in a timely manner.
diff --git a/Engine/Source/Runtime/RHI/Private/RHIUtilities.cpp b/Engine/Source/Runtime/RHI/Private/RHIUtilities.cpp
--- a/Engine/Source/Runtime/RHI/Private/RHIUtilities.cpp
+++ b/Engine/Source/Runtime/RHI/Private/RHIUtilities.cpp
@@ -21,12 +21,14 @@ RHIUtilities.cpp:
 
 #define USE_FRAME_OFFSET_THREAD 1
 
-#if !UE_BUILD_SHIPPING && PLATFORM_SUPPORTS_FLIP_TRACKING
+#if UE_INPUT_LATENCY_TRACKING
 bool GVsyncInformationInsights = false;
 static FAutoConsoleVariableRef CVarVsyncInformationInsights(
 	TEXT("r.VsyncInformationInsights"),
 	GVsyncInformationInsights,
 	TEXT("Whether to show Vsync and Input events in UnrealInsights"));
+
+static FRHIFrameInfo GRHIFrameInfo;
 #endif
 
 TAutoConsoleVariable<FString> FDumpTransitionsHelper::CVarDumpTransitionsForResource(
@@ -215,16 +217,6 @@ struct FRHIFrameOffsetThread : public FRunnable
 
 	static FEvent* WaitEvent;
 
-#if !UE_BUILD_SHIPPING
-	struct FFrameDebugInfo
-	{
-		uint64 PresentIndex;
-		uint64 FrameIndex;
-		uint64 InputTime;
-	};
-	TArray<FFrameDebugInfo> FrameDebugInfos;
-#endif
-
 	virtual uint32 Run() override
 	{
 		while (bRun)
@@ -241,33 +233,16 @@ struct FRHIFrameOffsetThread : public FRunnable
 			FPlatformProcess::Sleep(Timeout);
 
 			{
-				FScopeLock Lock(&CS);
-				LastFlipFrame = NewFlipFrame;
-				LastFlipFrame.FlipTimeInSeconds = LastFlipFrame.FlipTimeInSeconds + TargetFrameTimeInSeconds - SlackInSeconds;
-				LastFlipFrame.VBlankTimeInSeconds = LastFlipFrame.VBlankTimeInSeconds + TargetFrameTimeInSeconds - SlackInSeconds;
-				LastFlipFrame.PresentIndex++;
-
-#if !UE_BUILD_SHIPPING && PLATFORM_SUPPORTS_FLIP_TRACKING
-				for (int32 DebugInfoIndex = FrameDebugInfos.Num() - 1; DebugInfoIndex >= 0; --DebugInfoIndex)
 				{
-					auto const& DebugInfo = FrameDebugInfos[DebugInfoIndex];
-					if (NewFlipFrame.PresentIndex == DebugInfo.PresentIndex)
-					{
-						uint64 VBlankTimeInCycle = (uint64)(NewFlipFrame.VBlankTimeInSeconds / FPlatformTime::GetSecondsPerCycle64());
-
-						GInputLatencyTime = (VBlankTimeInCycle) - DebugInfo.InputTime;
-
-						if (GVsyncInformationInsights)
-						{
-							TRACE_BOOKMARK_CYCLES(VBlankTimeInCycle, TEXT("Vsync P:%" UINT64_FMT " F:%" UINT64_FMT " Lat:%.2f"), DebugInfo.PresentIndex, DebugInfo.FrameIndex, FPlatformTime::ToMilliseconds64(GInputLatencyTime))
-						}
-					}
-
-					if (DebugInfo.PresentIndex <= NewFlipFrame.PresentIndex)
-					{
-						FrameDebugInfos.RemoveAtSwap(DebugInfoIndex);
-					}
+					FScopeLock Lock(&CS);
+					LastFlipFrame = NewFlipFrame;
+					LastFlipFrame.FlipTimeInSeconds = LastFlipFrame.FlipTimeInSeconds + TargetFrameTimeInSeconds - SlackInSeconds;
+					LastFlipFrame.VBlankTimeInSeconds = LastFlipFrame.VBlankTimeInSeconds + TargetFrameTimeInSeconds - SlackInSeconds;
+					LastFlipFrame.PresentIndex++;
 				}
+
+#if UE_INPUT_LATENCY_TRACKING
+				RHIFrameInfoSetVsyncDetails(NewFlipFrame);
 #endif
 			}
 
@@ -347,24 +322,6 @@ public:
 		}
 	}
 
-	static void SetFrameDebugInfo(uint64 PresentIndex, uint64 FrameIndex, uint64 InputTime)
-	{
-#if !UE_BUILD_SHIPPING && PLATFORM_SUPPORTS_FLIP_TRACKING
-		FScopeLock Lock(&Singleton.CS);
-
-		FFrameDebugInfo DebugInfo;
-		DebugInfo.PresentIndex = PresentIndex;
-		DebugInfo.FrameIndex = FrameIndex;
-		DebugInfo.InputTime = InputTime;
-		Singleton.FrameDebugInfos.Add(DebugInfo);
-
- 		if (GVsyncInformationInsights)
- 		{
- 			TRACE_BOOKMARK_CYCLES(DebugInfo.InputTime, TEXT("Input P:%" UINT64_FMT " F:%" UINT64_FMT), DebugInfo.PresentIndex, DebugInfo.FrameIndex)
- 		}
-#endif
-	}
-
 private:
 	FEvent* GetOrInitializeWaitEvent()
 	{
@@ -566,47 +523,6 @@ void RHITriggerTaskEventOnFlip(uint64 PresentIndex, const UE::Tasks::FTaskEvent&
 	FRHIFrameFlipTrackingRunnable::TriggerTaskEventOnFlip(PresentIndex, Event);
 }
 
-void RHISetFrameDebugInfo(uint64 PresentIndex, uint64 FrameIndex, uint64 InputTime)
-{
-#if USE_FRAME_OFFSET_THREAD
-	FRHIFrameOffsetThread::SetFrameDebugInfo(PresentIndex, FrameIndex, InputTime);
-#endif
-}
-
-void RHISetVsyncDebugInfo(FRHIFlipDetails& NewFlipFrame)
-{
-#if !UE_BUILD_SHIPPING && USE_FRAME_OFFSET_THREAD && PLATFORM_SUPPORTS_FLIP_TRACKING
-	check(FRHIFrameOffsetThread::Thread == nullptr);
-
-	FScopeLock Lock(&FRHIFrameOffsetThread::Singleton.CS);
-	FRHIFrameOffsetThread::Singleton.LastFlipFrame = NewFlipFrame;
-
-	for (int32 DebugInfoIndex = FRHIFrameOffsetThread::Singleton.FrameDebugInfos.Num() - 1; DebugInfoIndex >= 0; --DebugInfoIndex)
-	{
-		auto const& DebugInfo = FRHIFrameOffsetThread::Singleton.FrameDebugInfos[DebugInfoIndex];
-		//Note that we are using PresentIndex-1 since RHISetFrameDebugInfo is setting GRHIPresentCounter-1 due to other platforms setup
-		if ((NewFlipFrame.PresentIndex-1) == DebugInfo.PresentIndex)
-		{
-			uint64 VBlankTimeInCycle = NewFlipFrame.VBlankTimeInCycles;
-
-			if (VBlankTimeInCycle > DebugInfo.InputTime)
-			{
-				GInputLatencyTime = VBlankTimeInCycle - DebugInfo.InputTime;
-			}
-			
-			if (GVsyncInformationInsights)
-			{
-				TRACE_BOOKMARK_CYCLES(VBlankTimeInCycle, TEXT("Vsync P:%" UINT64_FMT " F:%" UINT64_FMT " Lat:%.2f"), DebugInfo.PresentIndex, DebugInfo.FrameIndex, FPlatformTime::ToMilliseconds64(GInputLatencyTime))
-			}
-		}
-
-		if (DebugInfo.PresentIndex <= NewFlipFrame.PresentIndex)
-		{
-			FRHIFrameOffsetThread::Singleton.FrameDebugInfos.RemoveAtSwap(DebugInfoIndex);
-		}
-	}
-#endif
-}
 
 void RHIInitializeFlipTracking()
 {
@@ -825,3 +741,126 @@ EGpuVendorId RHIGetPreferredAdapterVendor()
 
 	return EGpuVendorId::Unknown;
 }
+
+void FRHIFrameInfo::SetInputSampleMarker(uint64 FrameIndex)
+{
+	uint64 Cycles64 = FPlatformTime::Cycles64(); 
+		
+	{
+		UE::TUniqueLock ScopeLock(FramesMutex);
+
+		FRHIFrameInfoMarker& FrameMarker = Frames.FindOrAdd(FrameIndex);
+		FrameMarker.InputCycle = Cycles64;
+	}	
+}
+
+void FRHIFrameInfo::SetPresentCounter(uint64 FrameIndex, uint64 PresentCounter)
+{
+	{
+		UE::TUniqueLock ScopeLock(FramesMutex);
+
+		FRHIFrameInfoMarker& FrameMarker = Frames.FindOrAdd(FrameIndex);
+		FrameMarker.PresentCounter = PresentCounter;
+	}
+	
+}
+
+
+void FRHIFrameInfo::SetVsyncDetails(const FRHIFlipDetails& NewFlipFrame)
+{ 
+	uint64 LatestFrameCounter = 0;
+
+	{
+		UE::TUniqueLock ScopeLock(FramesMutex);
+
+		for (TPair<uint64, FRHIFrameInfoMarker>& Elem : Frames)
+		{
+			if (Elem.Value.PresentCounter == NewFlipFrame.PresentIndex)
+			{
+				//if VBlankTimeInCycles is not set we compute it from VBlankTimeInSeconds
+				uint64 VBlankTimeInCycle = (NewFlipFrame.VBlankTimeInCycles != 0) ? NewFlipFrame.VBlankTimeInCycles : ((uint64)(NewFlipFrame.VBlankTimeInSeconds / FPlatformTime::GetSecondsPerCycle64()));
+
+				Elem.Value.VSyncCycle = VBlankTimeInCycle;
+				//LatestFrameCounter = Elem.Value.FrameIndex;
+				LatestFrameCounter = Elem.Key; 
+			
+				uint64 CurrentInputLatency = 0;
+
+				if (VBlankTimeInCycle > Elem.Value.InputCycle)
+				{
+					CurrentInputLatency = VBlankTimeInCycle - Elem.Value.InputCycle;
+				}
+
+				if (CurrentInputLatency)
+				{
+					GInputLatencyTime = CurrentInputLatency;
+				}
+
+				Elem.Value.InputLatencyMS = FPlatformTime::ToMilliseconds64(CurrentInputLatency);
+
+#if UE_INPUT_LATENCY_TRACKING
+				if (GVsyncInformationInsights)
+				{
+					TRACE_BOOKMARK_CYCLES(Elem.Value.InputCycle, TEXT("Input P:%" UINT64_FMT " F:%" UINT64_FMT), Elem.Value.PresentCounter, LatestFrameCounter)
+					TRACE_BOOKMARK_CYCLES(VBlankTimeInCycle, TEXT("Vsync P:%" UINT64_FMT " F:%" UINT64_FMT " Lat:%.2f"), Elem.Value.PresentCounter, LatestFrameCounter, Elem.Value.InputLatencyMS)
+				}
+#endif
+				break;
+			}
+		}
+	}
+
+	ComputeAndTrimHistory(LatestFrameCounter);
+}
+
+void FRHIFrameInfo::ComputeAndTrimHistory(uint64 FrameNumber)
+{
+	UE::TUniqueLock ScopeLock(FramesMutex);
+
+	const uint64 FrameThreshold = FrameNumber < FrameInfoHistoryCount ? FrameInfoHistoryCount : FrameNumber - FrameInfoHistoryCount;
+
+	uint32 ElementCount = 0;
+	double AccumulatedInputLatency = 0.0;
+
+	for (auto It = Frames.CreateIterator(); It; ++It)
+	{
+		if (It->Value.InputLatencyMS != 0.0)
+		{
+			AccumulatedInputLatency += It->Value.InputLatencyMS;
+			ElementCount++;
+		}
+
+		if (It->Key <= FrameThreshold)
+		{
+			It.RemoveCurrent();
+		}
+	}
+
+	if (ElementCount)
+	{
+		AverageInputLatencyMs = AccumulatedInputLatency / ((double)ElementCount);
+	}
+
+	// Array should not grow indefinitely 
+	if (Frames.Num() > (FrameInfoHistoryCount*2))
+	{
+		Frames.CreateIterator().RemoveCurrent();
+	}
+}
+
+#if UE_INPUT_LATENCY_TRACKING
+extern RHI_API void RHIFrameInfoSetVsyncDetails(const FRHIFlipDetails& NewFlipFrame)
+{
+	GRHIFrameInfo.SetVsyncDetails(NewFlipFrame);
+}
+
+extern RHI_API void RHIFrameInfoSetInputSampleMarker(uint64 FrameNumber)
+{
+	GRHIFrameInfo.SetInputSampleMarker(FrameNumber);
+}
+
+extern RHI_API void RHIFrameInfoSetPresentCounter(uint64 FrameNumber, uint64 PresentCounter)
+{
+	GRHIFrameInfo.SetPresentCounter(FrameNumber, PresentCounter);
+}
+#endif
diff --git a/Engine/Source/Runtime/RHI/Public/RHIUtilities.h b/Engine/Source/Runtime/RHI/Public/RHIUtilities.h
--- a/Engine/Source/Runtime/RHI/Public/RHIUtilities.h
+++ b/Engine/Source/Runtime/RHI/Public/RHIUtilities.h
@@ -4,6 +4,7 @@
 
 #include "CoreTypes.h"
 #include "Misc/AssertionMacros.h"
+#include "Async/Mutex.h"
 #include "HAL/UnrealMemory.h"
 #include "Math/UnrealMathUtility.h"
 #include "Containers/ResourceArray.h"
@@ -14,6 +15,13 @@
 #include "RHIDefinitions.h"
 #include "RHICommandList.h"
 #include "Tasks/Task.h"
+#include "Containers/Map.h"
+
+#if !UE_BUILD_SHIPPING && PLATFORM_SUPPORTS_FLIP_TRACKING
+	#define UE_INPUT_LATENCY_TRACKING 1
+#else
+	#define UE_INPUT_LATENCY_TRACKING 0
+#endif
 
 inline EPixelFormat GetBlockCompressedFormatUAVAliasFormat(EPixelFormat Format)
 {
@@ -737,12 +745,6 @@ extern RHI_API bool RHIGetSyncAllowVariable();
 
 extern RHI_API void RHITriggerTaskEventOnFlip(uint64 PresentIndex, const UE::Tasks::FTaskEvent& TaskEvent);
 
-/** Sets the FrameIndex and InputTime for the current frame. */
-extern RHI_API void RHISetFrameDebugInfo(uint64 PresentIndex, uint64 FrameIndex, uint64 InputTime);
-
-/** Sets the Vsync information for a new frame */
-extern RHI_API void RHISetVsyncDebugInfo(FRHIFlipDetails& NewFlipFrame);
-
 extern RHI_API void RHIInitializeFlipTracking();
 extern RHI_API void RHIShutdownFlipTracking();
 
@@ -753,3 +755,52 @@ extern RHI_API void RHICalculateFrameTime();
 
 /** Returns the VendorID of the preferred vendor or -1 if none were specified. */
 extern RHI_API EGpuVendorId RHIGetPreferredAdapterVendor();
+
+class FRHIFrameInfo
+{
+public:
+	static constexpr uint32 FrameInfoHistoryCount = 16;
+
+	struct FRHIFrameInfoMarker
+	{
+		//Present Counter that correlate the the Game Frame Index
+		uint64 PresentCounter = 0;
+		//Wall clock cycles when Input was sample on the CPU
+		uint64 InputCycle = 0;
+		//Wall clock cycles when the frame was flip to the Display composition
+		uint64 VSyncCycle = 0;
+		//Frame Input latency from InputCycle to Display composition in milliseconds
+		double InputLatencyMS = 0.0;
+	};
+
+
+	FRHIFrameInfo()		
+	{}
+
+	virtual ~FRHIFrameInfo() {}
+
+	std::atomic<float> AverageInputLatencyMs = 0.0f;
+	
+public:
+	void SetInputSampleMarker(uint64 FrameIndex) ;
+	void SetPresentCounter(uint64 FrameIndex, uint64 PresentCounter);
+	void SetVsyncDetails(const FRHIFlipDetails& NewFlipFrame);
+
+	float GetAverageInputLatencyInMs()  { return AverageInputLatencyMs; }
+
+private:
+	void ComputeAndTrimHistory(uint64 FrameIndex);
+
+	TMap<uint64,FRHIFrameInfoMarker> Frames;
+	UE::FMutex FramesMutex;
+};
+
+#if UE_INPUT_LATENCY_TRACKING
+extern RHI_API void RHIFrameInfoSetInputSampleMarker(uint64 FrameNumber);
+extern RHI_API void RHIFrameInfoSetPresentCounter(uint64 FrameNumber, uint64 PresentCounter);
+extern RHI_API void RHIFrameInfoSetVsyncDetails (const FRHIFlipDetails& NewFlipFrame);
+#else
+inline RHI_API void RHIFrameInfoSetInputSampleMarker(uint64 FrameNumber) {}
+inline RHI_API void RHIFrameInfoSetPresentCounter(uint64 FrameNumber, uint64 PresentCounter) {}
+inline RHI_API void RHIFrameInfoSetVsyncDetails(const FRHIFlipDetails& NewFlipFrame) {}
+#endif 
diff --git a/Engine/Source/Runtime/RenderCore/Private/RenderingThread.cpp b/Engine/Source/Runtime/RenderCore/Private/RenderingThread.cpp
--- a/Engine/Source/Runtime/RenderCore/Private/RenderingThread.cpp
+++ b/Engine/Source/Runtime/RenderCore/Private/RenderingThread.cpp
@@ -2500,20 +2500,6 @@ namespace FFrameEndSync
 
 		check(IsInGameThread());
 
-	#if !UE_BUILD_SHIPPING && PLATFORM_SUPPORTS_FLIP_TRACKING
-		// Set the FrameDebugInfo on platforms that have accurate frame tracking.
-		ENQUEUE_RENDER_COMMAND(FrameDebugInfo)(
-			[CurrentFrameCounter = GFrameCounter, CurrentInputTime = GInputTime](FRHICommandListImmediate& RHICmdList)
-		{
-			RHICmdList.EnqueueLambda(
-				[CurrentFrameCounter, CurrentInputTime](FRHICommandListImmediate&)
-			{
-				// Set the FrameCount and InputTime for input latency stats and flip debugging.
-				RHISetFrameDebugInfo(GRHIPresentCounter - 1, CurrentFrameCounter, CurrentInputTime);
-			});
-		});
-	#endif
-
 		// Always sync with the render thread (either current frame, or N-1 frame)
 		RenderThreadFences.Emplace();
 		while (RenderThreadFences.Num() > (bFullSync ? 0 : 1))
diff --git a/Engine/Source/Runtime/Windows/D3D11RHI/Private/D3D11Viewport.cpp b/Engine/Source/Runtime/Windows/D3D11RHI/Private/D3D11Viewport.cpp
--- a/Engine/Source/Runtime/Windows/D3D11RHI/Private/D3D11Viewport.cpp
+++ b/Engine/Source/Runtime/Windows/D3D11RHI/Private/D3D11Viewport.cpp
@@ -365,7 +365,7 @@ bool FD3D11Viewport::PresentChecked(IRHICommandContext& RHICmdContext, int32 Syn
 	HRESULT Result = S_OK;
 	bool bNeedNativePresent = true;
 
-#if !UE_BUILD_SHIPPING && PLATFORM_SUPPORTS_FLIP_TRACKING
+#if UE_INPUT_LATENCY_TRACKING
 	if (SwapChain.IsValid())
 	{
 		static FRHIFlipDetails LastFlipFrame;
@@ -378,7 +378,7 @@ bool FD3D11Viewport::PresentChecked(IRHICommandContext& RHICmdContext, int32 Syn
 			NewFlipFrame.PresentIndex = stats.PresentCount;
 			NewFlipFrame.VBlankTimeInCycles = stats.SyncQPCTime.QuadPart;
 
-			RHISetVsyncDebugInfo(NewFlipFrame);
+			RHIFrameInfoSetVsyncDetails(NewFlipFrame);
 
 			LastFlipFrame = NewFlipFrame;
 			LastStats = stats;
@@ -386,11 +386,6 @@ bool FD3D11Viewport::PresentChecked(IRHICommandContext& RHICmdContext, int32 Syn
 	}
 #endif
 
-	if (PresentBeginLatencyCallback)
-	{
-		PresentBeginLatencyCallback();
-	}
-
 	if (IsValidRef(CustomPresent))
 	{
 		SCOPE_CYCLE_COUNTER(STAT_D3D11CustomPresentTime);
@@ -401,6 +396,11 @@ bool FD3D11Viewport::PresentChecked(IRHICommandContext& RHICmdContext, int32 Syn
 	{
 		if (SwapChain.IsValid())
 		{
+			if (PresentBeginLatencyCallback)
+			{
+				PresentBeginLatencyCallback();
+			}
+
 			// Check if the viewport's swap chain has been invalidated by DXGI.
 			BOOL bSwapChainFullscreenState;
 			TRefCountPtr<IDXGIOutput> SwapChainOutput;
@@ -505,7 +505,7 @@ bool FD3D11Viewport::PresentChecked(IRHICommandContext& RHICmdContext, int32 Syn
 	UINT PresentID;
 	if ( (SwapChain.IsValid()) && (SUCCEEDED(SwapChain->GetLastPresentCount(&PresentID))) )
 	{
-		GRHIPresentCounter = PresentID;
+		GRHIPresentCounter = PresentID + 1;
 	}
 	else
 	{
