Convert a check into an ensure to prevent possible crashes in external programs

#Synced-CL 50446354

diff --git a/Engine/Source/Runtime/RHI/Private/RHIStrings.cpp b/Engine/Source/Runtime/RHI/Private/RHIStrings.cpp
--- a/Engine/Source/Runtime/RHI/Private/RHIStrings.cpp
+++ b/Engine/Source/Runtime/RHI/Private/RHIStrings.cpp
@@ -311,7 +311,6 @@ FString GetBufferUsageFlagsName(EBufferUsageFlags BufferUsage)
 	{
 		switch (BufferUsage)
 		{
-		default: checkNoEntry(); // fall through
 		case BUF_None:					return TEXT("BUF_None");
 		case BUF_Static:				return TEXT("BUF_Static");
 		case BUF_Dynamic:				return TEXT("BUF_Dynamic");
@@ -336,6 +335,8 @@ FString GetBufferUsageFlagsName(EBufferUsageFlags BufferUsage)
 		case BUF_UniformBuffer:			return TEXT("BUF_UniformBuffer");
 		case BUF_ReservedResource:		return TEXT("BUF_ReservedResource");
 		}
+		ensureMsgf(false, TEXT("Unexpected buffer usage flag bit: %" UPTRINT_x_FMT), BufferUsage);
+		return TEXT("BUF_Unknown");
 	});
 }
 
