- Refactor PSO precaching priority to make it shareable between PC and mobile logic
- Use EPSOPrecachePriority for CSV tracking of queued and in flight operations
- Add low priority PSO precache option
- Add option to give initial priority during PSO precache request

#Synced-CL 50391842

diff --git a/Engine/Plugins/FX/Niagara/Source/Niagara/Private/NiagaraEmitter.cpp b/Engine/Plugins/FX/Niagara/Source/Niagara/Private/NiagaraEmitter.cpp
--- a/Engine/Plugins/FX/Niagara/Source/Niagara/Private/NiagaraEmitter.cpp
+++ b/Engine/Plugins/FX/Niagara/Source/Niagara/Private/NiagaraEmitter.cpp
@@ -3522,6 +3522,12 @@ FGraphEventArray FVersionedNiagaraEmitterData::PrecacheComputePSOs(const UNiagar
 	{
 		return PSOReadyGraphTasks;
 	}
+	
+	FPSOPrecacheRequestParams PSORequestParams;
+	PSORequestParams.InMemoryMode = EPSOPrecacheInMemoryMode::Uncached;
+	PSORequestParams.PrecachePriority = EPSOPrecachePriority::High;
+	PSORequestParams.Name = TEXT("NiagaraEmitter");
+	PSORequestParams.bForcePrecache = true;
 
 	// Compute PSO can fail on some platforms even though the shader compiled successfully this path will wait for the PSO precache to complete before the emitter is allowed to run
 	if (GNiagaraEmitterComputePSOPrecacheMode == 3)
@@ -3536,7 +3542,7 @@ FGraphEventArray FVersionedNiagaraEmitterData::PrecacheComputePSOs(const UNiagar
 			}
 
 			FRHIComputeShader* ComputeShader = ShaderRef.GetComputeShader();
-			FPSOPrecacheRequestResult PSOPrecacheRequestResult = PipelineStateCache::PrecacheComputePipelineState(ComputeShader, EPSOPrecacheInMemoryMode::Uncached, TEXT("NiagaraEmitter"), true);
+			FPSOPrecacheRequestResult PSOPrecacheRequestResult = PipelineStateCache::PrecacheComputePipelineState(ComputeShader, PSORequestParams);
 			if (PSOPrecacheRequestResult.AsyncCompileEvent.IsValid())
 			{
 				PSOPrecacheEvents.Add(PSOPrecacheRequestResult.AsyncCompileEvent);
@@ -3628,7 +3634,7 @@ FGraphEventArray FVersionedNiagaraEmitterData::PrecacheComputePSOs(const UNiagar
 			}
 
 			FRHIComputeShader* ComputeShader = ShaderRef.GetComputeShader();
-			FPSOPrecacheRequestResult PSOPrecacheRequestResult = PipelineStateCache::PrecacheComputePipelineState(ComputeShader, EPSOPrecacheInMemoryMode::Uncached, TEXT("NiagaraEmitter"), true);
+			FPSOPrecacheRequestResult PSOPrecacheRequestResult = PipelineStateCache::PrecacheComputePipelineState(ComputeShader, PSORequestParams);
 			if (PSOPrecacheRequestResult.AsyncCompileEvent != nullptr && bAddToPSOReadyGraphTasks)
 			{
 				PSOReadyGraphTasks.Add(PSOPrecacheRequestResult.AsyncCompileEvent);
diff --git a/Engine/Plugins/FX/Niagara/Source/NiagaraShader/Private/NiagaraShared.cpp b/Engine/Plugins/FX/Niagara/Source/NiagaraShader/Private/NiagaraShared.cpp
--- a/Engine/Plugins/FX/Niagara/Source/NiagaraShader/Private/NiagaraShared.cpp
+++ b/Engine/Plugins/FX/Niagara/Source/NiagaraShader/Private/NiagaraShared.cpp
@@ -553,7 +553,12 @@ void FNiagaraShaderScript::UpdateCachedData_PostCompile(bool bCalledFromSerializ
 					check(NiagaraShader->GetFrequency() == SF_Compute);
 					FRHIShader* RHIShader = GameThreadShaderMap->GetResource()->GetShader(Shader->GetResourceIndex());
 					FRHIComputeShader* RHIComputeShader = static_cast<FRHIComputeShader*>(RHIShader);
-					PipelineStateCache::PrecacheComputePipelineState(RHIComputeShader, EPSOPrecacheInMemoryMode::Uncached, TEXT("NiagaraCompute"));
+
+					FPSOPrecacheRequestParams RequestParams;
+					RequestParams.InMemoryMode = EPSOPrecacheInMemoryMode::Uncached;
+					RequestParams.PrecachePriority = EPSOPrecachePriority::High;
+					RequestParams.Name = TEXT("NiagaraCompute");
+					PipelineStateCache::PrecacheComputePipelineState(RHIComputeShader, RequestParams);
 				}
 			}
 		}
diff --git a/Engine/Source/Runtime/Engine/Private/PSOPrecache.cpp b/Engine/Source/Runtime/Engine/Private/PSOPrecache.cpp
--- a/Engine/Source/Runtime/Engine/Private/PSOPrecache.cpp
+++ b/Engine/Source/Runtime/Engine/Private/PSOPrecache.cpp
@@ -232,6 +232,11 @@ FPSOPrecacheRequestResultArray RequestPrecachePSOs(EPSOPrecacheType PSOPrecacheT
 	FPSOPrecacheRequestResultArray Results;
 	for (const FPSOPrecacheData& PrecacheData : PSOInitializers)
 	{
+		FPSOPrecacheRequestParams PSORequestParams;
+		PSORequestParams.InMemoryMode = PrecacheData.InMemoryMode;
+		PSORequestParams.PrecachePriority = PrecacheData.Priority;
+
+		FPSOPrecacheRequestResult PSOPrecacheResult;
 		switch (PrecacheData.Type)
 		{
 			case FPSOPrecacheData::EType::Graphics:
@@ -242,12 +247,8 @@ FPSOPrecacheRequestResultArray RequestPrecachePSOs(EPSOPrecacheType PSOPrecacheT
 					CSV_CUSTOM_STAT(PSOPrecache, PrecachedGraphics, 1, ECsvCustomStatOp::Accumulate);
 				}
 				#endif // PSO_PRECACHING_VALIDATE			
-
-				FPSOPrecacheRequestResult PSOPrecacheResult = PipelineStateCache::PrecacheGraphicsPipelineState(PrecacheData.GraphicsPSOInitializer, PrecacheData.InMemoryMode);
-				if (PSOPrecacheResult.IsValid() && PrecacheData.bRequired)
-				{
-					Results.AddUnique(PSOPrecacheResult);
-				}
+				
+				PSOPrecacheResult = PipelineStateCache::PrecacheGraphicsPipelineState(PrecacheData.GraphicsPSOInitializer, PSORequestParams);
 				break;
 			}
 			case FPSOPrecacheData::EType::Compute:
@@ -259,14 +260,14 @@ FPSOPrecacheRequestResultArray RequestPrecachePSOs(EPSOPrecacheType PSOPrecacheT
 				}
 				#endif // PSO_PRECACHING_VALIDATE
 			
-				FPSOPrecacheRequestResult PSOPrecacheResult = PipelineStateCache::PrecacheComputePipelineState(PrecacheData.ComputeInitializer, PrecacheData.InMemoryMode);
-				if (PSOPrecacheResult.IsValid() && PrecacheData.bRequired)
-				{
-					Results.AddUnique(PSOPrecacheResult);
-				}
+				PSOPrecacheResult = PipelineStateCache::PrecacheComputePipelineState(PrecacheData.ComputeInitializer, PSORequestParams);
 				break;
 			}
 		}
+		if (PSOPrecacheResult.IsValid() && PrecacheData.bRequired)
+		{
+			Results.AddUnique(PSOPrecacheResult);
+		}
 	}
 
 	return Results;
@@ -307,7 +308,11 @@ public:
 		{
 			FGraphicsPipelineStateInitializer Initializer = InInitializer;
 			Initializer.StatePrecachePSOHash = StatePrecachePSOHash;
-			return PipelineStateCache::PrecacheGraphicsPipelineState(Initializer, EPSOPrecacheInMemoryMode::Uncached);
+			
+			FPSOPrecacheRequestParams PSORequestParams;
+			PSORequestParams.InMemoryMode = EPSOPrecacheInMemoryMode::Uncached;
+			PSORequestParams.PrecachePriority = EPSOPrecachePriority::Highest;
+			return PipelineStateCache::PrecacheGraphicsPipelineState(Initializer, PSORequestParams);
 		};
 		return PrecachePSOInternal(StatePrecachePSOHash, StartPSOPrecache);			
 	}
@@ -317,8 +322,11 @@ public:
 		uint64 ComputePSOHash = GetTypeHash(InInitializer.ComputeShader->GetHash());
 		
 		auto StartPSOPrecache = [InInitializer]()
-		{
-			return PipelineStateCache::PrecacheComputePipelineState(InInitializer, EPSOPrecacheInMemoryMode::Uncached);
+		{			
+			FPSOPrecacheRequestParams PSORequestParams;
+			PSORequestParams.InMemoryMode = EPSOPrecacheInMemoryMode::Uncached;
+			PSORequestParams.PrecachePriority = EPSOPrecachePriority::Highest;
+			return PipelineStateCache::PrecacheComputePipelineState(InInitializer, PSORequestParams);
 		};
 		return PrecachePSOInternal(ComputePSOHash, StartPSOPrecache);			
 	}
@@ -353,10 +361,6 @@ private:
 			FPrecacheStateAndRequestData NewStateAndRequestData;
 			NewStateAndRequestData.RequestData = StartPrecacheFunction();
 			NewStateAndRequestData.bIsPrecached = !PipelineStateCache::IsPrecaching(NewStateAndRequestData.RequestData.RequestID);
-			if (!NewStateAndRequestData.bIsPrecached)
-			{
-				PipelineStateCache::BoostPrecachePriority(EPSOPrecachePriority::Highest, NewStateAndRequestData.RequestData.RequestID);
-			}
 			ActiveRequestData.bIsPrecached = NewStateAndRequestData.bIsPrecached;
 			
 			// Add with write lock
diff --git a/Engine/Source/Runtime/Engine/Private/PSOPrecacheMaterial.cpp b/Engine/Source/Runtime/Engine/Private/PSOPrecacheMaterial.cpp
--- a/Engine/Source/Runtime/Engine/Private/PSOPrecacheMaterial.cpp
+++ b/Engine/Source/Runtime/Engine/Private/PSOPrecacheMaterial.cpp
@@ -199,7 +199,10 @@ public:
 		LLM_TAGSET_SCOPE_CLEAR(ELLMTagSet::Assets);
 		LLM_TAGSET_SCOPE_CLEAR(ELLMTagSet::AssetClasses);
 		UE_TRACE_METADATA_CLEAR_SCOPE();
-		FMaterialPSOPrecacheRequestID RequestID = INDEX_NONE;		
+		FMaterialPSOPrecacheRequestID RequestID = INDEX_NONE;
+
+		// only support medium and high priorities via the material PSO precaching system
+		check(Priority == EPSOPrecachePriority::Medium || Priority == EPSOPrecachePriority::High);
 
 		if (GetPSOPrecacheMode() == EPSOPrecacheMode::PreloadShader)
 		{
diff --git a/Engine/Source/Runtime/Engine/Private/ShaderCompiler/ShaderCompiler.cpp b/Engine/Source/Runtime/Engine/Private/ShaderCompiler/ShaderCompiler.cpp
--- a/Engine/Source/Runtime/Engine/Private/ShaderCompiler/ShaderCompiler.cpp
+++ b/Engine/Source/Runtime/Engine/Private/ShaderCompiler/ShaderCompiler.cpp
@@ -5110,12 +5110,14 @@ void PrecacheComputePipelineStatesForGlobalShaders(ERHIFeatureLevel::Type Featur
 						{
 							// PSO precache shaders are not required to all load correctly
 							constexpr bool bRequired = false;
-
 							FRHIComputeShader* RHIComputeShader = static_cast<FRHIComputeShader*>(GlobalShader.GetRHIShaderBase(SF_Compute, bRequired));
 							if (RHIComputeShader)
-							{
-								const TCHAR* TypeName = GlobalShader.GetType()->GetName();
-								PipelineStateCache::PrecacheComputePipelineState(RHIComputeShader, EPSOPrecacheInMemoryMode::Uncached, TypeName);
+							{								
+								FPSOPrecacheRequestParams PSORequestParams;
+								PSORequestParams.InMemoryMode = EPSOPrecacheInMemoryMode::Uncached;
+								PSORequestParams.PrecachePriority = EPSOPrecachePriority::Medium;
+								PSORequestParams.Name = GlobalShader.GetType()->GetName();
+								PipelineStateCache::PrecacheComputePipelineState(RHIComputeShader, PSORequestParams);
 							}
 						}
 					));
@@ -5137,9 +5139,12 @@ void PrecacheComputePipelineStatesForGlobalShaders(ERHIFeatureLevel::Type Featur
 							bool bRequired = false;
 							FRHIComputeShader* RHIComputeShader = static_cast<FRHIComputeShader*>(GlobalShader.GetRHIShaderBase(SF_Compute, bRequired));
 							if (RHIComputeShader)
-							{
-								const TCHAR* TypeName = GlobalShader.GetType()->GetName();
-								PipelineStateCache::PrecacheComputePipelineState(RHIComputeShader, EPSOPrecacheInMemoryMode::Uncached, TypeName);
+							{								
+								FPSOPrecacheRequestParams PSORequestParams;
+								PSORequestParams.InMemoryMode = EPSOPrecacheInMemoryMode::Uncached;
+								PSORequestParams.PrecachePriority = EPSOPrecachePriority::Medium;
+								PSORequestParams.Name = GlobalShader.GetType()->GetName();
+								PipelineStateCache::PrecacheComputePipelineState(RHIComputeShader, PSORequestParams);
 							}
 						}
 					});
diff --git a/Engine/Source/Runtime/Engine/Public/PSOPrecache.h b/Engine/Source/Runtime/Engine/Public/PSOPrecache.h
--- a/Engine/Source/Runtime/Engine/Public/PSOPrecache.h
+++ b/Engine/Source/Runtime/Engine/Public/PSOPrecache.h
@@ -255,6 +255,9 @@ struct FPSOPrecacheData
 	// (proxy creation won't wait for these PSOs if enabled)
 	bool bRequired;
 
+	// Precache priority mode
+	EPSOPrecachePriority Priority = EPSOPrecachePriority::Medium;
+
 	// Keep the PSO in memory until the owning material is request deleted
 	EPSOPrecacheInMemoryMode InMemoryMode;
 
diff --git a/Engine/Source/Runtime/Engine/Public/PSOPrecacheComponent.h b/Engine/Source/Runtime/Engine/Public/PSOPrecacheComponent.h
--- a/Engine/Source/Runtime/Engine/Public/PSOPrecacheComponent.h
+++ b/Engine/Source/Runtime/Engine/Public/PSOPrecacheComponent.h
@@ -101,10 +101,10 @@ private:
 	uint8 bPSOPrecacheCalled : 1 = false;
 
 	/** PSOs requested priority */
-	EPSOPrecachePriority PSOPrecacheRequestPriority : 2 = EPSOPrecachePriority::Medium;
-	static_assert((int)EPSOPrecachePriority::Highest < 1 << 2);
+	EPSOPrecachePriority PSOPrecacheRequestPriority : 3 = EPSOPrecachePriority::Medium;
+	static_assert((int)EPSOPrecachePriority::Highest < 1 << 3);
 
-	int32 LatestPSOPrecacheJobSet : 29 = 0;
+	int32 LatestPSOPrecacheJobSet : 28 = 0;
 
 	/** Cached array of material PSO requests which can be used to boost the priority */
 	TArray<FMaterialPSOPrecacheRequestID> MaterialPSOPrecacheRequestIDs;
diff --git a/Engine/Source/Runtime/RHI/Private/PipelineStateCache.cpp b/Engine/Source/Runtime/RHI/Private/PipelineStateCache.cpp
--- a/Engine/Source/Runtime/RHI/Private/PipelineStateCache.cpp
+++ b/Engine/Source/Runtime/RHI/Private/PipelineStateCache.cpp
@@ -48,27 +48,31 @@ DEFINE_LOG_CATEGORY_STATIC(LogPSOHitching, Log, All);
 DECLARE_DWORD_ACCUMULATOR_STAT(TEXT("Runtime Graphics PSO Hitch Count"), STAT_RuntimeGraphicsPSOHitchCount, STATGROUP_PipelineStateCache);
 DECLARE_DWORD_ACCUMULATOR_STAT(TEXT("Runtime Compute PSO Hitch Count"), STAT_RuntimeComputePSOHitchCount, STATGROUP_PipelineStateCache);
 
-DECLARE_DWORD_ACCUMULATOR_STAT(TEXT("Active Graphics PSO Precache Requests"), STAT_ActiveGraphicsPSOPrecacheRequests, STATGROUP_PipelineStateCache);
-DECLARE_DWORD_ACCUMULATOR_STAT(TEXT("Active Compute PSO Precache Requests"), STAT_ActiveComputePSOPrecacheRequests, STATGROUP_PipelineStateCache);
-
-DECLARE_DWORD_ACCUMULATOR_STAT(TEXT("High Priority Graphics PSO Precache Requests"), STAT_HighPriorityGraphicsPSOPrecacheRequests, STATGROUP_PipelineStateCache);
-DECLARE_DWORD_ACCUMULATOR_STAT(TEXT("High Priority Compute PSO Precache Requests"), STAT_HighPriorityComputePSOPrecacheRequests, STATGROUP_PipelineStateCache);
-
-DECLARE_DWORD_ACCUMULATOR_STAT(TEXT("Highest Priority Graphics PSO Precache Requests"), STAT_HighestPriorityGraphicsPSOPrecacheRequests, STATGROUP_PipelineStateCache);
-DECLARE_DWORD_ACCUMULATOR_STAT(TEXT("Highest Priority Compute PSO Precache Requests"), STAT_HighestPriorityComputePSOPrecacheRequests, STATGROUP_PipelineStateCache);
-
 DECLARE_DWORD_ACCUMULATOR_STAT(TEXT("Total Precached Graphics PSOs"), STAT_TotalGraphicsPrecachedPSOs, STATGROUP_PipelineStateCache);
 DECLARE_DWORD_ACCUMULATOR_STAT(TEXT("Total Precached Compute PSOs"), STAT_TotalComputePrecachedPSOs, STATGROUP_PipelineStateCache);
 
-DECLARE_STATS_GROUP(TEXT("PSO InMemory"), STATGROUP_PSOInMemory, STATCAT_Advanced);
+DECLARE_STATS_GROUP(TEXT("PSOPrecache InFlight"), STATGROUP_PSOPrecacheInFlight, STATCAT_Advanced);
+
+DECLARE_DWORD_ACCUMULATOR_STAT(TEXT("Active Graphics PSO Precache Requests"), STAT_ActiveGraphicsPSOPrecacheRequests, STATGROUP_PSOPrecacheInFlight);
+DECLARE_DWORD_ACCUMULATOR_STAT(TEXT("Active Compute PSO Precache Requests"), STAT_ActiveComputePSOPrecacheRequests, STATGROUP_PSOPrecacheInFlight);
+
+DECLARE_DWORD_ACCUMULATOR_STAT(TEXT("Highest Priority Graphics PSO Precache Requests"), STAT_HighestPriorityGraphicsPSOPrecacheRequests, STATGROUP_PSOPrecacheInFlight);
+DECLARE_DWORD_ACCUMULATOR_STAT(TEXT("Highest Priority Compute PSO Precache Requests"), STAT_HighestPriorityComputePSOPrecacheRequests, STATGROUP_PSOPrecacheInFlight);
+DECLARE_DWORD_ACCUMULATOR_STAT(TEXT("High Priority Graphics PSO Precache Requests"), STAT_HighPriorityGraphicsPSOPrecacheRequests, STATGROUP_PSOPrecacheInFlight);
+DECLARE_DWORD_ACCUMULATOR_STAT(TEXT("High Priority Compute PSO Precache Requests"), STAT_HighPriorityComputePSOPrecacheRequests, STATGROUP_PSOPrecacheInFlight);
+DECLARE_DWORD_ACCUMULATOR_STAT(TEXT("Medium Priority Graphics PSO Precache Requests"), STAT_MediumPriorityGraphicsPSOPrecacheRequests, STATGROUP_PSOPrecacheInFlight);
+DECLARE_DWORD_ACCUMULATOR_STAT(TEXT("Medium Priority Compute PSO Precache Requests"), STAT_MediumPriorityComputePSOPrecacheRequests, STATGROUP_PSOPrecacheInFlight);
+DECLARE_DWORD_ACCUMULATOR_STAT(TEXT("Low Priority Graphics PSO Precache Requests"), STAT_LowPriorityGraphicsPSOPrecacheRequests, STATGROUP_PSOPrecacheInFlight);
+DECLARE_DWORD_ACCUMULATOR_STAT(TEXT("Low Priority Compute PSO Precache Requests"), STAT_LowPriorityComputePSOPrecacheRequests, STATGROUP_PSOPrecacheInFlight);
 
-DECLARE_DWORD_ACCUMULATOR_STAT(TEXT("Precached Graphics PSOs InMemory"), STAT_PrecachedGraphicsPrecachedPSOsInMemory, STATGROUP_PSOInMemory);
-DECLARE_DWORD_ACCUMULATOR_STAT(TEXT("Precached Compute PSOs InMemory"), STAT_PrecachedComputePrecachedPSOsInMemory, STATGROUP_PSOInMemory);
+DECLARE_STATS_GROUP(TEXT("PSOPrecache InMemory"), STATGROUP_PSOPrecacheInMemory, STATCAT_Advanced);
 
-DECLARE_DWORD_ACCUMULATOR_STAT(TEXT("Runtime Graphics PSOs"), STAT_ActiveRuntimeGraphicsPSOs, STATGROUP_PSOInMemory);
-DECLARE_DWORD_ACCUMULATOR_STAT(TEXT("Runtime Compute PSOs"), STAT_ActiveRuntimeComputePSOs, STATGROUP_PSOInMemory);
-DECLARE_DWORD_ACCUMULATOR_STAT(TEXT("Runtime Graphics PSOs From InMemoryCache"), STAT_RuntimeGraphicsPSOsFromInMemoryPSOCache, STATGROUP_PSOInMemory);
-DECLARE_DWORD_ACCUMULATOR_STAT(TEXT("Runtime Compute PSOs From InMemoryCache"), STAT_RuntimeComputePSOsFromInMemoryPSOCache, STATGROUP_PSOInMemory);
+DECLARE_DWORD_ACCUMULATOR_STAT(TEXT("Precached Graphics PSOs InMemory"), STAT_PrecachedGraphicsPrecachedPSOsInMemory, STATGROUP_PSOPrecacheInMemory);
+DECLARE_DWORD_ACCUMULATOR_STAT(TEXT("Precached Compute PSOs InMemory"), STAT_PrecachedComputePrecachedPSOsInMemory, STATGROUP_PSOPrecacheInMemory);
+DECLARE_DWORD_ACCUMULATOR_STAT(TEXT("Runtime Graphics PSOs"), STAT_ActiveRuntimeGraphicsPSOs, STATGROUP_PSOPrecacheInMemory);
+DECLARE_DWORD_ACCUMULATOR_STAT(TEXT("Runtime Compute PSOs"), STAT_ActiveRuntimeComputePSOs, STATGROUP_PSOPrecacheInMemory);
+DECLARE_DWORD_ACCUMULATOR_STAT(TEXT("Runtime Graphics PSOs From InMemoryCache"), STAT_RuntimeGraphicsPSOsFromInMemoryPSOCache, STATGROUP_PSOPrecacheInMemory);
+DECLARE_DWORD_ACCUMULATOR_STAT(TEXT("Runtime Compute PSOs From InMemoryCache"), STAT_RuntimeComputePSOsFromInMemoryPSOCache, STATGROUP_PSOPrecacheInMemory);
 
 static inline uint32 GetTypeHash(const FBoundShaderStateInput& Input)
 {
@@ -241,10 +245,9 @@ static FAutoConsoleVariableRef CVarPSOOnlyPrecacheMaxPriority(
 	GPSOPrecacheProcessPriorityThreshold,
 	TEXT("The precache task priority threshold below which tasks are not processed.\n")
 	TEXT("0: All compile tasks are processed. (default)\n")
-	TEXT("1: Minimum priority compile tasks and above are processed.\n")
-	TEXT("2: Normal priority compile tasks and above are processed.\n")
-	TEXT("3: Only maximum priority compile tasks are processed.")
-	,
+	TEXT("1: Medium priority compile tasks and above are processed.\n")
+	TEXT("2: High priority compile tasks and above are processed.\n")
+	TEXT("3: Only very high priority compile tasks are processed."),
 	ECVF_Default
 );
 static_assert((int)FGraphicsPipelineStateInitializer::EPSOPrecacheCompileType::NumTypes == 4, "EPSOPrecacheCompileType has changed, 'r.PSOPrecaching.TaskPriorityThreshold' must be updated.");
@@ -456,23 +459,55 @@ const TCHAR* LexToString(EPSOPrecacheResult Result)
 	}
 }
 
-static FGraphicsPipelineStateInitializer::EPSOPrecacheCompileType GetPSOCompileTypeFromQueuePri(EQueuedWorkPriority QueuePri)
+static FGraphicsPipelineStateInitializer::EPSOPrecacheCompileType GetPSOCompileTypeFromPrecachePriority(EPSOPrecachePriority PrecachePriority)
 {
-	switch (QueuePri)
+	switch (PrecachePriority)
 	{
-	case EQueuedWorkPriority::Blocking:
-	case EQueuedWorkPriority::Highest:
+	case EPSOPrecachePriority::Highest:
 		return FGraphicsPipelineStateInitializer::EPSOPrecacheCompileType::MaxPri;
-	case EQueuedWorkPriority::High:
-	case EQueuedWorkPriority::Normal:
+	case EPSOPrecachePriority::High:
+	case EPSOPrecachePriority::Medium:
 		return FGraphicsPipelineStateInitializer::EPSOPrecacheCompileType::NormalPri;
 	default: checkNoEntry(); [[fallthrough]];
-	case EQueuedWorkPriority::Low:
-	case EQueuedWorkPriority::Lowest:
+	case EPSOPrecachePriority::Low:
 		return FGraphicsPipelineStateInitializer::EPSOPrecacheCompileType::MinPri;
 	}
 }
 
+static EQueuedWorkPriority GetQueuedWorkPriorityFromPSOPrecachePriority(EPSOPrecachePriority PrecachePriority)
+{
+	switch (PrecachePriority)
+	{
+		case EPSOPrecachePriority::Highest:
+			return EQueuedWorkPriority::Highest;
+		case EPSOPrecachePriority::High:
+			return EQueuedWorkPriority::High;
+		case EPSOPrecachePriority::Low:
+			return EQueuedWorkPriority::Low;
+		default: checkNoEntry(); [[fallthrough]];
+		case EPSOPrecachePriority::Medium:
+			return EQueuedWorkPriority::Normal;
+	}
+}
+
+static EPSOPrecachePriority GetPSOPrecachePriorityFromQueuePri(EQueuedWorkPriority QueuePri)
+{
+	switch (QueuePri)
+	{
+		case EQueuedWorkPriority::Blocking:
+		case EQueuedWorkPriority::Highest:
+			return EPSOPrecachePriority::Highest;
+		case EQueuedWorkPriority::High:
+			return EPSOPrecachePriority::High;
+		case EQueuedWorkPriority::Low:
+		case EQueuedWorkPriority::Lowest:
+			return EPSOPrecachePriority::Low;			
+		default: checkNoEntry(); [[fallthrough]];
+		case EQueuedWorkPriority::Normal:
+			return EPSOPrecachePriority::Medium;
+	}
+}
+
 // PSO thread pool wrapper allows us to control scheduling based on the priority of tasks.
 // when r.PSOPrecaching.TaskPriorityThreshold is set only jobs >= priority will be processed.
 // note, the remaining tasks are required to stick around as they are found and boosted when called for.
@@ -484,12 +519,12 @@ public:
 	{}
 protected:
 
-	virtual bool CanSchedule(EQueuedWorkPriority Priority) const override
+	virtual bool CanSchedule(EQueuedWorkPriority QueuePriority) const override
 	{
-		if( FQueuedThreadPoolWrapper::CanSchedule(Priority) )
+		if(FQueuedThreadPoolWrapper::CanSchedule(QueuePriority))
 		{
-			const FGraphicsPipelineStateInitializer::EPSOPrecacheCompileType TaskPri = GetPSOCompileTypeFromQueuePri(Priority);
-			return ((int32)TaskPri >= GPSOPrecacheProcessPriorityThreshold);
+			const EPSOPrecachePriority PrecachePriority = GetPSOPrecachePriorityFromQueuePri(QueuePriority);
+			return ((int32)PrecachePriority >= GPSOPrecacheProcessPriorityThreshold);
 		}
 		return false;
 	}
@@ -671,20 +706,20 @@ static void HandlePipelineCreationFailure(const FComputePipelineStateInitializer
 
 #if PSO_TRACK_CACHE_STATS
 
-static std::atomic_uint64_t TotalPrecompileCompleteTime[(int)EQueuedWorkPriority::Count];
-static std::atomic_uint64_t TotalPrecompileCompileTime[(int)EQueuedWorkPriority::Count];
-static std::atomic_uint64_t TotalPrecompileTimeToBegin[(int)EQueuedWorkPriority::Count];
+static std::atomic_uint64_t TotalPrecompileCompleteTime[(int)EPSOPrecachePriority::Num];
+static std::atomic_uint64_t TotalPrecompileCompileTime[(int)EPSOPrecachePriority::Num];
+static std::atomic_uint64_t TotalPrecompileTimeToBegin[(int)EPSOPrecachePriority::Num];
 
-static std::atomic_int64_t TotalNumPrecompileJobs[(int)EQueuedWorkPriority::Count];
-static std::atomic_int64_t TotalNumPrecompileJobsCompleted[(int)EQueuedWorkPriority::Count];
+static std::atomic_int64_t TotalNumPrecompileJobs[(int)EPSOPrecachePriority::Num];
+static std::atomic_int64_t TotalNumPrecompileJobsCompleted[(int)EPSOPrecachePriority::Num];
 
-static std::atomic_int64_t MaxPrecompileJobTime[(int)EQueuedWorkPriority::Count];
-static std::atomic_int64_t MaxPrecompileTimeToCompile[(int)EQueuedWorkPriority::Count];
-static std::atomic_int64_t MaxPrecompileTimeToBegin[(int)EQueuedWorkPriority::Count];
+static std::atomic_int64_t MaxPrecompileJobTime[(int)EPSOPrecachePriority::Num];
+static std::atomic_int64_t MaxPrecompileTimeToCompile[(int)EPSOPrecachePriority::Num];
+static std::atomic_int64_t MaxPrecompileTimeToBegin[(int)EPSOPrecachePriority::Num];
 
 void ResetPrecompileStats()
 {
-	for (int i = 0; i < (int)EQueuedWorkPriority::Count; i++)
+	for (int i = 0; i < (int)EPSOPrecachePriority::Num; i++)
 	{
 		TotalPrecompileCompleteTime[i] = 0;
 		TotalPrecompileCompileTime[i] = 0;
@@ -699,7 +734,7 @@ void ResetPrecompileStats()
 	}
 }
 
-void StatsEndPrecompile(uint64 CreateTime, uint64 RescheduleTime, uint64 TaskBeginTime, uint64 EndTime, EQueuedWorkPriority TaskPri)
+void StatsEndPrecompile(uint64 CreateTime, uint64 RescheduleTime, uint64 TaskBeginTime, uint64 EndTime, EPSOPrecachePriority PrecachePriority)
 {
 	uint64 TaskIssueTime = FMath::Max(CreateTime, RescheduleTime);
 	uint64 TimeToComplete = EndTime - TaskIssueTime;
@@ -707,15 +742,15 @@ void StatsEndPrecompile(uint64 CreateTime, uint64 RescheduleTime, uint64 TaskBeg
 	uint64 TimeToBegin = TaskBeginTime - TaskIssueTime;
 	check (TaskBeginTime > TaskIssueTime);
 
-	TotalPrecompileCompleteTime[(int)TaskPri] += TimeToComplete;
-	TotalPrecompileCompileTime[(int)TaskPri] += TimeToCompile;
-	TotalPrecompileTimeToBegin[(int)TaskPri] += TimeToBegin;
+	TotalPrecompileCompleteTime[(int)PrecachePriority] += TimeToComplete;
+	TotalPrecompileCompileTime[(int)PrecachePriority] += TimeToCompile;
+	TotalPrecompileTimeToBegin[(int)PrecachePriority] += TimeToBegin;
 
-	MaxPrecompileJobTime[(int)TaskPri] = FMath::Max((uint64)MaxPrecompileJobTime[(int)TaskPri].load(), (uint64)TimeToComplete);
-	MaxPrecompileTimeToCompile[(int)TaskPri] = FMath::Max((uint64)MaxPrecompileTimeToCompile[(int)TaskPri].load(), (uint64)TimeToCompile);
-	MaxPrecompileTimeToBegin[(int)TaskPri] = FMath::Max((uint64)MaxPrecompileTimeToBegin[(int)TaskPri].load(), (uint64)TimeToBegin);
+	MaxPrecompileJobTime[(int)PrecachePriority] = FMath::Max((uint64)MaxPrecompileJobTime[(int)TaskPri].load(), (uint64)TimeToComplete);
+	MaxPrecompileTimeToCompile[(int)PrecachePriority] = FMath::Max((uint64)MaxPrecompileTimeToCompile[(int)TaskPri].load(), (uint64)TimeToCompile);
+	MaxPrecompileTimeToBegin[(int)PrecachePriority] = FMath::Max((uint64)MaxPrecompileTimeToBegin[(int)TaskPri].load(), (uint64)TimeToBegin);
 
-	TotalNumPrecompileJobsCompleted[(int)TaskPri]++;
+	TotalNumPrecompileJobsCompleted[(int)PrecachePriority]++;
 }
 #endif
 
@@ -843,27 +878,28 @@ namespace PipelineStateCache
 #if CSV_PROFILER_STATS
 class FPSOCSVStatTracker
 {
-	bool bPrecache = false;
-	bool bCompute = false;
-
+private:
+	
 	struct FPersistentStats
 	{
 		// PSO compiles can take place over multiple frames, use the persistent stats.
 		TCsvPersistentCustomStat<int>* TotalCompletedPSOPrecompiles			= FCsvProfiler::Get()->GetOrCreatePersistentCustomStatInt(TEXT("TotalCompletedPSOPrecompiles"), CSV_CATEGORY_INDEX(PSOPrecacheCompiling));
 		TCsvPersistentCustomStat<int>* PSOPrecompileQueueLength				= FCsvProfiler::Get()->GetOrCreatePersistentCustomStatInt(TEXT("PSOPrecompileQueueLength"), CSV_CATEGORY_INDEX(PSOPrecacheCompiling));
+		TCsvPersistentCustomStat<int>* PSOGraphicsPrecompilesInProgress		= FCsvProfiler::Get()->GetOrCreatePersistentCustomStatInt(TEXT("PSOGraphicsPrecompilesInProgress"), CSV_CATEGORY_INDEX(PSOPrecacheCompiling));
 		TCsvPersistentCustomStat<int>* PSOComputePrecompilesInProgress		= FCsvProfiler::Get()->GetOrCreatePersistentCustomStatInt(TEXT("PSOComputePrecompilesInProgress"), CSV_CATEGORY_INDEX(PSOPrecacheCompiling));
-		TCsvPersistentCustomStat<int>* MinPriPrecacheInFlight				= FCsvProfiler::Get()->GetOrCreatePersistentCustomStatInt(TEXT("MinPriPrecacheInFlight"), CSV_CATEGORY_INDEX(PSOPrecacheCompiling));
-		TCsvPersistentCustomStat<int>* NormalPriPrecacheInFlight			= FCsvProfiler::Get()->GetOrCreatePersistentCustomStatInt(TEXT("NormalPriPrecacheInFlight"), CSV_CATEGORY_INDEX(PSOPrecacheCompiling));
-		TCsvPersistentCustomStat<int>* MaxPriPrecacheInFlight				= FCsvProfiler::Get()->GetOrCreatePersistentCustomStatInt(TEXT("MaxPriPrecacheInFlight"), CSV_CATEGORY_INDEX(PSOPrecacheCompiling));
 
-		TCsvPersistentCustomStat<int>* MaxPriPSOPrecompileQueueLength		= FCsvProfiler::Get()->GetOrCreatePersistentCustomStatInt(TEXT("MaxPriPSOPrecompileQueueLength"), CSV_CATEGORY_INDEX(PSOPrecacheCompiling));
-		TCsvPersistentCustomStat<int>* NormalPriPSOPrecompileQueueLength	= FCsvProfiler::Get()->GetOrCreatePersistentCustomStatInt(TEXT("NormalPriPSOPrecompileQueueLength"), CSV_CATEGORY_INDEX(PSOPrecacheCompiling));
-		TCsvPersistentCustomStat<int>* MinPriPSOPrecompileQueueLength		= FCsvProfiler::Get()->GetOrCreatePersistentCustomStatInt(TEXT("MinPriPSOPrecompileQueueLength"), CSV_CATEGORY_INDEX(PSOPrecacheCompiling));
-	};
+		TCsvPersistentCustomStat<int>* LowPriorityPrecacheInFlight			= FCsvProfiler::Get()->GetOrCreatePersistentCustomStatInt(TEXT("LowPriorityPrecacheInFlight"), CSV_CATEGORY_INDEX(PSOPrecacheCompiling));
+		TCsvPersistentCustomStat<int>* MediumPriorityPrecacheInFlight		= FCsvProfiler::Get()->GetOrCreatePersistentCustomStatInt(TEXT("MediumPriorityPrecacheInFlight"), CSV_CATEGORY_INDEX(PSOPrecacheCompiling));
+		TCsvPersistentCustomStat<int>* HighPriorityPrecacheInFlight			= FCsvProfiler::Get()->GetOrCreatePersistentCustomStatInt(TEXT("HighPriorityPrecacheInFlight"), CSV_CATEGORY_INDEX(PSOPrecacheCompiling));
+		TCsvPersistentCustomStat<int>* HighestPriorityPrecacheInFlight		= FCsvProfiler::Get()->GetOrCreatePersistentCustomStatInt(TEXT("HighestPriorityPrecacheInFlight"), CSV_CATEGORY_INDEX(PSOPrecacheCompiling));
 
-	std::atomic_int Pri = (int)FGraphicsPipelineStateInitializer::EPSOPrecacheCompileType::NotSet;
+		TCsvPersistentCustomStat<int>* LowPriorityPSOPrecompileQueueLength	= FCsvProfiler::Get()->GetOrCreatePersistentCustomStatInt(TEXT("LowPriorityPSOPrecompileQueueLength"), CSV_CATEGORY_INDEX(PSOPrecacheCompiling));
+		TCsvPersistentCustomStat<int>* MediumPriorityPSOPrecompileQueueLength	= FCsvProfiler::Get()->GetOrCreatePersistentCustomStatInt(TEXT("MediumPriorityPSOPrecompileQueueLength"), CSV_CATEGORY_INDEX(PSOPrecacheCompiling));
+		TCsvPersistentCustomStat<int>* HighPriorityPSOPrecompileQueueLength	= FCsvProfiler::Get()->GetOrCreatePersistentCustomStatInt(TEXT("HighPriorityPSOPrecompileQueueLength"), CSV_CATEGORY_INDEX(PSOPrecacheCompiling));
+		TCsvPersistentCustomStat<int>* HighestPriorityPSOPrecompileQueueLength = FCsvProfiler::Get()->GetOrCreatePersistentCustomStatInt(TEXT("HighestPriorityPSOPrecompileQueueLength"), CSV_CATEGORY_INDEX(PSOPrecacheCompiling));
+	};
 
-	static FPersistentStats& GetStats()
+	static FPersistentStats& GetPersistentStats()
 	{
 		static FPersistentStats PersistentStats;
 		return PersistentStats;
@@ -872,10 +908,29 @@ class FPSOCSVStatTracker
 	{
 		return bEnableStatTracker && FCsvProfiler::Get()->IsCategoryEnabled(CSV_CATEGORY_INDEX(PSOPrecacheCompiling));
 	}
+
+	enum class EStatType
+	{
+		InFlight,
+		Queued,
+	};
+
+	std::atomic_int PSOPrecachePriority = (int)EPSOPrecachePriority::NotSet;
+
+	uint8 bPrecache : 1 = 0;
+	uint8 bCompute : 1 = 0;
+	uint8 bProcessedFinished : 1 = 0;
+
 public:
+
 	static inline bool bEnableStatTracker = true;
 
-	void SetState(bool bInPrecache, bool bInCompute, FGraphicsPipelineStateInitializer::EPSOPrecacheCompileType InPri)
+	~FPSOCSVStatTracker()
+	{
+		check(!AreStatsEnabled() || bProcessedFinished || !bPrecache);
+	}
+
+	void SetState(bool bInPrecache, bool bInCompute, EPSOPrecachePriority InPriority)
 	{
 		if (!AreStatsEnabled())
 		{
@@ -887,99 +942,111 @@ public:
 
 		if (bPrecache)
 		{
-			FPersistentStats& PersistentStats = GetStats();
+			FPersistentStats& PersistentStats = GetPersistentStats();
 			PersistentStats.PSOPrecompileQueueLength->Add(1);
 
 			if (bCompute)
 			{
 				PersistentStats.PSOComputePrecompilesInProgress->Add(1);
 			}
+			else
+			{
+				PersistentStats.PSOGraphicsPrecompilesInProgress->Add(1);
+			}
 
-			SetPriority(InPri);
+			SetQueuedPriority(InPriority);
 		}
 	}
 	
-	void SetPriority(FGraphicsPipelineStateInitializer::EPSOPrecacheCompileType InPri)
+	void SetQueuedPriority(EPSOPrecachePriority InNewPriority)
 	{
 		if (!AreStatsEnabled())
 		{
 			return;
 		}
 
-		FGraphicsPipelineStateInitializer::EPSOPrecacheCompileType PriorValue = (FGraphicsPipelineStateInitializer::EPSOPrecacheCompileType)Pri.exchange((int)InPri);
-		if (bPrecache && PriorValue != InPri)
+		EPSOPrecachePriority PriorValue = (EPSOPrecachePriority)PSOPrecachePriority.exchange((int)InNewPriority);
+		if (bPrecache && PriorValue != InNewPriority)
 		{
 			// unset the previous priority
-			if (PriorValue != FGraphicsPipelineStateInitializer::EPSOPrecacheCompileType::NotSet)
+			if (PriorValue != EPSOPrecachePriority::NotSet)
 			{
-				GetPriStat(PriorValue, false)->Add(-1);
+				GetPriorityStat(PriorValue, EStatType::Queued)->Add(-1);
 			}
-			GetPriStat(InPri, false)->Add(1);
+			GetPriorityStat(InNewPriority, EStatType::Queued)->Add(1);
 		}
 	}
 
-	void SetInFlightPriority(FGraphicsPipelineStateInitializer::EPSOPrecacheCompileType InPri)
+	void SetInFlightPriority(EPSOPrecachePriority InFlightPriority)
 	{
 		if (!AreStatsEnabled())
 		{
 			return;
 		}
 
-		if (bPrecache)
-		{
-			// This is called during the task execution where the priority is final.
-			// We set both the priority and the inflight state.
-			check(InPri != FGraphicsPipelineStateInitializer::EPSOPrecacheCompileType::NotSet);
-			SetPriority(InPri);
-			GetPriStat(InPri, true)->Add(1);
-		}
+		check(bPrecache);
+
+		// This is called during the task execution where the priority is final.
+		// We set both the priority and the inflight state.
+		check(InFlightPriority != EPSOPrecachePriority::NotSet);
+		SetQueuedPriority(InFlightPriority);
+		GetPriorityStat(InFlightPriority, EStatType::InFlight)->Add(1);
 	}
 
-	~FPSOCSVStatTracker()
+	void PrecacheFinished()
 	{
 		if (!AreStatsEnabled())
 		{
 			return;
 		}
 
-		if (bPrecache)
-		{
-			FPersistentStats& PersistentStats = GetStats();
+		check(bPrecache)
 
-			PersistentStats.TotalCompletedPSOPrecompiles->Add(1);
-			PersistentStats.PSOPrecompileQueueLength->Add(-1);
+		FPersistentStats& PersistentStats = GetPersistentStats();
 
-			if (bCompute)
-			{
-				PersistentStats.PSOComputePrecompilesInProgress->Add(-1);
-			}
+		PersistentStats.TotalCompletedPSOPrecompiles->Add(1);
+		PersistentStats.PSOPrecompileQueueLength->Add(-1);
 
-			FGraphicsPipelineStateInitializer::EPSOPrecacheCompileType MyPri = (FGraphicsPipelineStateInitializer::EPSOPrecacheCompileType)Pri.load();
-			if (ensure(MyPri != FGraphicsPipelineStateInitializer::EPSOPrecacheCompileType::NotSet))
-			{
-				GetPriStat(MyPri, true)->Add(-1);
-				GetPriStat(MyPri, false)->Add(-1);
-			}
+		if (bCompute)
+		{
+			PersistentStats.PSOComputePrecompilesInProgress->Add(-1);
 		}
-	}
-	private:
-		TCsvPersistentCustomStat<int>* GetPriStat(FGraphicsPipelineStateInitializer::EPSOPrecacheCompileType InPri, bool bInFlight) const
+		else
 		{
-			FPersistentStats& PersistentStats = GetStats();
+			PersistentStats.PSOGraphicsPrecompilesInProgress->Add(-1);
+		}
 
-			switch (InPri)
-			{
-			case FGraphicsPipelineStateInitializer::EPSOPrecacheCompileType::MaxPri:
-				return bInFlight ? PersistentStats.MaxPriPrecacheInFlight : PersistentStats.MaxPriPSOPrecompileQueueLength;
-			case FGraphicsPipelineStateInitializer::EPSOPrecacheCompileType::NormalPri:
-				return bInFlight ? PersistentStats.NormalPriPrecacheInFlight : PersistentStats.NormalPriPSOPrecompileQueueLength ;
-			case FGraphicsPipelineStateInitializer::EPSOPrecacheCompileType::MinPri:
-				return bInFlight ? PersistentStats.MinPriPrecacheInFlight : PersistentStats.MinPriPSOPrecompileQueueLength;
+		EPSOPrecachePriority CurrentPriority = (EPSOPrecachePriority)PSOPrecachePriority.load();
+		if (ensure(CurrentPriority != EPSOPrecachePriority::NotSet))
+		{
+			GetPriorityStat(CurrentPriority, EStatType::Queued)->Add(-1);
+			GetPriorityStat(CurrentPriority, EStatType::InFlight)->Add(-1);
+		}
+
+		bProcessedFinished = true;
+	}
+
+private:
+
+	TCsvPersistentCustomStat < int > * GetPriorityStat(EPSOPrecachePriority InPriority, EStatType StatType) const
+	{
+		FPersistentStats& PersistentStats = GetPersistentStats();
+
+		switch (InPriority)
+		{
+			case EPSOPrecachePriority::Low:
+				return StatType == EStatType::InFlight ? PersistentStats.LowPriorityPrecacheInFlight : PersistentStats.LowPriorityPSOPrecompileQueueLength;				
+			case EPSOPrecachePriority::Medium:
+				return StatType == EStatType::InFlight ? PersistentStats.MediumPriorityPrecacheInFlight : PersistentStats.MediumPriorityPSOPrecompileQueueLength;				
+			case EPSOPrecachePriority::High:
+				return StatType == EStatType::InFlight ? PersistentStats.HighPriorityPrecacheInFlight : PersistentStats.HighPriorityPSOPrecompileQueueLength;				
+			case EPSOPrecachePriority::Highest:
+				return StatType == EStatType::InFlight ? PersistentStats.HighestPriorityPrecacheInFlight : PersistentStats.HighestPriorityPSOPrecompileQueueLength;
 			default:
 				checkNoEntry();
 				return nullptr;
-			}
 		}
+	}
 };
 
 static FAutoConsoleVariableRef CVarPSOPrecacheEnableCompileStats(
@@ -994,9 +1061,10 @@ class FPSOCSVStatTracker
 {
 public:
 	FPSOCSVStatTracker() {}
-	void SetState(bool bInPrecache, bool bInCompute, FGraphicsPipelineStateInitializer::EPSOPrecacheCompileType InPri) {}
-	void SetPriority(FGraphicsPipelineStateInitializer::EPSOPrecacheCompileType InPri) {}
-	void SetInFlightPriority(FGraphicsPipelineStateInitializer::EPSOPrecacheCompileType InPri) {}
+	void SetState(bool bInPrecache, bool bInCompute, EPSOPrecachePriority InPriority) {}
+	void SetQueuedPriority(EPSOPrecachePriority InPriority) {}
+	void SetInFlightPriority(EPSOPrecachePriority InPriority) {}
+	void PrecacheFinished() {}
 };
 #endif
 
@@ -2382,13 +2450,11 @@ private:
 enum class EPSOPrecacheStateMask : uint8
 {
 	None = 0,
-	Compiling = 1 << 0, // once the start is scheduled
-	Cached = 1 << 1, 		// PSO is cached
-	Uncached = 1 << 2,  	// PSO is released
-	Failed = 1 << 3,    	// PSO failed to compile
-	Boosted = 1 << 4,		// PSO has been boosted to high
-	HighestPri = 1 << 5,	// PSO has been boosted to highest
-	UsedForRendering = 1 << 6, // Set once this precached PSO is actually needed and used for rendering.
+	Compiling = 1 << 0, 		// once the start is scheduled
+	Cached = 1 << 1, 			// PSO is cached in memory
+	Uncached = 1 << 2,  		// PSO is released
+	Failed = 1 << 3,    		// PSO failed to compile
+	UsedForRendering = 1 << 4, 	// Set once this precached PSO is actually needed and used for rendering.
 };
 
 ENUM_CLASS_FLAGS(EPSOPrecacheStateMask)
@@ -2441,49 +2507,32 @@ public:
 	
 protected:
 
-	void RescheduleTaskToHighPriority(EPSOPrecacheStateMask NewState, EPSOPrecacheStateMask PrevState, TPipelineState* PipelineState)
+	void RescheduleTaskToNewPriority(EPSOPrecacheStateMask CurrentState, EPSOPrecachePriority NewPriority, EPSOPrecachePriority PrevPriority, TPipelineState* PipelineState)
 	{
-		bool bHighestPriority = EnumHasAnyFlags(NewState, EPSOPrecacheStateMask::HighestPri);
-		bool bWasPreviouslyHigh = EnumHasAnyFlags(PrevState, EPSOPrecacheStateMask::Boosted);
+		check(PrevPriority != EPSOPrecachePriority::Highest);
+		check((int32)NewPriority > (int32)PrevPriority);
 
-		check(!EnumHasAnyFlags(PrevState, EPSOPrecacheStateMask::HighestPri));
-
-		bool bCompleted = EnumHasAnyFlags(PrevState, EPSOPrecacheStateMask::Failed | EPSOPrecacheStateMask::Cached);
-		UE_CLOG(bCompleted, LogRHI, Error, TEXT("pso request has completed? prev %x, new %x"), (uint32)PrevState, (uint32)NewState);
+		bool bCompleted = EnumHasAnyFlags(CurrentState, EPSOPrecacheStateMask::Failed | EPSOPrecacheStateMask::Cached);
+		UE_CLOG(bCompleted, LogRHI, Error, TEXT("pso request has completed? prev %x, new %x"), (uint32)PrevPriority, (uint32)NewPriority);
 
 		if (FPSOPrecacheThreadPool::UsePool())
 		{
 			check(PipelineState->GetPrecompileTask());
-			EQueuedWorkPriority NewPriority = bHighestPriority ? EQueuedWorkPriority::Highest : EQueuedWorkPriority::High;
 			if(PipelineState->GetPrecompileTask())
-			{
-				EQueuedWorkPriority PrevPriority = PipelineState->GetPrecompileTask()->GetPriority();
-				check(PrevPriority > NewPriority);
-				if( PipelineState->GetPrecompileTask()->Reschedule(&GPSOPrecacheThreadPool.Get(), NewPriority) )
+			{				
+				EQueuedWorkPriority NewQueuePriority = GetQueuedWorkPriorityFromPSOPrecachePriority(NewPriority);
+				if( PipelineState->GetPrecompileTask()->Reschedule(&GPSOPrecacheThreadPool.Get(), NewQueuePriority) )
 				{
-					PipelineState->CSVStat.SetPriority(GetPSOCompileTypeFromQueuePri(NewPriority));
+					PipelineState->CSVStat.SetQueuedPriority(NewPriority);
 				}
 			}
 		}
-
-		if (bHighestPriority)
-		{
-			UpdateHighestPriorityCompileCount(true);
-			if (bWasPreviouslyHigh)
-			{
-				UpdateHighPriorityCompileCount(false /*decrement*/);
-			}
-		}
-		else
-		{
-			UpdateHighPriorityCompileCount(true);
-		}
 	}
 
-	FPSOPrecacheRequestResult TryAddNewState(const TPrecachedPSOInitializer& Initializer, EPSOPrecacheInMemoryMode InMemoryMode, const FString& PSOCompilationEventName, bool bDoAsyncCompile)
+	FPSOPrecacheRequestResult TryAddNewState(const TPrecachedPSOInitializer& Initializer, const FPSOPrecacheRequestParams& RequestParams, bool bDoAsyncCompile)
 	{
 		// Refcounted and max in memory PSOs can't be used together
-		check(InMemoryMode == EPSOPrecacheInMemoryMode::Uncached || !ShouldKeepPrecachedPSOsInMemory());
+		check(RequestParams.InMemoryMode == EPSOPrecacheInMemoryMode::Uncached || !ShouldKeepPrecachedPSOsInMemory());
 
 		FPSOPrecacheRequestResult Result;
 		uint64 InitializerHash = TPrecachePipelineCacheDerived::PipelineStateInitializerHash(Initializer);
@@ -2491,7 +2540,7 @@ protected:
 		// Fast check first with read lock
 		{
 			FRWScopeLock ReadLock(PrecachePSOsRWLock, SLT_ReadOnly);
-			if (IsCachedOrCompiling(Initializer, InitializerHash, InMemoryMode, Result))
+			if (IsCachedOrCompiling(Initializer, InitializerHash, RequestParams.InMemoryMode, Result))
 			{
 				return Result;
 			}
@@ -2501,7 +2550,7 @@ protected:
 		TPipelineState* NewPipelineState = nullptr;
 		{
 			FRWScopeLock WriteLock(PrecachePSOsRWLock, SLT_Write);
-			if (IsCachedOrCompiling(Initializer, InitializerHash, InMemoryMode, Result))
+			if (IsCachedOrCompiling(Initializer, InitializerHash, RequestParams.InMemoryMode, Result))
 			{
 				return Result;
 			}
@@ -2512,7 +2561,7 @@ protected:
 			{			
 				check(!ShouldKeepPrecachedPSOsInMemory());
 
-				check(InMemoryMode == EPSOPrecacheInMemoryMode::Cached);				
+				check(RequestParams.InMemoryMode == EPSOPrecacheInMemoryMode::Cached);				
 				check(FindResult->RefCount == 0);
 
 				FPlatformAtomics::InterlockedIncrement(&FindResult->RefCount);									
@@ -2569,7 +2618,8 @@ protected:
 			NewPipelineState = TPrecachePipelineCacheDerived::CreateNewPSO(Initializer);			
 			check(FindResult->PipelineState == nullptr);
 			FindResult->PipelineState = NewPipelineState;
-			FindResult->InMemoryMode = InMemoryMode;
+			FindResult->InMemoryMode = RequestParams.InMemoryMode;
+			FindResult->SetPrecachePriority(RequestParams.PrecachePriority);
 						
 			if (bDoAsyncCompile)
 			{
@@ -2580,6 +2630,7 @@ protected:
 				Result.AsyncCompileEvent = NewPipelineState->GetCompletionEvent();
 
 				UpdateActiveCompileCount(true /*Increment*/);
+				UpdatePerPriorityCompileCount(RequestParams.PrecachePriority, true /*Increment*/);
 			}
 
 			if (ShouldKeepPrecachedPSOsInMemory())
@@ -2603,7 +2654,7 @@ protected:
 			}
 		}
 
-		TPrecachePipelineCacheDerived::OnNewPipelineStateCreated(Initializer, NewPipelineState, PSOCompilationEventName, bDoAsyncCompile);
+		TPrecachePipelineCacheDerived::OnNewPipelineStateCreated(Initializer, RequestParams.PrecachePriority, NewPipelineState, RequestParams.Name, bDoAsyncCompile);
 
 		// A boost request might have been issued while we were kicking the task, need to check it here
 		{
@@ -2613,10 +2664,17 @@ protected:
 			if (FindResult != nullptr)
 			{
 				EPSOPrecacheStateMask PreviousStateMask = FindResult->AddPSOPrecacheState(EPSOPrecacheStateMask::Compiling);
-				// by the time we're here, PrecacheFinished might already have been called, so boost it only if we know we will call it
-				if (!IsCompilationDone(PreviousStateMask) && EnumHasAnyFlags(PreviousStateMask, EPSOPrecacheStateMask::Boosted) )
+
+				// The current priority is the potentially new (higher) priority while RequestParams.PrecachePriority was the original requested priority
+				// The priority might have changed via BoostPriority but not processed yet because the task might not have been kicked yet, but the internal
+				// tracked priority if the precache task is already updated and checked here
+				// NOTE: Priority can only modified via write lock in SetPrecachePriorityInternal so safe to read here with read lock
+				EPSOPrecachePriority CurrentPriority = FindResult->GetPrecachePriority();
+
+				// by the time we're here, PrecacheFinished might already have been called, so only update priority if it will still be processed
+				if (!IsCompilationDone(PreviousStateMask) && CurrentPriority != RequestParams.PrecachePriority)
 				{
-					RescheduleTaskToHighPriority(PreviousStateMask, (EPSOPrecacheStateMask)0, FindResult->PipelineState);
+					RescheduleTaskToNewPriority(PreviousStateMask, CurrentPriority, RequestParams.PrecachePriority, FindResult->PipelineState);
 				}
 			}
 		}
@@ -2651,15 +2709,15 @@ public:
 							// boost the priority of anything we're waiting for to improve throughput.
 							// ensurecomplete is guaranteed to complete, but We're not explicitly waiting for precache tasks here.
 							// So lower pri tasks may never complete when r.PSOPrecaching.TaskPriorityThreshold is in use.
-							if (!EnumHasAnyFlags(PrecacheTask.ReadPSOPrecacheState(), EPSOPrecacheStateMask::HighestPri))
+							if (PrecacheTask.GetPrecachePriority() != EPSOPrecachePriority::Highest)
 							{
-								BoostPriorityInternal(PrecacheTask, EPSOPrecachePriority::Highest);
+								SetPrecachePriorityInternal(PrecacheTask, EPSOPrecachePriority::Highest);
 							}
 						}
 						else
 						{
 							// everything should be boosted on the second pass.
-							check(EnumHasAnyFlags(PrecacheTask.ReadPSOPrecacheState(), EPSOPrecacheStateMask::HighestPri));
+							check(PrecacheTask.GetPrecachePriority() == EPSOPrecachePriority::Highest);
 							break; // release PrecachePSOsRWLock so as to avoid any further blocking of in-progress tasks.
 						}
 					}
@@ -2786,7 +2844,7 @@ public:
 		FPrecacheTask* FindResult = PrecachedPSOInitializerData.Find(InitializerHash);
 		check(FindResult);
 
-		BoostPriorityInternal(*FindResult, PSOPrecachePriority);
+		SetPrecachePriorityInternal(*FindResult, PSOPrecachePriority);
 	}
 
 	uint32 NumActivePrecacheRequests()
@@ -2797,9 +2855,9 @@ public:
 			case 0:
 				return FPlatformAtomics::AtomicRead(&ActiveCompileCount);
 			case 1:
-				return FPlatformAtomics::AtomicRead(&HighPriorityCompileCount) + FPlatformAtomics::AtomicRead(&HighestPriorityCompileCount);
+				return FPlatformAtomics::AtomicRead(&PerPriorityCompileCount[(int32)EPSOPrecachePriority::Highest]) + FPlatformAtomics::AtomicRead(&PerPriorityCompileCount[(int32)EPSOPrecachePriority::High]);
 			case 2:
-				return FPlatformAtomics::AtomicRead(&HighestPriorityCompileCount);
+				return FPlatformAtomics::AtomicRead(&PerPriorityCompileCount[(int32)EPSOPrecachePriority::Highest]);
 		}
 	}
 
@@ -2807,7 +2865,9 @@ public:
 	{
 		uint64 InitializerHash = TPrecachePipelineCacheDerived::PipelineStateInitializerHash(Initializer);
 
-		EPSOPrecacheStateMask PreviousStateMask = EPSOPrecacheStateMask::None;
+		EPSOPrecacheStateMask PreviousStateMask = EPSOPrecacheStateMask::None;	
+		EPSOPrecachePriority StatePriority = EPSOPrecachePriority::Medium;
+
 		{
 			FRWScopeLock WriteLock(PrecachePSOsRWLock, SLT_Write);
 
@@ -2823,10 +2883,14 @@ public:
 			}
 #endif
 
+			FindResult->PipelineState->CSVStat.PrecacheFinished();
+
 			// We still add the 'compiling' bit here because if the task is fast enough, we can get here before the end of TryAddNewState
 			const EPSOPrecacheStateMask CompleteStateMask = bValid ? (EPSOPrecacheStateMask::Cached | EPSOPrecacheStateMask::Compiling) : (EPSOPrecacheStateMask::Failed | EPSOPrecacheStateMask::Compiling);
 			PreviousStateMask = FindResult->AddPSOPrecacheState(CompleteStateMask);
 
+			StatePriority = FindResult->GetPrecachePriority();
+
 			// Release immediately if not kept in memory
 			if (FindResult->InMemoryMode != EPSOPrecacheInMemoryMode::Cached && !ShouldKeepPrecachedPSOsInMemory())
 			{
@@ -2851,18 +2915,8 @@ public:
 
         // Need to ensure that the boost request was actually executed: if only it was asked by BoostPriority, but not requested (ie TryAddNewState has not set the Compiling bit
         // yet) then we must ignore the request
-		if (EnumHasAllFlags(PreviousStateMask, EPSOPrecacheStateMask::Boosted | EPSOPrecacheStateMask::Compiling))
-		{
-			if (EnumHasAnyFlags(PreviousStateMask, EPSOPrecacheStateMask::HighestPri))
-			{
-				UpdateHighestPriorityCompileCount(false  /*Increment*/);
-			}
-			else
-			{
-				UpdateHighPriorityCompileCount(false /*Increment*/);
-			}
-		}
 		UpdateActiveCompileCount(false /*Increment*/);
+		UpdatePerPriorityCompileCount(StatePriority, false /*Increment*/);
 	}
 
 	static bool IsCompilationDone(EPSOPrecacheStateMask StateMask)
@@ -2874,8 +2928,11 @@ public:
 	void ProcessDelayedCleanup()
 	{
 		SET_DWORD_STAT_FName(TPrecachePipelineCacheDerived::GetActiveCompileStatName(), ActiveCompileCount);
-		SET_DWORD_STAT_FName(TPrecachePipelineCacheDerived::GetHighPriorityCompileStatName(), HighPriorityCompileCount);
-		SET_DWORD_STAT_FName(TPrecachePipelineCacheDerived::GetHighestPriorityCompileStatName(), HighestPriorityCompileCount);
+		for (int32 PrecachePriority = 0; PrecachePriority < (int32)EPSOPrecachePriority::Num; ++PrecachePriority)
+		{
+			SET_DWORD_STAT_FName(TPrecachePipelineCacheDerived::GetPerPriorityCompileStatName((EPSOPrecachePriority)PrecachePriority), PerPriorityCompileCount[PrecachePriority]);
+		}
+
 		FRWScopeLock WriteLock(PrecachePSOsRWLock, SLT_Write);
 		for (int32 Index = 0; Index < PrecachedPSOsToCleanup.Num(); ++Index)
 		{
@@ -3053,29 +3110,16 @@ protected:
 		}
 	}
 
-	void UpdateHighPriorityCompileCount(bool bIncrement)
+	void UpdatePerPriorityCompileCount(EPSOPrecachePriority PrecachePriority, bool bIncrement)
 	{
 		if (bIncrement)
 		{
-			FPlatformAtomics::InterlockedIncrement(&HighPriorityCompileCount);
+			FPlatformAtomics::InterlockedIncrement(&PerPriorityCompileCount[(int32)PrecachePriority]);
 		}
 		else
 		{
-			FPlatformAtomics::InterlockedDecrement(&HighPriorityCompileCount);
-			check(HighPriorityCompileCount >= 0);
-		}
-	}
-
-	void UpdateHighestPriorityCompileCount(bool bIncrement)
-	{
-		if (bIncrement)
-		{
-			FPlatformAtomics::InterlockedIncrement(&HighestPriorityCompileCount);
-		}
-		else
-		{
-			FPlatformAtomics::InterlockedDecrement(&HighestPriorityCompileCount);
-			check(HighestPriorityCompileCount >= 0);
+			FPlatformAtomics::InterlockedDecrement(&PerPriorityCompileCount[(int32)PrecachePriority]);
+			check(PerPriorityCompileCount[(int32)PrecachePriority] >= 0);
 		}
 	}
 
@@ -3123,19 +3167,28 @@ protected:
 			return (EPSOPrecacheStateMask)FPlatformAtomics::AtomicRead((volatile int8*)&StateMask);
 		}
 
+		inline EPSOPrecachePriority GetPrecachePriority()
+		{
+			return (EPSOPrecachePriority)FPlatformAtomics::AtomicRead((volatile int8*)&PrecachePriority);
+		}
+
+		inline void SetPrecachePriority(EPSOPrecachePriority NewPriority)
+		{
+			return FPlatformAtomics::AtomicStore((volatile int8*)&PrecachePriority, (int8)NewPriority);
+		}
+
 	private:
 		volatile EPSOPrecacheStateMask StateMask = EPSOPrecacheStateMask::None;
+		volatile EPSOPrecachePriority PrecachePriority = EPSOPrecachePriority::Medium;
 	};
 
 	TMap<uint64, FPrecacheTask> PrecachedPSOInitializerData;
 
-	// Number of open active compiles
+	// Total number of open active compiles
 	volatile int32 ActiveCompileCount = 0;
 
-	// Number of open high priority compiles
-	volatile int32 HighPriorityCompileCount = 0;
-
-	volatile int32 HighestPriorityCompileCount = 0;
+	// Number of open compiles per priority type
+	volatile int32 PerPriorityCompileCount[(int32)EPSOPrecachePriority::Num] = { 0, 0, 0, 0 };
 
 	// Finished Precached PSOs which can be garbage collected
 	TArray<uint64> PrecachedPSOsToCleanup;
@@ -3144,18 +3197,33 @@ protected:
 	TDeque<uint32> InMemoryPSOIndices;
 	uint32 MaxInMemoryPSOs = 0;
 
-	void BoostPriorityInternal(FPrecacheTask& FindResult, EPSOPrecachePriority PSOPrecachePriority)
+	void SetPrecachePriorityInternal(FPrecacheTask& FindResult, EPSOPrecachePriority NewPSOPrecachePriority)
 	{
-		// Assume PrecachePSOsRWLock is locked and FindResult will not change.
-		EPSOPrecacheStateMask NewMask = EPSOPrecacheStateMask::Boosted | (PSOPrecachePriority == EPSOPrecachePriority::Highest ? EPSOPrecacheStateMask::HighestPri : (EPSOPrecacheStateMask)0);
-		EPSOPrecacheStateMask PreviousStateMask = FindResult.AddPSOPrecacheState(NewMask);
-		// It's possible to get a boost request while the task has not been started yet. In this case, TryAddNewState will take care of it
-		// if TryAddNewState is done, then we can proceed to boost it, if the task is not done yet
-		if (!IsCompilationDone(PreviousStateMask) && EnumHasAnyFlags(PreviousStateMask, EPSOPrecacheStateMask::Compiling) && !EnumHasAnyFlags(PreviousStateMask, EPSOPrecacheStateMask::HighestPri))
+		// Ignore requests which try to lower the priority again 
+		// (can happen with material side boost but PSO already boosted to highest because required for rendering by other component)
+		EPSOPrecachePriority CurrentPriority = FindResult.GetPrecachePriority();
+		if ((int32)CurrentPriority >= (int32)NewPSOPrecachePriority)
+		{
+			return;
+		}
+
+		// It's possible to get a priority change request while the task has not been started yet. In this case, TryAddNewState will take care of it
+		// if TryAddNewState is done, then we can proceed to process the request it, if the task is not done yet
+		EPSOPrecacheStateMask CurrentState = FindResult.ReadPSOPrecacheState();
+		if (!IsCompilationDone(CurrentState))
 		{
-			if (!EnumHasAnyFlags(PreviousStateMask, EPSOPrecacheStateMask::Boosted) || EnumHasAnyFlags(NewMask, EPSOPrecacheStateMask::HighestPri))
+			// Only update priority when not done yet
+			FindResult.SetPrecachePriority(NewPSOPrecachePriority);
+
+			// Update the stats here even though the reschedule might not be processed yet because it can happen
+			// that the actual reschedule never happens if it's done by the time the TryAddNewState actually
+			// tried the reschedule and stats need to be correct to process process finished call
+			UpdatePerPriorityCompileCount(CurrentPriority, false);
+			UpdatePerPriorityCompileCount(NewPSOPrecachePriority, true);
+
+			if (EnumHasAnyFlags(CurrentState, EPSOPrecacheStateMask::Compiling))
 			{
-				RescheduleTaskToHighPriority(NewMask, PreviousStateMask, FindResult.PipelineState);
+				RescheduleTaskToNewPriority(CurrentState, NewPSOPrecachePriority, CurrentPriority, FindResult.PipelineState);
 			}
 		}
 	}
@@ -3172,20 +3240,25 @@ public:
 	}
 
 	FPrecacheComputePipelineCache(uint32 InMemoryPSOsMaxNum) : TPrecachePipelineCacheBase(FPSOPrecacheRequestID::EType::Compute, InMemoryPSOsMaxNum) {}
-	FPSOPrecacheRequestResult PrecacheComputePipelineState(const FComputePipelineStateInitializer& ComputeInitializer, EPSOPrecacheInMemoryMode InMemoryMode, const TCHAR* Name, bool bForcePrecache);
-	static void OnNewPipelineStateCreated(const FComputePipelineStateInitializer& ComputeInitializer, FComputePipelineState* NewComputePipelineState, const FString& PSOCompilationEventName, bool bDoAsyncCompile);
+	FPSOPrecacheRequestResult PrecacheComputePipelineState(const FComputePipelineStateInitializer& ComputeInitializer, const FPSOPrecacheRequestParams& RequestParams);
+	static void OnNewPipelineStateCreated(const FComputePipelineStateInitializer& ComputeInitializer, EPSOPrecachePriority PrecachePriority, FComputePipelineState* NewComputePipelineState, const FString& PSOCompilationEventName, bool bDoAsyncCompile);
  
 	static const FName GetActiveCompileStatName()
 	{
 		return GET_STATFNAME(STAT_ActiveComputePSOPrecacheRequests);
 	}
-	static const FName GetHighPriorityCompileStatName()
-	{
-		return GET_STATFNAME(STAT_HighPriorityComputePSOPrecacheRequests);
-	}
-	static const FName GetHighestPriorityCompileStatName()
+	static const FName GetPerPriorityCompileStatName(EPSOPrecachePriority PrecachePriority)
 	{
-		return GET_STATFNAME(STAT_HighestPriorityComputePSOPrecacheRequests);
+		switch (PrecachePriority)
+		{
+			case EPSOPrecachePriority::Highest: return GET_STATFNAME(STAT_HighestPriorityComputePSOPrecacheRequests);
+			case EPSOPrecachePriority::High:  	return GET_STATFNAME(STAT_HighPriorityComputePSOPrecacheRequests);
+			case EPSOPrecachePriority::Medium:  return GET_STATFNAME(STAT_MediumPriorityComputePSOPrecacheRequests);
+			default: checkNoEntry(); [[fallthrough]];
+			case EPSOPrecachePriority::Low:  	return GET_STATFNAME(STAT_LowPriorityComputePSOPrecacheRequests);
+		}
+
+		checkNoEntry();
 	}
 
 	static FORCEINLINE uint64 PipelineStateInitializerHash(const FComputePipelineStateInitializer& Key)
@@ -3214,18 +3287,23 @@ public:
 	{
 		return GET_STATFNAME(STAT_ActiveGraphicsPSOPrecacheRequests);
 	}
-	static const FName GetHighPriorityCompileStatName()
+	static const FName GetPerPriorityCompileStatName(EPSOPrecachePriority PrecachePriority)
 	{
-		return GET_STATFNAME(STAT_HighPriorityGraphicsPSOPrecacheRequests);
-	}
-	static const FName GetHighestPriorityCompileStatName()
-	{
-		return GET_STATFNAME(STAT_HighestPriorityGraphicsPSOPrecacheRequests);
+		switch (PrecachePriority)
+		{
+			case EPSOPrecachePriority::Highest: return GET_STATFNAME(STAT_HighestPriorityGraphicsPSOPrecacheRequests);
+			case EPSOPrecachePriority::High:  	return GET_STATFNAME(STAT_HighPriorityGraphicsPSOPrecacheRequests);
+			case EPSOPrecachePriority::Medium:  return GET_STATFNAME(STAT_MediumPriorityGraphicsPSOPrecacheRequests);
+			default: checkNoEntry(); [[fallthrough]];
+			case EPSOPrecachePriority::Low:  	return GET_STATFNAME(STAT_LowPriorityGraphicsPSOPrecacheRequests);
+		}
+
+		checkNoEntry();
 	}
 	FPrecacheGraphicsPipelineCache(uint32 InMemoryPSOsMaxNum) : TPrecachePipelineCacheBase(FPSOPrecacheRequestID::EType::Graphics, InMemoryPSOsMaxNum) {}
-	FPSOPrecacheRequestResult PrecacheGraphicsPipelineState(const FGraphicsPipelineStateInitializer& Initializer, EPSOPrecacheInMemoryMode InMemoryMode);
+	FPSOPrecacheRequestResult PrecacheGraphicsPipelineState(const FGraphicsPipelineStateInitializer& Initializer, const FPSOPrecacheRequestParams& RequestParams);
 
-	static void OnNewPipelineStateCreated(const FGraphicsPipelineStateInitializer& Initializer, FGraphicsPipelineState* NewGraphicsPipelineState, const FString& PSOCompilationEventName, bool bDoAsyncCompile);
+	static void OnNewPipelineStateCreated(const FGraphicsPipelineStateInitializer& Initializer, EPSOPrecachePriority PrecachePriority, FGraphicsPipelineState* NewGraphicsPipelineState, const FString& PSOCompilationEventName, bool bDoAsyncCompile);
 
 };
 
@@ -3498,11 +3576,11 @@ protected:
 	FPipelineStateAsync* Pipeline;
 	TPipelineInitializer Initializer;
 	EPSOPrecacheResult PSOPrecacheResult;
+	EPSOPrecachePriority PSOPrecachePriority = EPSOPrecachePriority::NotSet;
 	bool bPrecachedPSOMarkedUsed;
 	bool bInImmediateCmdList;
 	FPSOCompilationDebugData PSOCompilationDebugData;
 
-
 public:
 	TCompilePipelineStateTaskBase(
 		FPipelineStateAsync* InPipeline,
@@ -3535,12 +3613,18 @@ public:
 
 	static constexpr ESubsequentsMode::Type GetSubsequentsMode() { return ESubsequentsMode::TrackSubsequents; }
 
+	void SetPSOPrecachePriority(EPSOPrecachePriority TaskPriority)
+	{
+		check(PSOPrecachePriority == EPSOPrecachePriority::NotSet);
+		PSOPrecachePriority = TaskPriority;
+	}
+
 	void DoTask(ENamedThreads::Type CurrentThread, const FGraphEventRef& MyCompletionGraphEvent)
 	{
 		CompilePSO();
 	}
 
-	void CompilePSO(const FGraphicsPipelineStateInitializer::EPSOPrecacheCompileType* OptionalPriorityOverride = nullptr)
+	void CompilePSO()
 	{
 		LLM_SCOPE(ELLMTag::PSO);
 		FTaskTagScope Scope(GetDesiredThread() == ENamedThreads::RHIThread ? ETaskTag::EParallelRhiThread : ETaskTag::EParallelRenderingThread);
@@ -3570,7 +3654,7 @@ public:
 		TRACE_CPUPROFILER_EVENT_SCOPE_TEXT(PSOPrecacheResultScopeString);
 		{
 			TRACE_CPUPROFILER_EVENT_SCOPE_TEXT_CONDITIONAL(*PSOCompilationDebugData.PSOCompilationEventName, !PSOCompilationDebugData.PSOCompilationEventName.IsEmpty())
-			static_cast<TCompilePipelineStateTaskDerived*>(this)->CompilePSOInternal(OptionalPriorityOverride);
+			static_cast<TCompilePipelineStateTaskDerived*>(this)->CompilePSOInternal();
 		}
 
 #if WITH_RHI_BREADCRUMBS
@@ -3688,7 +3772,7 @@ public:
 		}
 	}
 
-	void CompilePSOInternal(const FGraphicsPipelineStateInitializer::EPSOPrecacheCompileType* OptionalPriorityOverride = nullptr)
+	void CompilePSOInternal()
 	{
 		bool bSkipCreation = false;
 		if (GRHISupportsMeshShadersTier0)
@@ -3721,11 +3805,12 @@ public:
 			bSkipCreation = true;
 		}
 
-		if (OptionalPriorityOverride)
+		if (Initializer.bPSOPrecache)
 		{
-			Initializer.PrecacheCompileType = (uint32)FMath::Clamp((uint32)*OptionalPriorityOverride, (uint32)FGraphicsPipelineStateInitializer::EPSOPrecacheCompileType::MinPri, (uint32)FGraphicsPipelineStateInitializer::EPSOPrecacheCompileType::MaxPri);
+			check(PSOPrecachePriority != EPSOPrecachePriority::NotSet);
+			Initializer.PrecacheCompileType = (uint32)GetPSOCompileTypeFromPrecachePriority(PSOPrecachePriority);			
+			Pipeline->CSVStat.SetInFlightPriority(PSOPrecachePriority);
 		}
-		Pipeline->CSVStat.SetInFlightPriority((FGraphicsPipelineStateInitializer::EPSOPrecacheCompileType)Initializer.PrecacheCompileType);
 
 		FGraphicsPipelineState* GfxPipeline = static_cast<FGraphicsPipelineState*>(Pipeline);
 
@@ -3841,20 +3926,18 @@ public:
 #endif // PLATFORM_WINDOWS
 	}
 
-	void CompilePSOInternal(const FGraphicsPipelineStateInitializer::EPSOPrecacheCompileType* OptionalPriorityOverride = nullptr)
+	void CompilePSOInternal()
 	{
 		FComputePipelineState* ComputePipeline = static_cast<FComputePipelineState*>(Pipeline);
 
 		uint64 StartTime = FPlatformTime::Cycles64();
-		using EPSOPrecacheCompileType = FGraphicsPipelineStateInitializer::EPSOPrecacheCompileType;
 
-		// when r.PSOPrecaching.TaskPriorityThreshold is active use pri threshold for precache compute PSOs to ensure they are processed.
-		EPSOPrecacheCompileType PSOCompileType = GPSOPrecacheProcessPriorityThreshold ? (EPSOPrecacheCompileType)GPSOPrecacheProcessPriorityThreshold : EPSOPrecacheCompileType::NormalPri;
-		if (OptionalPriorityOverride)
+		if (Initializer.bPSOPrecache)
 		{
-			PSOCompileType = (EPSOPrecacheCompileType)FMath::Clamp((uint32)*OptionalPriorityOverride, (uint32)EPSOPrecacheCompileType::MinPri, (uint32)EPSOPrecacheCompileType::MaxPri);
+			check(PSOPrecachePriority != EPSOPrecachePriority::NotSet);
+			//Initializer.PrecacheCompileType = (uint32)GetPSOCompileTypeFromPrecachePriority(PSOPrecachePriority);			
+			Pipeline->CSVStat.SetInFlightPriority(PSOPrecachePriority);
 		}
-		Pipeline->CSVStat.SetInFlightPriority(PSOCompileType);
 
 		// TODO: Send the priority to the RHI.
 		ComputePipeline->RHIPipeline = RHICreateComputePipelineState(Initializer);
@@ -3994,7 +4077,7 @@ uint64 PipelineStateCache::RetrieveGraphicsPipelineStateSortKey(const FGraphicsP
 }
 
 template <typename TCompilePipelineTask, typename TPipelineInitializer>
-static void InternalAsyncCompilePipelineState(const TPipelineInitializer& Initializer, EPSOPrecacheResult PSOPrecacheResult, bool bPrecachedPSOMarkedUsed, FPipelineStateAsync* CachedState, const FPSOCompilationDebugData& PSOCompilationDebugData, bool bInImmediateCmdList)
+static void InternalAsyncCompilePipelineState(const TPipelineInitializer& Initializer, EPSOPrecacheResult PSOPrecacheResult, EPSOPrecachePriority PSOPrecachePriority, bool bPrecachedPSOMarkedUsed, FPipelineStateAsync* CachedState, const FPSOCompilationDebugData& PSOCompilationDebugData, bool bInImmediateCmdList)
 {
 	if (!Initializer.bPSOPrecache || !FPSOPrecacheThreadPool::UsePool())
 	{
@@ -4010,25 +4093,20 @@ static void InternalAsyncCompilePipelineState(const TPipelineInitializer& Initia
 			{
 				// Convert the task priority to PSO precompile priority.
 				// Update here as the task's priority may have changed since creation.
-				FGraphicsPipelineStateInitializer::EPSOPrecacheCompileType PriOverride = GetPSOCompileTypeFromQueuePri(ThisTask->GetPriority());
-				ThreadPoolTask->CompilePSO(&PriOverride);
+				EPSOPrecachePriority TaskPrecachePriority = GetPSOPrecachePriorityFromQueuePri(ThisTask->GetPriority());
+				ThreadPoolTask->SetPSOPrecachePriority(TaskPrecachePriority);
+				ThreadPoolTask->CompilePSO();
 			}
 		));
 
-		EQueuedWorkPriority QueuedWorkPriority = EQueuedWorkPriority::Normal;
-
-		// Graphics pipeline initializers support overriding the priority for PSO precaching, so start them at low.
-		if constexpr (std::is_same_v<TPipelineInitializer, FGraphicsPipelineStateInitializer>)
-		{
-			QueuedWorkPriority = Initializer.bPSOPrecache ? EQueuedWorkPriority::Low : EQueuedWorkPriority::Normal;
-		}
+		CachedState->CSVStat.SetState(Initializer.bPSOPrecache, CachedState->IsCompute(), PSOPrecachePriority);
 
-		CachedState->CSVStat.SetState(Initializer.bPSOPrecache, CachedState->IsCompute(), GetPSOCompileTypeFromQueuePri(QueuedWorkPriority));
+		EQueuedWorkPriority QueuedWorkPriority = GetQueuedWorkPriorityFromPSOPrecachePriority(PSOPrecachePriority);		
 		CachedState->GetPrecompileTask()->StartBackgroundTask(&GPSOPrecacheThreadPool.Get(), QueuedWorkPriority);
 	}
 }
 
-static void InternalCreateComputePipelineState(const FComputePipelineStateInitializer& Initializer, bool bDoAsyncCompile, EPSOPrecacheResult PSOPrecacheResult, bool bPrecachedPSOMarkedUsed, FComputePipelineState* CachedState, const FPSOCompilationDebugData& PSOCompilationDebugData, bool bInImmediateCmdList)
+static void InternalCreateComputePipelineState(const FComputePipelineStateInitializer& Initializer, bool bDoAsyncCompile, EPSOPrecacheResult PSOPrecacheResult, EPSOPrecachePriority PrecachePriority, bool bPrecachedPSOMarkedUsed, FComputePipelineState* CachedState, const FPSOCompilationDebugData& PSOCompilationDebugData, bool bInImmediateCmdList)
 {
 	FGraphEventRef GraphEvent = CachedState->GetCompletionEvent();
 
@@ -4036,7 +4114,7 @@ static void InternalCreateComputePipelineState(const FComputePipelineStateInitia
 	if (bDoAsyncCompile)
 	{
 		check(GraphEvent != nullptr);
-		InternalAsyncCompilePipelineState<FCompileComputePipelineStateTask, FComputePipelineStateInitializer>(Initializer, PSOPrecacheResult, bPrecachedPSOMarkedUsed, CachedState, PSOCompilationDebugData, bInImmediateCmdList);
+		InternalAsyncCompilePipelineState<FCompileComputePipelineStateTask, FComputePipelineStateInitializer>(Initializer, PSOPrecacheResult, PrecachePriority, bPrecachedPSOMarkedUsed, CachedState, PSOCompilationDebugData, bInImmediateCmdList);
 	}
 	else
 	{
@@ -4125,7 +4203,7 @@ FComputePipelineState* PipelineStateCache::GetAndOrCreateComputePipelineState(FR
 #endif // WITH_RHI_BREADCRUMBS
 		
 		FGraphEventRef GraphEvent = OutCachedState->GetCompletionEvent();
-		InternalCreateComputePipelineState(ComputeInitializer, DoAsyncCompile, PSOPrecacheResult, bPrecachedPSOMarkedUsed, OutCachedState, PSOCompilationDebugData, RHICmdList.IsImmediate());
+		InternalCreateComputePipelineState(ComputeInitializer, DoAsyncCompile, PSOPrecacheResult, EPSOPrecachePriority::Medium, bPrecachedPSOMarkedUsed, OutCachedState, PSOCompilationDebugData, RHICmdList.IsImmediate());
 
 		// Don't add precached PSOs as a dispatch prerequisite. We don't need to wait for them to complete before the RHICmdList can be dispatched.
 		if (GraphEvent.IsValid() && !ComputeInitializer.bFromPSOFileCache)
@@ -4515,7 +4593,7 @@ inline void ValidateGraphicsPipelineStateInitializer(const FGraphicsPipelineStat
 	check(Initializer.DepthStencilState && Initializer.BlendState && Initializer.RasterizerState);
 }
 
-static void InternalCreateGraphicsPipelineState(const FGraphicsPipelineStateInitializer& Initializer, EPSOPrecacheResult PSOPrecacheResult, bool bPrecachedPSOMarkedUsed, bool bDoAsyncCompile, FGraphicsPipelineState* CachedState, const FPSOCompilationDebugData& PSOCompilationDebugData, bool bInImmediateCmdList)
+static void InternalCreateGraphicsPipelineState(const FGraphicsPipelineStateInitializer& Initializer, EPSOPrecacheResult PSOPrecacheResult, EPSOPrecachePriority PSOPrecachePriority, bool bPrecachedPSOMarkedUsed, bool bDoAsyncCompile, FGraphicsPipelineState* CachedState, const FPSOCompilationDebugData& PSOCompilationDebugData, bool bInImmediateCmdList)
 {
 	FGraphEventRef GraphEvent = CachedState->GetCompletionEvent();
 
@@ -4523,7 +4601,7 @@ static void InternalCreateGraphicsPipelineState(const FGraphicsPipelineStateInit
 	if (bDoAsyncCompile)
 	{
 		check(GraphEvent != nullptr);
-		InternalAsyncCompilePipelineState<FCompileGraphicsPipelineStateTask, FGraphicsPipelineStateInitializer>(Initializer, PSOPrecacheResult, bPrecachedPSOMarkedUsed, CachedState, PSOCompilationDebugData, bInImmediateCmdList);
+		InternalAsyncCompilePipelineState<FCompileGraphicsPipelineStateTask, FGraphicsPipelineStateInitializer>(Initializer, PSOPrecacheResult, PSOPrecachePriority, bPrecachedPSOMarkedUsed, CachedState, PSOCompilationDebugData, bInImmediateCmdList);
 	}
 	else
 	{		
@@ -4661,7 +4739,7 @@ FGraphicsPipelineState* PipelineStateCache::GetAndOrCreateGraphicsPipelineState(
 #endif // WITH_RHI_BREADCRUMBS
 
 		FGraphEventRef GraphEvent = OutCachedState->GetCompletionEvent();
-		InternalCreateGraphicsPipelineState(Initializer, PSOPrecacheResult, bPrecachedPSOMarkedUsed, DoAsyncCompile, OutCachedState, PSOCompilationDebugData, RHICmdList.IsImmediate());
+		InternalCreateGraphicsPipelineState(Initializer, PSOPrecacheResult, EPSOPrecachePriority::Medium, bPrecachedPSOMarkedUsed, DoAsyncCompile, OutCachedState, PSOCompilationDebugData, RHICmdList.IsImmediate());
 
 		// Add dispatch pre requisite for non precaching jobs only
 		//if (GraphEvent.IsValid() && (!bPSOPrecache || !FPSOPrecacheThreadPool::UsePool()))
@@ -4777,10 +4855,10 @@ bool PipelineStateCache::IsPSOPrecachingEnabled()
 #endif // WITH_EDITOR
 }
 
-FPSOPrecacheRequestResult FPrecacheComputePipelineCache::PrecacheComputePipelineState(const FComputePipelineStateInitializer& ComputeInitializer, EPSOPrecacheInMemoryMode InMemoryMode, const TCHAR* Name, bool bForcePrecache)
+FPSOPrecacheRequestResult FPrecacheComputePipelineCache::PrecacheComputePipelineState(const FComputePipelineStateInitializer& ComputeInitializer, const FPSOPrecacheRequestParams& RequestParams)
 {
 	FPSOPrecacheRequestResult Result;
-	if (!PipelineStateCache::IsPSOPrecachingEnabled() && !bForcePrecache)
+	if (!PipelineStateCache::IsPSOPrecachingEnabled() && !RequestParams.bForcePrecache)
 	{
 		return Result;
 	}
@@ -4795,10 +4873,10 @@ FPSOPrecacheRequestResult FPrecacheComputePipelineCache::PrecacheComputePipeline
 	FComputePipelineStateInitializer PrecacheComputeInitializer(ComputeInitializer);
 	PrecacheComputeInitializer.bPSOPrecache = true;
 	
-	return TryAddNewState(PrecacheComputeInitializer, InMemoryMode, Name, bDoAsyncCompile);
+	return TryAddNewState(PrecacheComputeInitializer, RequestParams, bDoAsyncCompile);
 }
 
-void FPrecacheComputePipelineCache::OnNewPipelineStateCreated(const FComputePipelineStateInitializer& ComputeInitializer, FComputePipelineState* CachedState, const FString& PSOCompilationEventName, bool bDoAsyncCompile)
+void FPrecacheComputePipelineCache::OnNewPipelineStateCreated(const FComputePipelineStateInitializer& ComputeInitializer, EPSOPrecachePriority PrecachePriority, FComputePipelineState* CachedState, const FString& PSOCompilationEventName, bool bDoAsyncCompile)
 {
 	check(ComputeInitializer.ComputeShader);
 	check(ComputeInitializer.bPSOPrecache);
@@ -4806,15 +4884,15 @@ void FPrecacheComputePipelineCache::OnNewPipelineStateCreated(const FComputePipe
 	FPSOCompilationDebugData PSOCompilationDebugData;
 	PSOCompilationDebugData.PSOCompilationEventName = PSOCompilationEventName;
 
-	InternalCreateComputePipelineState(ComputeInitializer, bDoAsyncCompile, EPSOPrecacheResult::Active, false, CachedState, PSOCompilationDebugData, false);
+	InternalCreateComputePipelineState(ComputeInitializer, bDoAsyncCompile, EPSOPrecacheResult::Active, PrecachePriority, false, CachedState, PSOCompilationDebugData, false);
 }
 
-FPSOPrecacheRequestResult PipelineStateCache::PrecacheComputePipelineState(const FComputePipelineStateInitializer& ComputeInitializer, EPSOPrecacheInMemoryMode InMemoryMode, const TCHAR* Name, bool bForcePrecache)
+FPSOPrecacheRequestResult PipelineStateCache::PrecacheComputePipelineState(const FComputePipelineStateInitializer& ComputeInitializer, const FPSOPrecacheRequestParams& RequestParams)
 {
-	return GPrecacheComputePipelineCache->PrecacheComputePipelineState(ComputeInitializer, InMemoryMode, Name, bForcePrecache);
+	return GPrecacheComputePipelineCache->PrecacheComputePipelineState(ComputeInitializer, RequestParams);
 }
 
-FPSOPrecacheRequestResult FPrecacheGraphicsPipelineCache::PrecacheGraphicsPipelineState(const FGraphicsPipelineStateInitializer& Initializer, EPSOPrecacheInMemoryMode InMemoryMode)
+FPSOPrecacheRequestResult FPrecacheGraphicsPipelineCache::PrecacheGraphicsPipelineState(const FGraphicsPipelineStateInitializer& Initializer, const FPSOPrecacheRequestParams& RequestParams)
 {
 	FPSOPrecacheRequestResult Result;
 	if (!PipelineStateCache::IsPSOPrecachingEnabled())
@@ -4827,13 +4905,10 @@ FPSOPrecacheRequestResult FPrecacheGraphicsPipelineCache::PrecacheGraphicsPipeli
 	// Use async compilation if available
 	static bool bDoAsyncCompile = FApp::ShouldUseThreadingForPerformance();
 
-	FString PSOCompilationEventName;
-
-	// try and create new graphics state
-	return TryAddNewState(Initializer, InMemoryMode, PSOCompilationEventName, bDoAsyncCompile);
+	return TryAddNewState(Initializer, RequestParams, bDoAsyncCompile);
 }
 	
-void FPrecacheGraphicsPipelineCache::OnNewPipelineStateCreated(const FGraphicsPipelineStateInitializer & Initializer, FGraphicsPipelineState * NewGraphicsPipelineState, const FString& PSOCompilationEventName, bool bDoAsyncCompile)
+void FPrecacheGraphicsPipelineCache::OnNewPipelineStateCreated(const FGraphicsPipelineStateInitializer & Initializer, EPSOPrecachePriority PrecachePriority, FGraphicsPipelineState * NewGraphicsPipelineState, const FString& PSOCompilationEventName, bool bDoAsyncCompile)
 {
 	ValidateGraphicsPipelineStateInitializer(Initializer);
 	check((NewGraphicsPipelineState->GetCompletionEvent() != nullptr) == bDoAsyncCompile);
@@ -4846,12 +4921,12 @@ void FPrecacheGraphicsPipelineCache::OnNewPipelineStateCreated(const FGraphicsPi
 	PSOCompilationDebugData.PSOCompilationEventName = PSOCompilationEventName;
 
 	// Start the precache task	
-	InternalCreateGraphicsPipelineState(InitializerCopy, EPSOPrecacheResult::Active, false, bDoAsyncCompile, NewGraphicsPipelineState, PSOCompilationDebugData, false);
+	InternalCreateGraphicsPipelineState(InitializerCopy, EPSOPrecacheResult::Active, PrecachePriority, false, bDoAsyncCompile, NewGraphicsPipelineState, PSOCompilationDebugData, false);
 }
 
-FPSOPrecacheRequestResult PipelineStateCache::PrecacheGraphicsPipelineState(const FGraphicsPipelineStateInitializer& Initializer, EPSOPrecacheInMemoryMode InMemoryMode)
+FPSOPrecacheRequestResult PipelineStateCache::PrecacheGraphicsPipelineState(const FGraphicsPipelineStateInitializer& Initializer, const FPSOPrecacheRequestParams& RequestParams)
 {
-	return GPrecacheGraphicsPipelineCache->PrecacheGraphicsPipelineState(Initializer, InMemoryMode);
+	return GPrecacheGraphicsPipelineCache->PrecacheGraphicsPipelineState(Initializer, RequestParams);
 }
 
 EPSOPrecacheResult PipelineStateCache::CheckPipelineStateInCache(const FGraphicsPipelineStateInitializer& PipelineStateInitializer, FGraphEventRef& OutCompletionEvent, FPSOPrecacheRequestID& OutPSOPrecacheRequestID)
@@ -5083,7 +5158,7 @@ void DumpPipelineCacheStats()
 	}
 
 	UE_LOG(LogRHI, Log, TEXT("Have %d GraphicsPipeline entries"), NumCachedItems);
-	for (int i = 0; i < (int)EQueuedWorkPriority::Count; i++)
+	for (int i = 0; i < (int)EPSOPrecachePriority::Num; i++)
 	{
 		if(TotalPrecompileCompleteTime[i].load() > 0)
 		{
diff --git a/Engine/Source/Runtime/RHI/Public/PipelineStateCache.h b/Engine/Source/Runtime/RHI/Public/PipelineStateCache.h
--- a/Engine/Source/Runtime/RHI/Public/PipelineStateCache.h
+++ b/Engine/Source/Runtime/RHI/Public/PipelineStateCache.h
@@ -31,9 +31,12 @@ ENUM_CLASS_FLAGS(EApplyRendertargetOption);
  */
 enum class EPSOPrecachePriority : uint8
 {
+	Low = 0,
 	Medium,
 	High,
 	Highest,
+	Num,
+	NotSet = Num,
 };
 
 /**
@@ -45,6 +48,17 @@ enum class EPSOPrecacheInMemoryMode : uint8
 	Uncached,	//< Don't keep precached PSO in memory (release after async compile)
 };
 
+/**
+* PSO Precache parameters provided to the graphics and compute precache function
+*/
+struct FPSOPrecacheRequestParams
+{
+	EPSOPrecachePriority PrecachePriority = EPSOPrecachePriority::Medium;
+	EPSOPrecacheInMemoryMode InMemoryMode = EPSOPrecacheInMemoryMode::Uncached;	
+	FString Name;
+	bool bForcePrecache = false;
+};
+
 enum class ERayTracingPipelineCacheFlags : uint8
 {
 	// Query the pipeline cache, create pipeline if necessary.
@@ -223,10 +237,10 @@ namespace PipelineStateCache
 	extern RHI_API bool						IsPSOPrecachingEnabled();
 
 	/* Precache the compute shader and return a request ID if precached async */
-	extern RHI_API FPSOPrecacheRequestResult PrecacheComputePipelineState(const FComputePipelineStateInitializer& ComputeInitializer, EPSOPrecacheInMemoryMode InMemoryMode, const TCHAR* Name = nullptr, bool bForcePrecache = false);
+	extern RHI_API FPSOPrecacheRequestResult PrecacheComputePipelineState(const FComputePipelineStateInitializer& ComputeInitializer, const FPSOPrecacheRequestParams& RequestParams);
 
 	/* Precache the graphics PSO and return an optional graph event if precached async */
-	extern RHI_API FPSOPrecacheRequestResult PrecacheGraphicsPipelineState(const FGraphicsPipelineStateInitializer& PipelineStateInitializer, EPSOPrecacheInMemoryMode InMemoryMode);
+	extern RHI_API FPSOPrecacheRequestResult PrecacheGraphicsPipelineState(const FGraphicsPipelineStateInitializer& PipelineStateInitializer, const FPSOPrecacheRequestParams& RequestParams);
 
 	/* Retrieve the current PSO precache result state (slightly slower than IsPrecaching) */
 	extern RHI_API EPSOPrecacheResult		CheckPipelineStateInCache(const FGraphicsPipelineStateInitializer& PipelineStateInitializer, FGraphEventRef& OutCompletionEvent, FPSOPrecacheRequestID& OutPSOPrecacheRequestID);	
