Fix RHIVendorIdToString for Qualcomm on Windows

diff --git a/Engine/Source/Runtime/RHI/Private/RHIStrings.cpp b/Engine/Source/Runtime/RHI/Private/RHIStrings.cpp
--- a/Engine/Source/Runtime/RHI/Private/RHIStrings.cpp
+++ b/Engine/Source/Runtime/RHI/Private/RHIStrings.cpp
@@ -57,6 +57,12 @@ const TCHAR* RHIVendorIdToString()
 
 const TCHAR* RHIVendorIdToString(EGpuVendorId VendorId)
 {
+	// Special case for alternate Qualcomm vendor id.
+	if (static_cast<int32>(VendorId) == 0x4D4F4351)
+	{
+		return TEXT("Qualcomm");
+	}
+
 	switch (VendorId)
 	{
 	case EGpuVendorId::Amd:			return TEXT("AMD");
