Change check into ensure when stat tracking of per priority PSOs gets below 0

#Synced-CL 50645574

diff --git a/Engine/Source/Runtime/RHI/Private/PipelineStateCache.cpp b/Engine/Source/Runtime/RHI/Private/PipelineStateCache.cpp
--- a/Engine/Source/Runtime/RHI/Private/PipelineStateCache.cpp
+++ b/Engine/Source/Runtime/RHI/Private/PipelineStateCache.cpp
@@ -3132,7 +3132,7 @@ protected:
 		else
 		{
 			FPlatformAtomics::InterlockedDecrement(&PerPriorityCompileCount[(int32)PrecachePriority]);
-			check(PerPriorityCompileCount[(int32)PrecachePriority] >= 0);
+			ensure(PerPriorityCompileCount[(int32)PrecachePriority] >= 0);
 		}
 	}
 
