Fix possible race condition on priority boost set for PSO precaching

#Synced-CL 50687180

diff --git a/Engine/Source/Runtime/RHI/Private/PipelineStateCache.cpp b/Engine/Source/Runtime/RHI/Private/PipelineStateCache.cpp
--- a/Engine/Source/Runtime/RHI/Private/PipelineStateCache.cpp
+++ b/Engine/Source/Runtime/RHI/Private/PipelineStateCache.cpp
@@ -2632,7 +2632,7 @@ protected:
 			check(FindResult->PipelineState == nullptr);
 			FindResult->PipelineState = NewPipelineState;
 			FindResult->InMemoryMode = RequestParams.InMemoryMode;
-			FindResult->SetPrecachePriority(RequestParams.PrecachePriority);
+			verify(FindResult->AddPrecachePriority(RequestParams.PrecachePriority) == EPSOPrecachePriority::NotSet);
 						
 			if (bDoAsyncCompile)
 			{
@@ -2914,6 +2914,9 @@ public:
 				}
 				PrecachedPSOsToCleanup.Add(InitializerHash);
 			}
+
+			UpdateActiveCompileCount(false /*Increment*/);
+			UpdatePerPriorityCompileCount(StatePriority, false /*Increment*/);
 		}
 
 		// Mark as in memory now - if it's not cached then it was added to list for cleanup and will be subtracted there
@@ -2925,11 +2928,6 @@ public:
 		{
 			INC_DWORD_STAT(STAT_PrecachedComputePrecachedPSOsInMemory);
 		}
-
-        // Need to ensure that the boost request was actually executed: if only it was asked by BoostPriority, but not requested (ie TryAddNewState has not set the Compiling bit
-        // yet) then we must ignore the request
-		UpdateActiveCompileCount(false /*Increment*/);
-		UpdatePerPriorityCompileCount(StatePriority, false /*Increment*/);
 	}
 
 	static bool IsCompilationDone(EPSOPrecacheStateMask StateMask)
@@ -3182,17 +3180,27 @@ protected:
 
 		inline EPSOPrecachePriority GetPrecachePriority()
 		{
-			return (EPSOPrecachePriority)FPlatformAtomics::AtomicRead((volatile int8*)&PrecachePriority);
+			int8 BitSet = FPlatformAtomics::AtomicRead((volatile int8*)&PrecachePriorityBitSet);
+			check(BitSet != 0);
+			return (EPSOPrecachePriority)FMath::FloorLog2(BitSet);
 		}
 
-		inline void SetPrecachePriority(EPSOPrecachePriority NewPriority)
+		inline EPSOPrecachePriority AddPrecachePriority(EPSOPrecachePriority NewPriority)
 		{
-			return FPlatformAtomics::AtomicStore((volatile int8*)&PrecachePriority, (int8)NewPriority);
+			int8 PreviousBitSet = FPlatformAtomics::InterlockedOr((volatile int8*)&PrecachePriorityBitSet, (1<<(int8)NewPriority));
+			if (PreviousBitSet == 0)
+			{
+				return EPSOPrecachePriority::NotSet;
+			}
+			else
+			{
+				return (EPSOPrecachePriority)FMath::FloorLog2(PreviousBitSet);
+			}
 		}
 
 	private:
 		volatile EPSOPrecacheStateMask StateMask = EPSOPrecacheStateMask::None;
-		volatile EPSOPrecachePriority PrecachePriority = EPSOPrecachePriority::Medium;
+		volatile int8 PrecachePriorityBitSet = 0;
 	};
 
 	TMap<uint64, FPrecacheTask> PrecachedPSOInitializerData;
@@ -3212,31 +3220,29 @@ protected:
 
 	void SetPrecachePriorityInternal(FPrecacheTask& FindResult, EPSOPrecachePriority NewPSOPrecachePriority)
 	{
-		// Ignore requests which try to lower the priority again 
-		// (can happen with material side boost but PSO already boosted to highest because required for rendering by other component)
-		EPSOPrecachePriority CurrentPriority = FindResult.GetPrecachePriority();
-		if ((int32)CurrentPriority >= (int32)NewPSOPrecachePriority)
-		{
-			return;
-		}
-
-		// It's possible to get a priority change request while the task has not been started yet. In this case, TryAddNewState will take care of it
-		// if TryAddNewState is done, then we can proceed to process the request it, if the task is not done yet
+		// Make sure compilation is not done yet (written via write lock)
 		EPSOPrecacheStateMask CurrentState = FindResult.ReadPSOPrecacheState();
 		if (!IsCompilationDone(CurrentState))
-		{
-			// Only update priority when not done yet
-			FindResult.SetPrecachePriority(NewPSOPrecachePriority);
+		{			
+			// Ignore requests which try to lower the priority again or set to the same one
+			// (can happen with material side boost but PSO already boosted to highest because required for rendering by other component)		
+			EPSOPrecachePriority PreviousPriority = FindResult.AddPrecachePriority(NewPSOPrecachePriority);		
+			if ((int32)PreviousPriority >= (int32)NewPSOPrecachePriority)
+			{
+				return;
+			}
 
 			// Update the stats here even though the reschedule might not be processed yet because it can happen
 			// that the actual reschedule never happens if it's done by the time the TryAddNewState actually
 			// tried the reschedule and stats need to be correct to process process finished call
-			UpdatePerPriorityCompileCount(CurrentPriority, false);
+			UpdatePerPriorityCompileCount(PreviousPriority, false);
 			UpdatePerPriorityCompileCount(NewPSOPrecachePriority, true);
 
+			// It's possible to get a priority change request while the task has not been started yet. In this case, TryAddNewState will take care of it
+			// if TryAddNewState is done, then we can proceed to process the request it, if the task is not done yet
 			if (EnumHasAnyFlags(CurrentState, EPSOPrecacheStateMask::Compiling))
 			{
-				RescheduleTaskToNewPriority(CurrentState, NewPSOPrecachePriority, CurrentPriority, FindResult.PipelineState);
+				RescheduleTaskToNewPriority(CurrentState, NewPSOPrecachePriority, PreviousPriority, FindResult.PipelineState);
 			}
 		}
 	}
