Allow disallowing some GPUs from using D3D11 ES31

diff --git a/Engine/Source/Runtime/RHI/Private/Windows/WindowsDynamicRHI.cpp b/Engine/Source/Runtime/RHI/Private/Windows/WindowsDynamicRHI.cpp
--- a/Engine/Source/Runtime/RHI/Private/Windows/WindowsDynamicRHI.cpp
+++ b/Engine/Source/Runtime/RHI/Private/Windows/WindowsDynamicRHI.cpp
@@ -501,6 +501,36 @@ static bool ShouldAvoidD3D11_SM5()
 	return bAvoidD3D11SM5;
 }
 
+static bool ShouldAvoidD3D11_ES31()
+{
+	bool bAvoidD3D11ES31 = false;
+
+	static TOptional<bool> AvoidD3D11ES31;
+	if (AvoidD3D11ES31.IsSet())
+	{
+		bAvoidD3D11ES31 = AvoidD3D11ES31.GetValue();
+	}
+	else if (!FParse::Param(FCommandLine::Get(), TEXT("d3d11")) &&
+		!FParse::Param(FCommandLine::Get(), TEXT("dx11")) &&
+		!FParse::Param(FCommandLine::Get(), TEXT("es31")) &&
+		!FParse::Param(FCommandLine::Get(), TEXT("FeatureLevelES31")) &&
+		!FParse::Param(FCommandLine::Get(), TEXT("FeatureLevelES3_1")))
+	{
+		FWindowsGPUInfo BestGPUInfo;
+		GetGPUInfo(BestGPUInfo);
+
+		bAvoidD3D11ES31 = IsGPUInHardwareListForRHI(BestGPUInfo, TEXT("DeviceAvoidRHIList"), TEXT("D3D11_ES31"));
+		if (bAvoidD3D11ES31)
+		{
+			UE_LOG(LogRHI, Log, TEXT("Avoiding feature level D3D11 %s"), GetLogName(ERHIFeatureLevel::ES3_1));
+		}
+
+		AvoidD3D11ES31 = bAvoidD3D11ES31;
+	}
+
+	return bAvoidD3D11ES31;
+}
+
 static TArray<EShaderPlatform> ParseShaderPlatformsConfig(const TCHAR* InSettingName)
 {
 	TArray<FString> TargetedShaderFormats;
@@ -536,6 +566,19 @@ static TArray<EShaderPlatform> ParseShaderPlatformsConfig(const TCHAR* InSetting
 					});
 			}
 		}
+
+		bool bUseES31AvoidGPUList = false;
+		if (GConfig->GetBool(TEXT("D3D11_ES31"), TEXT("bUseES31AvoidGPUList"), bUseES31AvoidGPUList, GEngineIni))
+		{
+			if (bUseES31AvoidGPUList && ShouldAvoidD3D11_ES31())
+			{
+				ShaderPlatforms.RemoveAll(
+					[](EShaderPlatform ShaderPlatform)
+					{
+						return FDataDrivenShaderPlatformInfo::GetMaxFeatureLevel(ShaderPlatform) == ERHIFeatureLevel::ES3_1;
+					});
+			}
+		}
 	}
 
 	return ShaderPlatforms;
@@ -1245,9 +1288,15 @@ static IDynamicRHIModule* LoadDynamicRHIModule(ERHIFeatureLevel::Type& DesiredFe
 	}
 	else if (PreferredRHI)
 	{
-		ChosenRHI = PreferredRHI.GetValue();
-
-		UE_LOG(LogRHI, Log, TEXT("Using Preferred RHI: %s"), GetLogName(ChosenRHI));
+		if (Config.RHIConfigs[static_cast<uint32>(PreferredRHI.GetValue())].ShaderPlatforms.IsEmpty())
+		{
+			UE_LOG(LogRHI, Log, TEXT("Preferred RHI has no ShaderPlatforms available: %s"), GetLogName(PreferredRHI.GetValue()));
+		}
+		else
+		{
+			ChosenRHI = PreferredRHI.GetValue();
+			UE_LOG(LogRHI, Log, TEXT("Using Preferred RHI: %s"), GetLogName(ChosenRHI));
+		}
 	}
 	else
 	{
