Prevent FRHIFrameInfo to grow indefinitely if SetVsyncDetails is not called.

#Synced-CL 50646776

diff --git a/Engine/Source/Runtime/RHI/Private/RHIUtilities.cpp b/Engine/Source/Runtime/RHI/Private/RHIUtilities.cpp
--- a/Engine/Source/Runtime/RHI/Private/RHIUtilities.cpp
+++ b/Engine/Source/Runtime/RHI/Private/RHIUtilities.cpp
@@ -751,6 +751,7 @@ void FRHIFrameInfo::SetInputSampleMarker(uint64 FrameIndex)
 
 		FRHIFrameInfoMarker& FrameMarker = Frames.FindOrAdd(FrameIndex);
 		FrameMarker.InputCycle = Cycles64;
+		Trim();
 	}	
 }
 
@@ -761,6 +762,7 @@ void FRHIFrameInfo::SetPresentCounter(uint64 FrameIndex, uint64 PresentCounter)
 
 		FRHIFrameInfoMarker& FrameMarker = Frames.FindOrAdd(FrameIndex);
 		FrameMarker.PresentCounter = PresentCounter;
+		Trim();
 	}
 	
 }
@@ -810,10 +812,10 @@ void FRHIFrameInfo::SetVsyncDetails(const FRHIFlipDetails& NewFlipFrame)
 		}
 	}
 
-	ComputeAndTrimHistory(LatestFrameCounter);
+	ComputeHistory(LatestFrameCounter);
 }
 
-void FRHIFrameInfo::ComputeAndTrimHistory(uint64 FrameNumber)
+void FRHIFrameInfo::ComputeHistory(uint64 FrameNumber)
 {
 	UE::TUniqueLock ScopeLock(FramesMutex);
 
@@ -841,8 +843,13 @@ void FRHIFrameInfo::ComputeAndTrimHistory(uint64 FrameNumber)
 		AverageInputLatencyMs = AccumulatedInputLatency / ((double)ElementCount);
 	}
 
+	Trim();
+}
+
+void FRHIFrameInfo::Trim()
+{
 	// Array should not grow indefinitely 
-	if (Frames.Num() > (FrameInfoHistoryCount*2))
+	if (Frames.Num() > (FrameInfoHistoryCount * 2))
 	{
 		Frames.CreateIterator().RemoveCurrent();
 	}
diff --git a/Engine/Source/Runtime/RHI/Public/RHIUtilities.h b/Engine/Source/Runtime/RHI/Public/RHIUtilities.h
--- a/Engine/Source/Runtime/RHI/Public/RHIUtilities.h
+++ b/Engine/Source/Runtime/RHI/Public/RHIUtilities.h
@@ -789,7 +789,8 @@ public:
 	float GetAverageInputLatencyInMs()  { return AverageInputLatencyMs; }
 
 private:
-	void ComputeAndTrimHistory(uint64 FrameIndex);
+	void ComputeHistory(uint64 FrameIndex);
+	void Trim();
 
 	TMap<uint64,FRHIFrameInfoMarker> Frames;
 	UE::FMutex FramesMutex;
