Removing more HLSLcc code
- Deprecating FDataDrivenShaderPlatformInfo::GetIsHlslcc
- Deprecating IsHlslccShaderPlatform
- All uses of these functions removed

#Synced-CL 50826034

diff --git a/Engine/Config/Android/DataDrivenPlatformInfo.ini b/Engine/Config/Android/DataDrivenPlatformInfo.ini
--- a/Engine/Config/Android/DataDrivenPlatformInfo.ini
+++ b/Engine/Config/Android/DataDrivenPlatformInfo.ini
@@ -75,7 +75,6 @@ bSupportsVertexShaderSRVs=false
 bSupportsUniformBufferObjects = true
 bSupportsDxc=true
 bSupportsIndependentSamplers=false
-bIsHlslcc=false
 bSupportsComputeFramework = true
 bSupportsDistanceFields=false
 FriendlyName=LOCTEXT("FriendlyShaderPlatformName_Android", "Android OpenGL Mobile")
@@ -100,7 +99,6 @@ bSupportsUniformBufferObjects = true
 bSupportsDualSourceBlending = true
 bSupportsDxc=true
 bSupportsIndependentSamplers=true
-bIsHlslcc=false
 bSupportsComputeFramework = true
 bSupportsDistanceFields=true
 
@@ -133,7 +131,6 @@ bSupportsDualSourceBlending=false
 bSupportsVariableRateShading = false
 bSupportsDxc=true
 bSupportsIndependentSamplers=true
-bIsHlslcc=false
 bSupportsWaveOperations=RuntimeDependent
 MinimumWaveSize=4
 MaximumWaveSize=128
diff --git a/Engine/Plugins/Runtime/HairStrands/Source/HairStrandsCore/Private/GroomTextureBuilder.cpp b/Engine/Plugins/Runtime/HairStrands/Source/HairStrandsCore/Private/GroomTextureBuilder.cpp
--- a/Engine/Plugins/Runtime/HairStrands/Source/HairStrandsCore/Private/GroomTextureBuilder.cpp
+++ b/Engine/Plugins/Runtime/HairStrands/Source/HairStrandsCore/Private/GroomTextureBuilder.cpp
@@ -478,11 +478,6 @@ public:
 	static void ModifyCompilationEnvironment(const FGlobalShaderPermutationParameters& Parameters, FShaderCompilerEnvironment& OutEnvironment)
 	{
 		FGlobalShader::ModifyCompilationEnvironment(Parameters, OutEnvironment);
-		// Work around issues with HLSLcc by switching to DXC
-		if (IsHlslccShaderPlatform(Parameters.Platform))
-		{
-			OutEnvironment.CompilerFlags.Add(CFLAG_ForceDXC);
-		}
 		OutEnvironment.SetDefine(TEXT("SHADER_TEXTURE_CLEAR"), 1);
 	}
 };
@@ -783,11 +778,6 @@ class FHairStrandsTextureCS : public FGlobalShader
 	static void ModifyCompilationEnvironment(const FGlobalShaderPermutationParameters& Parameters, FShaderCompilerEnvironment& OutEnvironment)
 	{
 		FGlobalShader::ModifyCompilationEnvironment(Parameters, OutEnvironment);
-		// Work around issues with HLSLcc by switching to DXC
-		if (IsHlslccShaderPlatform(Parameters.Platform))
-		{
-			OutEnvironment.CompilerFlags.Add(CFLAG_ForceDXC);
-		}
 		OutEnvironment.SetDefine(TEXT("SHADER_COMPUTE"), 1);
 	}
 };
diff --git a/Engine/Source/Runtime/RHI/Private/DataDrivenShaderPlatformInfo.cpp b/Engine/Source/Runtime/RHI/Private/DataDrivenShaderPlatformInfo.cpp
--- a/Engine/Source/Runtime/RHI/Private/DataDrivenShaderPlatformInfo.cpp
+++ b/Engine/Source/Runtime/RHI/Private/DataDrivenShaderPlatformInfo.cpp
@@ -278,7 +278,6 @@ void FGenericDataDrivenShaderPlatformInfo::ParseDataDrivenShaderInfo(const FConf
 	GET_SECTION_BOOL_HELPER(bStripNaniteVertexShaders);
 	GET_SECTION_BOOL_HELPER(bRequiresUnwrappedMeshShaderArgs);
 	GET_SECTION_BOOL_HELPER(bSupportsPerPixelDBufferMask);
-	GET_SECTION_BOOL_HELPER(bIsHlslcc);
 	GET_SECTION_BOOL_HELPER(bSupportsDxc);
 	GET_SECTION_BOOL_HELPER(bSupportsVariableRateShading);
 	GET_SECTION_BOOL_HELPER(bIsSPIRV);
@@ -533,7 +532,6 @@ void FGenericDataDrivenShaderPlatformInfo::UpdatePreviewPlatforms()
 				// Always inherit these core settings from the preview
 				PREVIEW_USE_RUNTIME_VALUE(ShaderFormat);
 				PREVIEW_USE_RUNTIME_VALUE(Language);
-				PREVIEW_USE_RUNTIME_VALUE(bIsHlslcc);
 				PREVIEW_USE_RUNTIME_VALUE(bSupportsDxc);
 				PREVIEW_USE_RUNTIME_VALUE(bIsSPIRV);
 
diff --git a/Engine/Source/Runtime/RHI/Public/Android/AndroidDataDrivenShaderPlatformInfo.inl b/Engine/Source/Runtime/RHI/Public/Android/AndroidDataDrivenShaderPlatformInfo.inl
--- a/Engine/Source/Runtime/RHI/Public/Android/AndroidDataDrivenShaderPlatformInfo.inl
+++ b/Engine/Source/Runtime/RHI/Public/Android/AndroidDataDrivenShaderPlatformInfo.inl
@@ -6,7 +6,6 @@ class FDataDrivenShaderPlatformInfo final : public FGenericDataDrivenShaderPlatf
 {
 public:
 	IMPLEMENT_DDPSPI_SETTING(GetTargetsTiledGPU,				true);
-	IMPLEMENT_DDPSPI_SETTING(GetIsHlslcc,						false);
 	IMPLEMENT_DDPSPI_SETTING(GetSupportsDxc,					true);
 	IMPLEMENT_DDPSPI_SETTING(GetSupportsUniformBufferObjects,	true);
 
diff --git a/Engine/Source/Runtime/RHI/Public/DataDrivenShaderPlatformInfo.h b/Engine/Source/Runtime/RHI/Public/DataDrivenShaderPlatformInfo.h
--- a/Engine/Source/Runtime/RHI/Public/DataDrivenShaderPlatformInfo.h
+++ b/Engine/Source/Runtime/RHI/Public/DataDrivenShaderPlatformInfo.h
@@ -87,7 +87,6 @@ class FGenericDataDrivenShaderPlatformInfo
 	uint32 MaxMeshShaderThreadGroupSize : 10;
 	uint32 bRequiresUnwrappedMeshShaderArgs : 1;
 	uint32 bSupportsPerPixelDBufferMask : 1;
-	uint32 bIsHlslcc : 1;
 	uint32 bSupportsDxc : 1; // Whether DirectXShaderCompiler (DXC) is supported
 	uint32 bIsSPIRV : 1;
 	uint32 bSupportsVariableRateShading : 1;
@@ -631,10 +630,10 @@ public:
 		return Infos[Platform].bSupportsPerPixelDBufferMask;
 	}
 
+	UE_DEPRECATED(5.8, "GetIsHlslcc is always false so it will be removed.")
 	static inline const bool GetIsHlslcc(const FStaticShaderPlatform Platform)
 	{
-		check(IsValid(Platform));
-		return Infos[Platform].bIsHlslcc;
+		return false;
 	}
 
 	static inline const bool GetSupportsDxc(const FStaticShaderPlatform Platform)
@@ -1029,9 +1028,10 @@ inline bool IsD3DPlatform(const FStaticShaderPlatform Platform)
 	return FDataDrivenShaderPlatformInfo::GetIsLanguageD3D(Platform);
 }
 
+UE_DEPRECATED(5.8, "IsHlslccShaderPlatform is always false so it will be removed.")
 inline bool IsHlslccShaderPlatform(const FStaticShaderPlatform Platform)
 {
-	return IsOpenGLPlatform(Platform) || FDataDrivenShaderPlatformInfo::GetIsHlslcc(Platform);
+	return false;
 }
 
 inline FStaticFeatureLevel GetMaxSupportedFeatureLevel(const FStaticShaderPlatform InShaderPlatform)
@@ -1101,7 +1101,7 @@ inline bool RHISupportsShaderPipelines(const FStaticShaderPlatform Platform)
 inline bool RHISupportsDualSourceBlending(const FStaticShaderPlatform Platform)
 {
 	// Check if the platform supports dual src blending from DDPI
-	return FDataDrivenShaderPlatformInfo::GetSupportsDualSourceBlending(Platform) && !FDataDrivenShaderPlatformInfo::GetIsHlslcc(Platform);
+	return FDataDrivenShaderPlatformInfo::GetSupportsDualSourceBlending(Platform);
 }
 
 // helper to check that the shader platform supports creating a UAV off an index buffer.
diff --git a/Engine/Source/Runtime/Renderer/Private/InstanceCulling/InstanceCullingContext.cpp b/Engine/Source/Runtime/Renderer/Private/InstanceCulling/InstanceCullingContext.cpp
--- a/Engine/Source/Runtime/Renderer/Private/InstanceCulling/InstanceCullingContext.cpp
+++ b/Engine/Source/Runtime/Renderer/Private/InstanceCulling/InstanceCullingContext.cpp
@@ -539,12 +539,6 @@ public:
 		FGlobalShader::ModifyCompilationEnvironment(Parameters, OutEnvironment);
 		FInstanceProcessingGPULoadBalancer::SetShaderDefines(OutEnvironment);
 
-		// Force use of DXC for platforms compiled with hlslcc due to hlslcc's inability to handle member functions in structs
-		if (FDataDrivenShaderPlatformInfo::GetIsHlslcc(Parameters.Platform))
-		{
-			OutEnvironment.CompilerFlags.Add(CFLAG_ForceDXC);
-		}
-
 		// This shader takes a very long time to compile with FXC, so we pre-compile it with DXC first and then forward the optimized HLSL to FXC.
 		OutEnvironment.CompilerFlags.Add(CFLAG_PrecompileWithDXC);
 
diff --git a/Engine/Source/Runtime/Renderer/Private/MobileBasePassRendering.cpp b/Engine/Source/Runtime/Renderer/Private/MobileBasePassRendering.cpp
--- a/Engine/Source/Runtime/Renderer/Private/MobileBasePassRendering.cpp
+++ b/Engine/Source/Runtime/Renderer/Private/MobileBasePassRendering.cpp
@@ -232,13 +232,6 @@ void MobileBasePassModifyCompilationEnvironment(const FMaterialShaderPermutation
 	// translucency is in the same subpass with deferred shading shaders, so it has access to GBuffer
 	const bool bDeferredShadingSubpass = (bDeferredShadingEnabled && bTranslucentMaterial && !Parameters.MaterialParameters.bIsMobileSeparateTranslucencyEnabled);
 	OutEnvironment.SetDefine(TEXT("IS_MOBILE_DEFERREDSHADING_SUBPASS"), bDeferredShadingSubpass ? 1u : 0u);
-
-	// HLSLcc does not support dual source blending, so force DXC if needed
-	if (bTranslucentMaterial && FDataDrivenShaderPlatformInfo::GetSupportsDxc(Parameters.Platform) && IsHlslccShaderPlatform(Parameters.Platform) &&
-		MaterialRequiresColorTransmittanceBlending(Parameters.MaterialParameters) && MobileDefaultTranslucentColorTransmittanceMode(Parameters.Platform) == EMobileTranslucentColorTransmittanceMode::DUAL_SRC_BLENDING)
-	{
-		OutEnvironment.CompilerFlags.Add(CFLAG_ForceDXC);
-	}
 }
 
 template<typename LightMapPolicyType>
diff --git a/Engine/Source/Runtime/Renderer/Private/OIT/OIT.cpp b/Engine/Source/Runtime/Renderer/Private/OIT/OIT.cpp
--- a/Engine/Source/Runtime/Renderer/Private/OIT/OIT.cpp
+++ b/Engine/Source/Runtime/Renderer/Private/OIT/OIT.cpp
@@ -494,7 +494,7 @@ public:
 
 	static bool SupportDebugMode(EShaderPlatform InPlatform)
 	{
-		return ShaderPrint::IsSupported(InPlatform) && !IsHlslccShaderPlatform(InPlatform);
+		return ShaderPrint::IsSupported(InPlatform);
 	}
 
 	static bool ShouldCompilePermutation(const FGlobalShaderPermutationParameters& Parameters) 
diff --git a/Engine/Source/Runtime/Renderer/Private/PhysicsFieldRendering.cpp b/Engine/Source/Runtime/Renderer/Private/PhysicsFieldRendering.cpp
--- a/Engine/Source/Runtime/Renderer/Private/PhysicsFieldRendering.cpp
+++ b/Engine/Source/Runtime/Renderer/Private/PhysicsFieldRendering.cpp
@@ -84,7 +84,7 @@ class FPhysicsFieldRayMarchingCS : public FGlobalShader
 public:
 	static bool ShouldCompilePermutation(const FGlobalShaderPermutationParameters& Parameters) 
 	{
-		return ShaderPrint::IsSupported(Parameters.Platform) && IsFeatureLevelSupported(Parameters.Platform, ERHIFeatureLevel::SM5) && !IsHlslccShaderPlatform(Parameters.Platform);
+		return ShaderPrint::IsSupported(Parameters.Platform) && IsFeatureLevelSupported(Parameters.Platform, ERHIFeatureLevel::SM5);
 	}
 	static void ModifyCompilationEnvironment(const FGlobalShaderPermutationParameters& Parameters, FShaderCompilerEnvironment& OutEnvironment)
 	{
diff --git a/Engine/Source/Runtime/Renderer/Private/PostProcess/PostProcessing.cpp b/Engine/Source/Runtime/Renderer/Private/PostProcess/PostProcessing.cpp
--- a/Engine/Source/Runtime/Renderer/Private/PostProcess/PostProcessing.cpp
+++ b/Engine/Source/Runtime/Renderer/Private/PostProcess/PostProcessing.cpp
@@ -4132,7 +4132,6 @@ public:
 		return 
 			IsFeatureLevelSupported(Platform, ERHIFeatureLevel::SM5) &&
 			ShaderPrint::IsSupported(Platform) && 
-			!IsHlslccShaderPlatform(Platform) &&
 			!IsMobilePlatform(Platform) && 
 			(!Substrate::IsSubstrateEnabled() || Substrate::IsSubstrateBlendableGBufferEnabled(Platform));
 	}
diff --git a/Engine/Source/Runtime/Renderer/Private/ShaderPrint.cpp b/Engine/Source/Runtime/Renderer/Private/ShaderPrint.cpp
--- a/Engine/Source/Runtime/Renderer/Private/ShaderPrint.cpp
+++ b/Engine/Source/Runtime/Renderer/Private/ShaderPrint.cpp
@@ -295,12 +295,6 @@ namespace ShaderPrint
 
 	void ModifyCompilationEnvironment(const EShaderPlatform Platform, FShaderCompilerEnvironment& OutEnvironment)
 	{
-		// Work around issues with HLSLcc by switching to DXC
-		if (IsHlslccShaderPlatform(Platform))
-		{
-			OutEnvironment.CompilerFlags.Add(CFLAG_ForceDXC);
-		}
-
 		OutEnvironment.SetDefine(TEXT("SHADER_PRINT_EXPLICIT_BINDING"), 1);
 	}
 
