More concise GPU profiler log output - only print the header once, helps with quickly figuring out which pipeline you are looking at

#Synced-CL 50792079

diff --git a/Engine/Source/Runtime/RHI/Private/GPUProfiler.cpp b/Engine/Source/Runtime/RHI/Private/GPUProfiler.cpp
--- a/Engine/Source/Runtime/RHI/Private/GPUProfiler.cpp
+++ b/Engine/Source/Runtime/RHI/Private/GPUProfiler.cpp
@@ -1342,8 +1342,7 @@ namespace UE::RHI::GPUProfiler
 						);
 					}
 
-					FString GPUString = FString::Printf(TEXT("Frame %d - %s %d - GPU %d")
-						, FrameNumber
+					FString PipelineString = FString::Printf(TEXT("%s pipeline %d - GPU %d")
 						, QueueState.Queue.GetTypeString()
 						, QueueState.Queue.Index
 						, QueueState.Queue.GPU
@@ -1361,16 +1360,26 @@ namespace UE::RHI::GPUProfiler
 						bShowQueue = NodeOutput.Len() > 0;
 					}
 
-					FString Final = FString::Printf(
-						TEXT("\n")
-						TEXT("GPU Profile for %s\n")
-						TEXT("\n")
-						TEXT("%s")
-						TEXT("%s")
-						, *GPUString
-						, *Header
-						, *NodeOutput
-					);
+					FString Final;
+
+					if (QueueState.Queue.Type == FQueue::EType::Graphics)
+					{
+						Final = FString::Printf(
+							TEXT("\n")
+							TEXT("GPU Profile for Frame %d - %s\n")
+							TEXT("\n")
+							TEXT("%s")
+							TEXT("%s")
+							, FrameNumber
+							, *PipelineString
+							, *Header
+							, *NodeOutput
+						);
+					}
+					else
+					{
+						Final = FString::Printf(TEXT("%s\n\n%s"), *PipelineString, *NodeOutput);
+					}
 
 					if (bShowQueue)
 					{
@@ -1426,7 +1435,7 @@ namespace UE::RHI::GPUProfiler
 
 						// Launch the visualizer on the game thread
 						FFunctionGraphTask::CreateAndDispatchWhenReady([
-							Title = GPUString,
+							Title = PipelineString,
 							VisualizerData = BuildData(BuildData, First, nullptr)
 						]
 						{
