Modify r.PSOPrecaching.WaitForHighPriorityRequestsOnly to r.PSOPrecaching.MinPriorityForActivePrecacheRequests so any minimum priority can be set when collecting active PSO requests used to wait during loading screens for example

#Synced-CL 51088011

diff --git a/Engine/Source/Runtime/RHI/Private/PipelineStateCache.cpp b/Engine/Source/Runtime/RHI/Private/PipelineStateCache.cpp
--- a/Engine/Source/Runtime/RHI/Private/PipelineStateCache.cpp
+++ b/Engine/Source/Runtime/RHI/Private/PipelineStateCache.cpp
@@ -217,13 +217,15 @@ static FAutoConsoleVariableRef CVarPSOPrecaching(
 	ECVF_Default
 );
 
-int32 GPSOWaitForHighPriorityRequestsOnly = 0;
+int32 GPSOPrecacheMinPriorityForActivePrecacheRequests = 0;
 static FAutoConsoleVariableRef CVarPSOWaitForHighPriorityRequestsOnly(
-	TEXT("r.PSOPrecaching.WaitForHighPriorityRequestsOnly"),
-	GPSOWaitForHighPriorityRequestsOnly,
-	TEXT("0 to wait for all pending PSO precache requests during loading (default)\n")
-	TEXT("1 to only wait for the high priority and above PSO precache requests during loading\n")
-	TEXT("2 to only wait for the highest priority PSO precache requests during loading"),
+	TEXT("r.PSOPrecaching.MinPriorityForActivePrecacheRequests"),
+	GPSOPrecacheMinPriorityForActivePrecacheRequests,
+	TEXT("Priority of active request to take into account when computing the number of active PSOs (can be used in loading screen waits for example).\n")
+	TEXT("0: Add all pending PSO precache requests (default)\n")
+	TEXT("1: Medium priority and above only\n")
+	TEXT("2: High priority and above only\n")
+	TEXT("3: Highest priority only\n"),
 	ECVF_Default
 );
 
@@ -253,7 +255,7 @@ static FAutoConsoleVariableRef CVarPSOOnlyPrecacheMaxPriority(
 	TEXT("0: All compile tasks are processed. (default)\n")
 	TEXT("1: Medium priority compile tasks and above are processed.\n")
 	TEXT("2: High priority compile tasks and above are processed.\n")
-	TEXT("3: Only very high priority compile tasks are processed."),
+	TEXT("3: Only highest priority compile tasks are processed."),
 	ECVF_Default
 );
 static_assert((int)FGraphicsPipelineStateInitializer::EPSOPrecacheCompileType::NumTypes == 4, "EPSOPrecacheCompileType has changed, 'r.PSOPrecaching.TaskPriorityThreshold' must be updated.");
@@ -2867,16 +2869,20 @@ public:
 
 	uint32 NumActivePrecacheRequests()
 	{
-		switch (GPSOWaitForHighPriorityRequestsOnly)
+		// Take process priority into account as well, because otherwise lower priority requests will never be compiled when we try
+		// to wait for them
+		int32 MinPriorityForActivePrecacheRequests = FMath::Max(GPSOPrecacheProcessPriorityThreshold, GPSOPrecacheMinPriorityForActivePrecacheRequests);
+		if (MinPriorityForActivePrecacheRequests <= 0)
 		{
-			default: checkNoEntry(); [[fallthrough]];
-			case 0:
-				return FPlatformAtomics::AtomicRead(&ActiveCompileCount);
-			case 1:
-				return FPlatformAtomics::AtomicRead(&PerPriorityCompileCount[(int32)EPSOPrecachePriority::Highest]) + FPlatformAtomics::AtomicRead(&PerPriorityCompileCount[(int32)EPSOPrecachePriority::High]);
-			case 2:
-				return FPlatformAtomics::AtomicRead(&PerPriorityCompileCount[(int32)EPSOPrecachePriority::Highest]);
+			return FPlatformAtomics::AtomicRead(&ActiveCompileCount);
+		}
+
+		uint32 ActivePrecacheRequests = 0;
+		for (int32 Priority = MinPriorityForActivePrecacheRequests; Priority < (int32)EPSOPrecachePriority::Num; ++Priority)
+		{
+			ActivePrecacheRequests += FPlatformAtomics::AtomicRead(&PerPriorityCompileCount[Priority]);
 		}
+		return ActivePrecacheRequests;
 	}
 
 	void PrecacheFinished(const TPrecachedPSOInitializer& Initializer, bool bValid, uint64 PSOCompilationTime)
