ShaderPlatformConfig Changes
- Adding more documentation to ShaderPlatformConfig.h
- Adding "Preview Runtime Platform" to DDSPI preview platform field. This saves the actual runtime SP used to initialize the preview platform before ShaderPlatformForFeatureLevel can be changed.
- Fixing SPC "RuntimePlatform" logic to use the new DDSPI preview platform field.

#Synced-CL 51242114

diff --git a/Engine/Source/Runtime/RHI/Private/DataDrivenShaderPlatformInfo.cpp b/Engine/Source/Runtime/RHI/Private/DataDrivenShaderPlatformInfo.cpp
--- a/Engine/Source/Runtime/RHI/Private/DataDrivenShaderPlatformInfo.cpp
+++ b/Engine/Source/Runtime/RHI/Private/DataDrivenShaderPlatformInfo.cpp
@@ -26,6 +26,7 @@ struct FDataDrivenShaderPlatformInfoEditorOnly
 {
 	TMap<FName,FText> DeviceProfileToFriendlyName;
 	EShaderPlatform PreviewShaderPlatformParent = SP_NumPlatforms;
+	EShaderPlatform RuntimePlatform = SP_NumPlatforms;
 	uint32 bCanUsePreviewPlatformForMaterialValidation : 1;
 };
 static FDataDrivenShaderPlatformInfoEditorOnly DataDrivenShaderPlatformInfoEditorOnlyInfos[SP_NumPlatforms];
@@ -519,6 +520,9 @@ void FGenericDataDrivenShaderPlatformInfo::UpdatePreviewPlatforms()
 				FGenericDataDrivenShaderPlatformInfo& PreviewInfo = Infos[PreviewPlatform];
 				const FGenericDataDrivenShaderPlatformInfo& RuntimeInfo = Infos[RuntimePlatform];
 
+				FDataDrivenShaderPlatformInfoEditorOnly& PreviewEditorInfo = DataDrivenShaderPlatformInfoEditorOnlyInfos[PreviewPlatform];
+				PreviewEditorInfo.RuntimePlatform = RuntimePlatform;
+
 #define PREVIEW_USE_RUNTIME_VALUE(SettingName) \
 	PreviewInfo.SettingName = RuntimeInfo.SettingName
 
@@ -612,7 +616,7 @@ void FGenericDataDrivenShaderPlatformInfo::UpdatePreviewPlatforms()
 	}
 }
 
-FText FGenericDataDrivenShaderPlatformInfo::GetFriendlyName(const FStaticShaderPlatform Platform, FName DeviceProfileName)
+FText FGenericDataDrivenShaderPlatformInfo::GetFriendlyName(EShaderPlatform Platform, FName DeviceProfileName)
 {
 	if (IsRunningCommandlet() || GUsingNullRHI)
 	{
@@ -629,19 +633,26 @@ FText FGenericDataDrivenShaderPlatformInfo::GetFriendlyName(const FStaticShaderP
 	return FText::FromName(DeviceProfileName);
 }
 
-const EShaderPlatform FGenericDataDrivenShaderPlatformInfo::GetPreviewShaderPlatformParent(const FStaticShaderPlatform Platform)
+const EShaderPlatform FGenericDataDrivenShaderPlatformInfo::GetPreviewShaderPlatformParent(EShaderPlatform Platform)
 {
 	check(IsValid(Platform));
 	return DataDrivenShaderPlatformInfoEditorOnlyInfos[Platform].PreviewShaderPlatformParent;
 }
 
-const bool FGenericDataDrivenShaderPlatformInfo::CanUseForMaterialValidation(const FStaticShaderPlatform Platform)
+const EShaderPlatform FGenericDataDrivenShaderPlatformInfo::GetPreviewRuntimePlatform(EShaderPlatform Platform)
+{
+	check(IsValid(Platform));
+	check(GetIsPreviewPlatform(Platform));
+	return DataDrivenShaderPlatformInfoEditorOnlyInfos[Platform].RuntimePlatform;
+}
+
+const bool FGenericDataDrivenShaderPlatformInfo::CanUseForMaterialValidation(EShaderPlatform Platform)
 {
 	check(IsValid(Platform));
 	return DataDrivenShaderPlatformInfoEditorOnlyInfos[Platform].bCanUsePreviewPlatformForMaterialValidation || !Infos[Platform].bIsPreviewPlatform;
 }
 
-FName FGenericDataDrivenShaderPlatformInfo::GetPreferredShaderCodeFormat(const FStaticShaderPlatform Platform)
+FName FGenericDataDrivenShaderPlatformInfo::GetPreferredShaderCodeFormat(EShaderPlatform Platform)
 {
 	return FDataDrivenShaderPlatformInfo::GetIsPreviewPlatform(Platform) ? GRHIGlobals.PreferredPreviewShaderCodeFormat : NAME_None;
 }
diff --git a/Engine/Source/Runtime/RHI/Public/DataDrivenShaderPlatformInfo.h b/Engine/Source/Runtime/RHI/Public/DataDrivenShaderPlatformInfo.h
--- a/Engine/Source/Runtime/RHI/Public/DataDrivenShaderPlatformInfo.h
+++ b/Engine/Source/Runtime/RHI/Public/DataDrivenShaderPlatformInfo.h
@@ -910,10 +910,11 @@ PRAGMA_ENABLE_DEPRECATION_WARNINGS
 
 #if WITH_EDITOR
 	RHI_API static void UpdatePreviewPlatforms();
-	RHI_API static FText GetFriendlyName(const FStaticShaderPlatform Platform, FName DeviceProfileName = NAME_None);
-	RHI_API static const EShaderPlatform GetPreviewShaderPlatformParent(const FStaticShaderPlatform Platform);
-	RHI_API static const bool CanUseForMaterialValidation(const FStaticShaderPlatform Platform);
-	RHI_API static FName GetPreferredShaderCodeFormat(const FStaticShaderPlatform Platform);
+	RHI_API static FText GetFriendlyName(EShaderPlatform Platform, FName DeviceProfileName = NAME_None);
+	RHI_API static const EShaderPlatform GetPreviewShaderPlatformParent(EShaderPlatform Platform);
+	RHI_API static const EShaderPlatform GetPreviewRuntimePlatform(EShaderPlatform Platform);
+	RHI_API static const bool CanUseForMaterialValidation(EShaderPlatform Platform);
+	RHI_API static FName GetPreferredShaderCodeFormat(EShaderPlatform Platform);
 
 	RHI_API static TMap<FString, TFunction<bool(const FStaticShaderPlatform Platform)>> PropertyToShaderPlatformFunctionMap;
 #endif
diff --git a/Engine/Source/Runtime/RenderCore/Private/ShaderPlatformConfig.cpp b/Engine/Source/Runtime/RenderCore/Private/ShaderPlatformConfig.cpp
--- a/Engine/Source/Runtime/RenderCore/Private/ShaderPlatformConfig.cpp
+++ b/Engine/Source/Runtime/RenderCore/Private/ShaderPlatformConfig.cpp
@@ -137,7 +137,8 @@ void FShaderPlatformConfig::Initialize()
 			{
 #if WITH_EDITOR
 				const ERHIFeatureLevel::Type PreviewFeatureLevel = FDataDrivenShaderPlatformInfo::GetMaxFeatureLevel(ShaderPlatform);
-				const EShaderPlatform RuntimePlatform = GRHIGlobals.ShaderPlatformForFeatureLevel[PreviewFeatureLevel];
+				const EShaderPlatform RuntimePlatform = FDataDrivenShaderPlatformInfo::GetPreviewRuntimePlatform(ShaderPlatform);
+				check(RuntimePlatform != SP_NumPlatforms);
 
 				if (RuntimePlatform < SP_NumPlatforms && PlatformConfigs[RuntimePlatform].IsValid())
 				{
diff --git a/Engine/Source/Runtime/RenderCore/Public/ShaderPlatformConfig.h b/Engine/Source/Runtime/RenderCore/Public/ShaderPlatformConfig.h
--- a/Engine/Source/Runtime/RenderCore/Public/ShaderPlatformConfig.h
+++ b/Engine/Source/Runtime/RenderCore/Public/ShaderPlatformConfig.h
@@ -9,7 +9,8 @@ class FConfigCacheIni;
 
 // Shader Platform Configuration
 // 
-// Similar to DataDrivenShaderPlatformInfo (DDSPI), but settings come from Engine config hierarchy.
+// Similar to DataDrivenShaderPlatformInfo (DDSPI), but settings come from Engine config hierarchy. DDSPI is specifically for immutable shader platform
+// capabilities, while SPC is for shader platform configuration.
 // 
 // Configs are stored in one section per shader platform with the same names that DDSPI uses.
 // ex: [ShaderPlatform PCD3D_SM6] -> [ShaderPlatformConfig PCD3D_SM6]
@@ -17,6 +18,17 @@ class FConfigCacheIni;
 // Interface initially built with static getters but those use per-platform structs with getters. While the static getters suffice for initial adoption, the
 // ideal usage would be passing the per-platform struct around to code that needs them. This usage would make it easier to adopt "custom shader platforms" in
 // the future.
+//
+// New fields need the following:
+//   1. Value inside FPlatformConfig
+//   2. Member getter inside FPlatformConfig
+//   3. Static getter inside FShaderPlatformConfig
+//   4. Default value inside FShaderPlatformConfig::FPlatformConfig::SetDefaultValues
+//   5. Initialization inside FShaderPlatformConfig::FPlatformConfig::InitializeProperties
+//   6. (optional) Engine defaults inside Base{Platform}Engine.ini
+//   7. (optional) Preview override inside FShaderPlatformConfig::FPlatformConfig::InitializePropertiesForPreview
+//
+
 class FShaderPlatformConfig
 {
 public:
