PipelineStateCache - Use EPSOPrecacheInMemoryMode::Uncached until it's fixed for all platforms

Submitting in behalf of Kenzo Ter Elst

#tests local run

#Synced-CL 51263606

diff --git a/Engine/Source/Runtime/RHI/Private/PipelineStateCache.cpp b/Engine/Source/Runtime/RHI/Private/PipelineStateCache.cpp
--- a/Engine/Source/Runtime/RHI/Private/PipelineStateCache.cpp
+++ b/Engine/Source/Runtime/RHI/Private/PipelineStateCache.cpp
@@ -2552,7 +2552,8 @@ protected:
 	FPSOPrecacheRequestResult TryAddNewState(const TPrecachedPSOInitializer& Initializer, const FPSOPrecacheRequestParams& RequestParams, bool bDoAsyncCompile)
 	{
 		// Refcounted and max in memory PSOs can't be used together
-		EPSOPrecacheInMemoryMode InMemoryMode = ShouldKeepPrecachedPSOsInMemory() ? EPSOPrecacheInMemoryMode::Uncached : RequestParams.InMemoryMode;
+		// Using EPSOPrecacheInMemoryMode::Uncached for now until it's fixed on all platforms
+		EPSOPrecacheInMemoryMode InMemoryMode = EPSOPrecacheInMemoryMode::Uncached; //ShouldKeepPrecachedPSOsInMemory() ? EPSOPrecacheInMemoryMode::Uncached : RequestParams.InMemoryMode;
 
 		FPSOPrecacheRequestResult Result;
 		uint64 InitializerHash = TPrecachePipelineCacheDerived::PipelineStateInitializerHash(Initializer);
