Add abilities to disable and restore the SM5 feature level

INI setting Engine:[RHI]:bBlockSM5=True can disable SM5 support, and CVar r.ForceSM5Support allows for hotfixing it back. The CVar is unfortunately read too early to be specified in config, so that's why the INI setting exists.

#Synced-CL 51402227

diff --git a/Engine/Source/Runtime/RHI/Private/DynamicRHI.cpp b/Engine/Source/Runtime/RHI/Private/DynamicRHI.cpp
--- a/Engine/Source/Runtime/RHI/Private/DynamicRHI.cpp
+++ b/Engine/Source/Runtime/RHI/Private/DynamicRHI.cpp
@@ -31,6 +31,7 @@
 FDynamicRHI* GDynamicRHI = NULL;
 RHIGetGPUUsageType RHIGetGPUUsage = nullptr;
 bool bDriverDenylistMessageShown = false;
+bool GDynamicRHIIgnoredSM5Preference = false;
 
 static int32 GWarnOfBadDrivers = true;
 static FAutoConsoleVariableRef CVarWarnOfBadDrivers(
@@ -353,6 +354,7 @@ void RHIInit(bool bHasEditorToken)
 				FGenericCrashContext::SetEngineData(TEXT("RHI.InternalDriverVersion"), GRHIAdapterInternalDriverVersion);
 				FGenericCrashContext::SetEngineData(TEXT("RHI.DriverDate"), GRHIAdapterDriverDate);
 				FGenericCrashContext::SetEngineData(TEXT("RHI.FeatureLevel"), FeatureLevelString);
+				FGenericCrashContext::SetEngineData(TEXT("RHI.IgnoredSM5Preference"), GDynamicRHIIgnoredSM5Preference ? TEXT("True") : TEXT("False"));
 				FGenericCrashContext::SetEngineData(TEXT("RHI.GPUVendor"), RHIVendorIdToString());
 				FGenericCrashContext::SetEngineData(TEXT("RHI.DeviceId"), FString::Printf(TEXT("%04X"), GRHIDeviceId));
 			}
diff --git a/Engine/Source/Runtime/RHI/Private/Windows/WindowsDynamicRHI.cpp b/Engine/Source/Runtime/RHI/Private/Windows/WindowsDynamicRHI.cpp
--- a/Engine/Source/Runtime/RHI/Private/Windows/WindowsDynamicRHI.cpp
+++ b/Engine/Source/Runtime/RHI/Private/Windows/WindowsDynamicRHI.cpp
@@ -125,7 +125,27 @@ static TAutoConsoleVariable<bool> CVarIgnorePerformanceModeCheck(
 	false,
 	TEXT("Ignore performance mode check"),
 	ECVF_RenderThreadSafe
-	);
+);
+
+static TAutoConsoleVariable<bool> CVarForceSM5Support(
+	TEXT("r.ForceSM5Support"),
+	true,
+	TEXT("Whether to force SM5 support even if disabled in config"),
+	ECVF_ReadOnly
+);
+
+static bool IsSM5Allowed()
+{
+	bool bAllowSM5 = GIsEditor || CVarForceSM5Support.GetValueOnGameThread();
+	if (!bAllowSM5)
+	{
+		bool bBlockSM5 = false;
+		GConfig->GetBool(TEXT("RHI"), TEXT("bBlockSM5"), bBlockSM5, GEngineIni);
+		bAllowSM5 = !bBlockSM5;
+	}
+
+	return bAllowSM5;
+}
 
 static const TCHAR* ModuleNameFromWindowsRHI(EWindowsRHI InWindowsRHI)
 {
@@ -516,31 +536,43 @@ static TArray<EShaderPlatform> ParseShaderPlatformsConfig(const TCHAR* InSetting
 				return FDataDrivenShaderPlatformInfo::GetMaxFeatureLevel(ShaderPlatform) < ERHIFeatureLevel::SM5;
 			});
 	}
-	else if (FCString::Strcmp(InSettingName, TEXT("D3D11TargetedShaderFormats")) == 0)
+	else
 	{
-		bool bUseSM5AvoidGPUList = false;
-		if (GConfig->GetBool(TEXT("D3D11_SM5"), TEXT("bUseSM5AvoidGPUList"), bUseSM5AvoidGPUList, GEngineIni))
+		if (!IsSM5Allowed())
 		{
-			if (bUseSM5AvoidGPUList && ShouldAvoidD3D11_SM5())
-			{
-				ShaderPlatforms.RemoveAll(
-					[](EShaderPlatform ShaderPlatform)
-					{
-						return FDataDrivenShaderPlatformInfo::GetMaxFeatureLevel(ShaderPlatform) == ERHIFeatureLevel::SM5;
-					});
-			}
+			ShaderPlatforms.RemoveAll(
+				[](EShaderPlatform ShaderPlatform)
+				{
+					return FDataDrivenShaderPlatformInfo::GetMaxFeatureLevel(ShaderPlatform) == ERHIFeatureLevel::SM5;
+				});
 		}
 
-		bool bUseES31AvoidGPUList = false;
-		if (GConfig->GetBool(TEXT("D3D11_ES31"), TEXT("bUseES31AvoidGPUList"), bUseES31AvoidGPUList, GEngineIni))
+		if (FCString::Strcmp(InSettingName, TEXT("D3D11TargetedShaderFormats")) == 0)
 		{
-			if (bUseES31AvoidGPUList && ShouldAvoidD3D11_ES31())
+			bool bUseSM5AvoidGPUList = false;
+			if (GConfig->GetBool(TEXT("D3D11_SM5"), TEXT("bUseSM5AvoidGPUList"), bUseSM5AvoidGPUList, GEngineIni))
 			{
-				ShaderPlatforms.RemoveAll(
-					[](EShaderPlatform ShaderPlatform)
-					{
-						return FDataDrivenShaderPlatformInfo::GetMaxFeatureLevel(ShaderPlatform) == ERHIFeatureLevel::ES3_1;
-					});
+				if (bUseSM5AvoidGPUList && ShouldAvoidD3D11_SM5())
+				{
+					ShaderPlatforms.RemoveAll(
+						[](EShaderPlatform ShaderPlatform)
+						{
+							return FDataDrivenShaderPlatformInfo::GetMaxFeatureLevel(ShaderPlatform) == ERHIFeatureLevel::SM5;
+						});
+				}
+			}
+
+			bool bUseES31AvoidGPUList = false;
+			if (GConfig->GetBool(TEXT("D3D11_ES31"), TEXT("bUseES31AvoidGPUList"), bUseES31AvoidGPUList, GEngineIni))
+			{
+				if (bUseES31AvoidGPUList && ShouldAvoidD3D11_ES31())
+				{
+					ShaderPlatforms.RemoveAll(
+						[](EShaderPlatform ShaderPlatform)
+						{
+							return FDataDrivenShaderPlatformInfo::GetMaxFeatureLevel(ShaderPlatform) == ERHIFeatureLevel::ES3_1;
+						});
+				}
 			}
 		}
 	}
@@ -725,11 +757,19 @@ static TOptional<ERHIFeatureLevel::Type> GetPreferredFeatureLevel(EWindowsRHI Ch
 			FString PreferredFeatureLevelName;
 			if (GConfig->GetString(TEXT("D3DRHIPreference"), TEXT("PreferredFeatureLevel"), PreferredFeatureLevelName, GGameUserSettingsIni))
 			{
-				UE_LOG(LogRHI, Log, TEXT("Found D3DRHIPreference PreferredFeatureLevel: %s"), *PreferredFeatureLevelName);
+				UE_LOGF(LogRHI, Log, "Found D3DRHIPreference PreferredFeatureLevel: %ls", *PreferredFeatureLevelName);
 
 				if (PreferredFeatureLevelName == TEXT("sm5"))
 				{
-					PreferredFeatureLevel = ERHIFeatureLevel::SM5;
+					if (IsSM5Allowed())
+					{
+						PreferredFeatureLevel = ERHIFeatureLevel::SM5;
+					}
+					else
+					{
+						UE_LOGF(LogRHI, Log, "Ignoring PreferredFeatureLevel for SM5, which is not allowed");
+						GDynamicRHIIgnoredSM5Preference = true;
+					}
 				}
 				else if (PreferredFeatureLevelName == TEXT("sm6"))
 				{
@@ -741,7 +781,7 @@ static TOptional<ERHIFeatureLevel::Type> GetPreferredFeatureLevel(EWindowsRHI Ch
 				}
 				else
 				{
-					UE_LOG(LogRHI, Error, TEXT("unknown feature level name \"%s\" in game user settings, using default"), *PreferredFeatureLevelName);
+					UE_LOGF(LogRHI, Error, "unknown feature level name \"%ls\" in game user settings, using default", *PreferredFeatureLevelName);
 				}
 			}
 			else if (Config.IsFeatureLevelTargeted(ChosenRHI, ERHIFeatureLevel::ES3_1))
@@ -750,12 +790,12 @@ static TOptional<ERHIFeatureLevel::Type> GetPreferredFeatureLevel(EWindowsRHI Ch
 				bool bFoundPreference = GConfig->GetBool(TEXT("D3DRHIPreference"), TEXT("bPreferFeatureLevelES31"), bPreferFeatureLevelES31, GGameUserSettingsIni);
 				if (bFoundPreference)
 				{
-					UE_LOG(LogRHI, Log, TEXT("Found D3DRHIPreference bPreferFeatureLevelES31: %s"), bPreferFeatureLevelES31 ? TEXT("true") : TEXT("false"));
+					UE_LOGF(LogRHI, Log, "Found D3DRHIPreference bPreferFeatureLevelES31: %s", bPreferFeatureLevelES31 ? "true" : "false");
 				}
 
 				// Force low-spec users into performance mode but respect their choice once they have set a preference
 				bool bDefaultES31 = false;
-				if (!bFoundPreference && !CVarIgnorePerformanceModeCheck.GetValueOnAnyThread())
+				if (!bFoundPreference && !CVarIgnorePerformanceModeCheck.GetValueOnGameThread())
 				{
 					bDefaultES31 = ShouldDefaultToPerformanceMode(bOutForceD3D11);
 				}
@@ -853,11 +893,15 @@ static bool IsRHIAllowedAsDefault(EWindowsRHI InRHI, ERHIFeatureLevel::Type InFe
 			bAllowed &= BestGPUInfo.DedicatedVideoMemory >= MinDedicatedMemory;
 		}
 
-		bool bUseSM5PreferredGPUList = true;
-		GConfig->GetBool(TEXT("D3D12_SM6"), TEXT("bUseSM5PreferredGPUList"), bUseSM5PreferredGPUList, GEngineIni);
-		if (bUseSM5PreferredGPUList)
+		bool bAllowSM5 = IsSM5Allowed();
+		if (bAllowSM5)
 		{
-			bAllowed &= !ShouldDevicePreferSM5(BestGPUInfo.DeviceId);
+			bool bUseSM5PreferredGPUList = true;
+			GConfig->GetBool(TEXT("D3D12_SM6"), TEXT("bUseSM5PreferredGPUList"), bUseSM5PreferredGPUList, GEngineIni);
+			if (bUseSM5PreferredGPUList)
+			{
+				bAllowed &= !ShouldDevicePreferSM5(BestGPUInfo.DeviceId);
+			}
 		}
 	}
 	return bAllowed;
@@ -914,35 +958,52 @@ static TOptional<EWindowsRHI> ChoosePreferredRHI(EWindowsRHI InDefaultRHI)
 		FString PreferredRHIName;
 		if (GConfig->GetString(TEXT("D3DRHIPreference"), TEXT("PreferredRHI"), PreferredRHIName, GGameUserSettingsIni))
 		{
-			UE_LOG(LogRHI, Log, TEXT("Found D3DRHIPreference PreferredRHI: %s"), *PreferredRHIName);
-
-			if (PreferredRHIName == TEXT("dx12"))
+			bool bBlockSM5 = !IsSM5Allowed();
+			if (bBlockSM5)
 			{
-				RHIPreference = EWindowsRHI::D3D12;
+				FString PreferredFeatureLevelName;
+				bBlockSM5 =
+					GConfig->GetString(TEXT("D3DRHIPreference"), TEXT("PreferredFeatureLevel"), PreferredFeatureLevelName, GGameUserSettingsIni) &&
+					(PreferredFeatureLevelName == TEXT("sm5"));
 			}
-			else if (PreferredRHIName == TEXT("dx11"))
-			{
-				RHIPreference = EWindowsRHI::D3D11;
-			}
-			else if (PreferredRHIName == TEXT("vulkan"))
-			{
-				RHIPreference = EWindowsRHI::Vulkan;
-			}
-			else if (PreferredRHIName == TEXT("opengl"))
+
+			if (bBlockSM5)
 			{
-				RHIPreference = EWindowsRHI::OpenGL;
+				UE_LOGF(LogRHI, Log, "Ignoring PreferredRHI \"%ls\" because PreferredFeatureLevel specifies SM5, which is not allowed", *PreferredRHIName);
+				GDynamicRHIIgnoredSM5Preference = true;
 			}
 			else
 			{
-				UE_LOG(LogRHI, Error, TEXT("unknown RHI name \"%s\" in game user settings, using default"), *PreferredRHIName);
+				UE_LOGF(LogRHI, Log, "Found D3DRHIPreference PreferredRHI: %ls", *PreferredRHIName);
+
+				if (PreferredRHIName == TEXT("dx12"))
+				{
+					RHIPreference = EWindowsRHI::D3D12;
+				}
+				else if (PreferredRHIName == TEXT("dx11"))
+				{
+					RHIPreference = EWindowsRHI::D3D11;
+				}
+				else if (PreferredRHIName == TEXT("vulkan"))
+				{
+					RHIPreference = EWindowsRHI::Vulkan;
+				}
+				else if (PreferredRHIName == TEXT("opengl"))
+				{
+					RHIPreference = EWindowsRHI::OpenGL;
+				}
+				else
+				{
+					UE_LOGF(LogRHI, Error, "unknown RHI name \"%ls\" in game user settings, using default", *PreferredRHIName);
+				}
 			}
 		}
 		else if (GConfig->GetBool(TEXT("D3DRHIPreference"), TEXT("bUseD3D12InGame"), bUseD3D12InGame, GGameUserSettingsIni) &&
 			GConfig->GetBool(TEXT("D3DRHIPreference"), TEXT("bPreferFeatureLevelES31"), bPreferFeatureLevelES31, GGameUserSettingsIni))
 		{
 			// Respect legacy GameUserSettings values if present
-			UE_LOG(LogRHI, Log, TEXT("Found D3DRHIPreference bUseD3D12InGame: %s"), bUseD3D12InGame ? TEXT("true") : TEXT("false"));
-			UE_LOG(LogRHI, Log, TEXT("Found D3DRHIPreference bPreferFeatureLevelES31: %s"), bPreferFeatureLevelES31 ? TEXT("true") : TEXT("false"));
+			UE_LOGF(LogRHI, Log, "Found D3DRHIPreference bUseD3D12InGame: %s", bUseD3D12InGame ? "true" : "false");
+			UE_LOGF(LogRHI, Log, "Found D3DRHIPreference bPreferFeatureLevelES31: %s", bPreferFeatureLevelES31 ? "true" : "false");
 
 			if (bUseD3D12InGame)
 			{
@@ -1014,7 +1075,7 @@ static TOptional<EWindowsRHI> ChooseForcedRHI(TOptional<ERHIFeatureLevel::Type>
 			if (!Config.IsFeatureLevelTargeted(ForcedRHI.GetValue(), ForcedFeatureLevel.GetValue()))
 			{
 				FMessageDialog::Open(EAppMsgType::Ok, FText::Format(LOCTEXT("RHINotSupported", "{0} does not support {1}."), GetTextName(ForcedRHI.GetValue()), GetTextName(ForcedFeatureLevel.GetValue())));
-				UE_LOG(LogRHI, Fatal, TEXT("%s does not support %s."), GetLogName(ForcedRHI.GetValue()), GetLogName(ForcedFeatureLevel.GetValue()));
+				FPlatformMisc::RequestExit(true, TEXT("ChooseForcedRHI"));
 			}
 		}
 		else if ((ForcedRHI = Config.GetFirstRHIWithFeatureLevelSupport(ForcedFeatureLevel.GetValue())))
diff --git a/Engine/Source/Runtime/RHI/Public/DynamicRHI.h b/Engine/Source/Runtime/RHI/Public/DynamicRHI.h
--- a/Engine/Source/Runtime/RHI/Public/DynamicRHI.h
+++ b/Engine/Source/Runtime/RHI/Public/DynamicRHI.h
@@ -1494,6 +1494,7 @@ FDynamicRHI* PlatformCreateDynamicRHI();
 extern RHI_API const TCHAR* GetSelectedDynamicRHIModuleName(bool bCleanup = true);
 
 extern RHI_API bool GDynamicRHIFailedToInitializeAdvancedPlatform;
+extern RHI_API bool GDynamicRHIIgnoredSM5Preference;
 
 //
 // Helper for acquiring and releasing thread ownership of the RHI within a scope.
