Add support for transition from EPSOPrecacheInMemoryMode::Uncached to EPSOPrecacheInMemoryMode::Cached in pso precache cache (from bundles pso requests to actual component tracked requests)

#Synced-CL 51487274

diff --git a/Engine/Source/Runtime/RHI/Private/PipelineStateCache.cpp b/Engine/Source/Runtime/RHI/Private/PipelineStateCache.cpp
--- a/Engine/Source/Runtime/RHI/Private/PipelineStateCache.cpp
+++ b/Engine/Source/Runtime/RHI/Private/PipelineStateCache.cpp
@@ -291,6 +291,9 @@ struct FPSOCompilationDebugData
 #if WITH_RHI_BREADCRUMBS
 	FRHIBreadcrumbNode const* BreadcrumbRoot = nullptr;
 	FRHIBreadcrumbNode const* BreadcrumbNode = nullptr;
+
+	// Adding Shared ref to BreadcrumbNodeAllocator to keep it alive. During the warming of the PSO cache at launch, the RHICmdLis doesn't wait for the PSO to finish compiling.
+	TSharedPtr<const FRHIBreadcrumbAllocator> BreadcrumbNodeAllocator;
 #endif // WITH_RHI_BREADCRUMBS
 };
 
@@ -2581,11 +2584,23 @@ protected:
 			if (FindResult)
 			{			
 				check(!ShouldKeepPrecachedPSOsInMemory());
-
 				check(InMemoryMode == EPSOPrecacheInMemoryMode::Cached);				
-				check(FindResult->RefCount == 0);
 
-				FPlatformAtomics::InterlockedIncrement(&FindResult->RefCount);									
+				// If no refs then it should be in the clean up list when it's not cached yet
+				// (if it's still referenced then it goes from uncached to cached
+				bool bInCleanupListWhenNotCached = FindResult->RefCount == 0;
+				if (bInCleanupListWhenNotCached)
+				{
+					FPlatformAtomics::InterlockedIncrement(&FindResult->RefCount);
+				}
+				else
+				{
+					check(FindResult->InMemoryMode == EPSOPrecacheInMemoryMode::Uncached);
+				}
+
+				FindResult->InMemoryMode = EPSOPrecacheInMemoryMode::Cached;
+					
+				check(FindResult->RefCount == 1);
 				Result.RequestID = FindResult->RequestID;
 				check(Result.RequestID.IsValid());
 
@@ -2594,7 +2609,14 @@ protected:
 				EPSOPrecacheStateMask CurrentState = FindResult->ReadPSOPrecacheState();
 				if (!EnumHasAnyFlags(CurrentState, EPSOPrecacheStateMask::Uncached))
 				{	
-					verify(PrecachedPSOsToCleanup.Remove(InitializerHash) == 1);	
+					if (bInCleanupListWhenNotCached)
+					{
+						verify(PrecachedPSOsToCleanup.Remove(InitializerHash) == 1);
+					}
+					else
+					{
+						check(!PrecachedPSOsToCleanup.Contains(InitializerHash));
+					}
 
 					// Add async compile event if it's still compiling
 					if (!IsCompilationDone(CurrentState))
@@ -2606,6 +2628,7 @@ protected:
 
 				// Uncached - reset the state and kick the compilation again
 				FindResult->ResetStateMask();
+				FindResult->ResetPrecachePriority();
 
 				check(!PrecachedPSOsToCleanup.Contains(InitializerHash));
 			}
@@ -2779,12 +2802,20 @@ public:
 		check(FindResult);
 		if (FindResult->InMemoryMode == EPSOPrecacheInMemoryMode::Cached)
 		{
-			check(FindResult->RefCount > 0);
-			int32 NewRefCount = FPlatformAtomics::InterlockedDecrement(&FindResult->RefCount);
-			if (NewRefCount == 0)
+			EPSOPrecacheStateMask CurrentState = FindResult->ReadPSOPrecacheState();
+			if (!EnumHasAnyFlags(CurrentState, EPSOPrecacheStateMask::Uncached))
 			{
-				PrecachedPSOsToCleanup.Add(InitializerHash);
-				return true;
+				check(FindResult->RefCount > 0);
+				int32 NewRefCount = FPlatformAtomics::InterlockedDecrement(&FindResult->RefCount);
+				if (NewRefCount == 0)
+				{
+					PrecachedPSOsToCleanup.Add(InitializerHash);
+					return true;
+				}
+			}
+			else
+			{
+				check(FindResult->RefCount == 0);
 			}
 		}
 		
@@ -3025,21 +3056,14 @@ protected:
 		// Found and not uncached?
 		if (FindResult)
 		{
-			Result.RequestID = FindResult->RequestID;
-			check(Result.RequestID.IsValid());
-
-			if (ShouldKeepPrecachedPSOsInMemory())
+			if (!ShouldKeepPrecachedPSOsInMemory())
 			{
-				// If not compiled yet, then return the request ID so the caller can check the state
-				if (!IsCompilationDone(FindResult->ReadPSOPrecacheState()))
+				// If it's uncached and want to cache then return false because needs to be processed with write lock
+				if (FindResult->InMemoryMode != InMemoryMode && FindResult->InMemoryMode == EPSOPrecacheInMemoryMode::Uncached)
 				{
-					Result.AsyncCompileEvent = FindResult->PipelineState->GetCompletionEvent();
+					return false;
 				}
-			}
-			else
-			{
-				check(FindResult->InMemoryMode == InMemoryMode);
-
+				
 				// Special ref count tracking when kept in memory
 				if (InMemoryMode == EPSOPrecacheInMemoryMode::Cached)
 				{
@@ -3060,14 +3084,17 @@ protected:
 					FPlatformAtomics::InterlockedIncrement(&FindResult->RefCount);
 					check(!PrecachedPSOsToCleanup.Contains(InitializerHash));
 				}
-
-				// Add async compile event if it's still compiling
-				if (!IsCompilationDone(FindResult->ReadPSOPrecacheState()))
-				{
-					Result.AsyncCompileEvent = FindResult->PipelineState->GetCompletionEvent();
-				}
+			}
+			
+			// Add async compile event if it's still compiling
+			if (!IsCompilationDone(FindResult->ReadPSOPrecacheState()))
+			{
+				Result.AsyncCompileEvent = FindResult->PipelineState->GetCompletionEvent();
 			}
 
+			Result.RequestID = FindResult->RequestID;
+			check(Result.RequestID.IsValid());
+
 			return true;
 		}
 
@@ -3182,6 +3209,11 @@ protected:
 			return (EPSOPrecacheStateMask)FPlatformAtomics::InterlockedExchange((volatile int8*)&StateMask, (int8)EPSOPrecacheStateMask::None);
 		}
 
+		void ResetPrecachePriority()
+		{
+			FPlatformAtomics::InterlockedExchange((volatile int8*)&PrecachePriorityBitSet, 0);
+		}
+
 		inline EPSOPrecacheStateMask ReadPSOPrecacheState()
 		{
 			static_assert(sizeof(EPSOPrecacheStateMask) == sizeof(int8), "Fix the cast below");
@@ -4254,6 +4286,12 @@ FComputePipelineState* PipelineStateCache::GetAndOrCreateComputePipelineState(FR
 			PSOCompilationDebugData.BreadcrumbNode = PSOCompilationDebugData.BreadcrumbRoot
 				? RHICmdList.GetCurrentBreadcrumbRef()
 				: nullptr;
+
+			// Adding Shared ref to BreadcrumbNodeAllocator to keep it alive. During the warming of the PSO cache at launch, the RHICmdLis doesn't wait for the PSO to finish compiling.
+			if (PSOCompilationDebugData.BreadcrumbNode)
+			{
+				PSOCompilationDebugData.BreadcrumbNodeAllocator = PSOCompilationDebugData.BreadcrumbNode->Allocator->AsShared();
+			}
 		}
 #endif // WITH_RHI_BREADCRUMBS
 		
@@ -4789,6 +4827,12 @@ FGraphicsPipelineState* PipelineStateCache::GetAndOrCreateGraphicsPipelineState(
 			PSOCompilationDebugData.BreadcrumbNode = PSOCompilationDebugData.BreadcrumbRoot
 				? RHICmdList.GetCurrentBreadcrumbRef()
 				: nullptr;
+
+			// Adding Shared ref to BreadcrumbNodeAllocator to keep it alive. During the warming of the PSO cache at launch, the RHICmdLis doesn't wait for the PSO to finish compiling.
+			if (PSOCompilationDebugData.BreadcrumbNode)
+			{
+				PSOCompilationDebugData.BreadcrumbNodeAllocator = PSOCompilationDebugData.BreadcrumbNode->Allocator->AsShared();
+			}
 		}
 #endif // WITH_RHI_BREADCRUMBS
 
