Renable support for cached in memory precached psos

#Synced-CL 51591995

diff --git a/Engine/Source/Runtime/RHI/Private/PipelineStateCache.cpp b/Engine/Source/Runtime/RHI/Private/PipelineStateCache.cpp
--- a/Engine/Source/Runtime/RHI/Private/PipelineStateCache.cpp
+++ b/Engine/Source/Runtime/RHI/Private/PipelineStateCache.cpp
@@ -2555,8 +2555,7 @@ protected:
 	FPSOPrecacheRequestResult TryAddNewState(const TPrecachedPSOInitializer& Initializer, const FPSOPrecacheRequestParams& RequestParams, bool bDoAsyncCompile)
 	{
 		// Refcounted and max in memory PSOs can't be used together
-		// Using EPSOPrecacheInMemoryMode::Uncached for now until it's fixed on all platforms
-		EPSOPrecacheInMemoryMode InMemoryMode = EPSOPrecacheInMemoryMode::Uncached; //ShouldKeepPrecachedPSOsInMemory() ? EPSOPrecacheInMemoryMode::Uncached : RequestParams.InMemoryMode;
+		EPSOPrecacheInMemoryMode InMemoryMode = ShouldKeepPrecachedPSOsInMemory() ? EPSOPrecacheInMemoryMode::Uncached : RequestParams.InMemoryMode;
 
 		FPSOPrecacheRequestResult Result;
 		uint64 InitializerHash = TPrecachePipelineCacheDerived::PipelineStateInitializerHash(Initializer);
