Allow DX11 SM5 to be disabled

diff --git a/Engine/Source/Runtime/RHI/Private/Windows/WindowsDynamicRHI.cpp b/Engine/Source/Runtime/RHI/Private/Windows/WindowsDynamicRHI.cpp
--- a/Engine/Source/Runtime/RHI/Private/Windows/WindowsDynamicRHI.cpp
+++ b/Engine/Source/Runtime/RHI/Private/Windows/WindowsDynamicRHI.cpp
@@ -129,7 +129,7 @@ static TAutoConsoleVariable<bool> CVarIgnorePerformanceModeCheck(
 
 static TAutoConsoleVariable<bool> CVarForceSM5Support(
 	TEXT("r.ForceSM5Support"),
-	true,
+	false,
 	TEXT("Whether to force SM5 support even if disabled in config"),
 	ECVF_ReadOnly
 );
