[RHI/stats] Fixed stats emitted from GpuProfiler to have the correct stats message type. STAT_RHIDraws, STAT_RHIDispatches, STAT_RHIPrimitives and STAT_RHIVertices counters are registered using DECLARE_DWORD_COUNTER_STAT which declares them as EStatDataType::ST_int64. If instead of an int64 value, the AddMessage() is used with a uint64 value, the respective stat messages will be interpreted as EStatDataType::ST_Ptr and it will fail some assumptions in the Stats system (checkStats).

#Synced-CL 51759468

diff --git a/Engine/Source/Runtime/RHI/Private/GPUProfiler.cpp b/Engine/Source/Runtime/RHI/Private/GPUProfiler.cpp
--- a/Engine/Source/Runtime/RHI/Private/GPUProfiler.cpp
+++ b/Engine/Source/Runtime/RHI/Private/GPUProfiler.cpp
@@ -2072,10 +2072,10 @@ namespace UE::RHI::GPUProfiler
 			Manager.AccumulateFrameStats(DrawCounts);
 
 		#if STATS
-			Stats->AddMessage(GET_STATFNAME(STAT_RHIDraws     ), EStatOperation::Set, uint64(Total.NumDraws     ));
-			Stats->AddMessage(GET_STATFNAME(STAT_RHIDispatches), EStatOperation::Set, uint64(Total.NumDispatches));
-			Stats->AddMessage(GET_STATFNAME(STAT_RHIPrimitives), EStatOperation::Set, uint64(Total.NumPrimitives));
-			Stats->AddMessage(GET_STATFNAME(STAT_RHIVertices  ), EStatOperation::Set, uint64(Total.NumVertices  ));
+			Stats->AddMessage(GET_STATFNAME(STAT_RHIDraws     ), EStatOperation::Set, int64(Total.NumDraws     ));
+			Stats->AddMessage(GET_STATFNAME(STAT_RHIDispatches), EStatOperation::Set, int64(Total.NumDispatches));
+			Stats->AddMessage(GET_STATFNAME(STAT_RHIPrimitives), EStatOperation::Set, int64(Total.NumPrimitives));
+			Stats->AddMessage(GET_STATFNAME(STAT_RHIVertices  ), EStatOperation::Set, int64(Total.NumVertices  ));
 		#endif
 
 		#if CSV_PROFILER_STATS
