Updated the r.ShaderPipelineCache.PrintNewPSODescriptor description. It now calls out that it prints additional information not only for graphics PSOs, but also compute and raytracing.

#Synced-CL 51879730

diff --git a/Engine/Source/Runtime/RHI/Private/PipelineFileCache.cpp b/Engine/Source/Runtime/RHI/Private/PipelineFileCache.cpp
--- a/Engine/Source/Runtime/RHI/Private/PipelineFileCache.cpp
+++ b/Engine/Source/Runtime/RHI/Private/PipelineFileCache.cpp
@@ -142,7 +142,7 @@ static int32 GPSOFileCachePrintNewPSODescriptors = 0;
 static FAutoConsoleVariableRef CVarPSOFileCachePrintNewPSODescriptors(
 														   TEXT("r.ShaderPipelineCache.PrintNewPSODescriptors"),
 														   GPSOFileCachePrintNewPSODescriptors,
-														   TEXT("1 prints descriptions for all new PSO entries to the log/console while 0 does not. 2 prints additional details about graphics PSO. Defaults to 0."),
+														   TEXT("1 prints descriptions for all new PSO entries to the log/console while 0 does not. 2 prints additional details, such as the PSO type (graphics, compute, or raytracing), usage mask and bind count. Defaults to 0."),
 														   ECVF_Default
 														   );
 
