Refactored high level submission logic to allow for completed IRHIPlatformCommandLists to be passed to the platform RHI as soon as they are finalized, rather than waiting for an explicit DispatchToRHIThread immediate command list flush on the render thread.
- Expedites the submission of work to the GPU, reducing latency in the pipeline. The new mode is enabled when "r.RHICmd.ParallelTranslate.BatchSubmissions" is false (which is the default).
- The dispatch pipe builds a linked list of translate chains and enqueues a submit task to run as soon as the first translate completes.
- The submit task pulls as many completed chains from the linked list as possible. This means that completed chains reach the RHI much sooner than before, and we also still maintain some amount of batching of calls to RHISubmitCommandLists.
- If the submit task finds a translate chain that has not completed yet, it calls ExtendTask() to resume processing of the queue once the chain is done.

#Synced-CL 51451766

diff --git a/Engine/Source/Runtime/RHI/Private/RHICommandList.cpp b/Engine/Source/Runtime/RHI/Private/RHICommandList.cpp
--- a/Engine/Source/Runtime/RHI/Private/RHICommandList.cpp
+++ b/Engine/Source/Runtime/RHI/Private/RHICommandList.cpp
@@ -70,6 +70,12 @@ TAutoConsoleVariable<bool> CVarRHICmdParallelTranslateEnable(
 	TEXT("When true, allows recorded RHI command lists to be translated in parallel, on supported platforms. ")
 	TEXT("Setting this to false will make all command lists translate on the RHI thread."));
 
+TAutoConsoleVariable<bool> CVarRHICmdParallelTranslateBatchSubmissions(
+	TEXT("r.RHICmd.ParallelTranslate.BatchSubmissions"),
+	false,
+	TEXT("When true, generated platform command lists are not submitted to the RHI until the immediate command list is dispatched with the SubmitToGPU flag. ")
+	TEXT("When false, platform command lists are handed to the RHI immediately after finalizing (but still in submission order), reducing pipeline latency. "));
+
 TAutoConsoleVariable<int32> CVarRHICmdParallelTranslateMaxCommandsPerTranslate(
 	TEXT("r.RHICmd.ParallelTranslate.MaxCommandsPerTranslate"),
 	256,
@@ -518,9 +524,9 @@ void FRHICommandListBase::Execute()
 	}
 }
 
-struct FRHICommandListExecutor::FTaskPipe::FTask
+struct FRHICommandListExecutor::FTaskPipe::FTask final : public ITask
 {
-	TFunction<void()> Lambda;
+	FFunc Lambda;
 	FGraphEventArray Prereqs;
 
 	std::atomic<FTask*> Next{ nullptr };
@@ -529,7 +535,9 @@ struct FRHICommandListExecutor::FTaskPipe::FTask
 	ENamedThreads::Type LogicalThread;
 	ENamedThreads::Type ActualThread;
 
-	FTask(ENamedThreads::Type NamedThread, FGraphEventArray&& Prereqs, TFunction<void()>&& Lambda)
+	bool bExtended = false;
+
+	FTask(ENamedThreads::Type NamedThread, FGraphEventArray&& Prereqs, FFunc&& Lambda)
 		: Lambda       (MoveTemp(Lambda))
 		, Prereqs      (MoveTemp(Prereqs))
 		, LogicalThread(NamedThread)
@@ -551,6 +559,15 @@ struct FRHICommandListExecutor::FTaskPipe::FTask
 		}
 	}
 
+	void ExtendTask(FGraphEventArray&& NewPrereqs, FFunc&& NewLambda) override
+	{
+		check(!bExtended);
+		Prereqs = MoveTemp(NewPrereqs);
+		Lambda = MoveTemp(NewLambda);
+
+		bExtended = true;
+	}
+
 	// Memory pool for fast alloc of these FTask structs
 	static TLockFreePointerListUnordered<void, PLATFORM_CACHE_LINE_SIZE> MemoryPool;
 
@@ -592,7 +609,7 @@ FGraphEventRef FRHICommandListExecutor::FTaskPipe::LaunchTask(FTask* Task) const
 	);
 }
 
-void FRHICommandListExecutor::FTaskPipe::Enqueue(ENamedThreads::Type NamedThread, FGraphEventArray&& Prereqs, TFunction<void()>&& Lambda)
+void FRHICommandListExecutor::FTaskPipe::Enqueue(ENamedThreads::Type NamedThread, FGraphEventArray&& Prereqs, FTaskPipe::FFunc&& Lambda)
 {
 	if (LastThread != NamedThread)
 	{
@@ -709,7 +726,13 @@ Restart:
 	}
 
 	// All prereqs are resolved (if any). Run the task.
-	Task->Lambda();
+	Task->bExtended = false;
+	Task->Lambda(*Task);
+
+	if (Task->bExtended)
+	{
+		goto Restart;
+	}
 
 	// Attempt to close the command chain.
 	FTask* Expected = nullptr;
@@ -738,9 +761,10 @@ void FRHICommandListExecutor::AddNextDispatchPrerequisite(FGraphEventRef Prereq)
 	NextDispatchTaskPrerequisites.Add(MoveTemp(Prereq));
 }
 
-FRHICommandListExecutor::FTaskPipe* FRHICommandListExecutor::EnqueueDispatchTask(FGraphEventArray&& Prereqs, TFunction<void()>&& Lambda)
+void FRHICommandListExecutor::EnqueueDispatchTask(FGraphEventArray&& Prereqs, FTaskPipe::FFunc&& Lambda)
 {
 	check(IsInRenderingThread());
+
 	ENamedThreads::Type NamedThread = !AllowParallel()
 		? ENamedThreads::GetRenderThread_Local()
 		: ENamedThreads::AnyHiPriThreadHiPriTask;
@@ -749,272 +773,382 @@ FRHICommandListExecutor::FTaskPipe* FRHICommandListExecutor::EnqueueDispatchTask
 	NextDispatchTaskPrerequisites.Append(MoveTemp(Prereqs));
 
 	DispatchPipe.Enqueue(NamedThread, MoveTemp(NextDispatchTaskPrerequisites), MoveTemp(Lambda));
-	return &DispatchPipe;
 }
 
-FRHICommandListExecutor::FTaskPipe* FRHICommandListExecutor::FTranslateState::GetTranslateTaskPipe(ENamedThreads::Type& NamedThread)
+FRHICommandListExecutor::FTaskPipe* FRHICommandListExecutor::FTranslateState::EnqueueTranslateTask(FGraphEventArray&& Prereqs, FTaskPipe::FFunc&& Lambda)
 {
-	NamedThread = ENamedThreads::AnyHiPriThreadHiPriTask;
-	FTaskPipe* Pipe = &TranslatePipe;
+	// This is called on the dispatch pipe
 
+	ENamedThreads::Type NamedThread = ENamedThreads::AnyHiPriThreadHiPriTask;
+	FTaskPipe* Pipe = &TranslatePipe;
+	
 	if (!GRHICommandList.AllowParallel())
 	{
 		NamedThread = ENamedThreads::GetRenderThread_Local();
-		Pipe = &GRHICommandList.RHIThreadPipe;
+		Pipe = bAllowParallelTranslate
+			? &GRHICommandList.SingleTranslatePipe
+			: &GRHICommandList.SubmitPipe;
 	}
 	else if (!bParallel)
 	{
 		NamedThread = ENamedThreads::RHIThread;
-		Pipe = &GRHICommandList.RHIThreadPipe;
+		Pipe = bAllowParallelTranslate
+			? &GRHICommandList.SingleTranslatePipe
+			: &GRHICommandList.SubmitPipe;
 	}
-	
-	return Pipe;
-}
-
-FRHICommandListExecutor::FTaskPipe* FRHICommandListExecutor::FTranslateState::EnqueueTranslateTask(FGraphEventArray&& Prereqs, TFunction<void()>&& Lambda)
-{
-	// This is called on the dispatch thread
-	ENamedThreads::Type NamedThread = ENamedThreads::AnyHiPriThreadHiPriTask;
-	FTaskPipe* Pipe = GetTranslateTaskPipe(NamedThread);
 
 	Pipe->Enqueue(NamedThread, MoveTemp(Prereqs), MoveTemp(Lambda));
 	return Pipe;
 }
 
-FRHICommandListExecutor::FTaskPipe* FRHICommandListExecutor::EnqueueSubmitTask(FGraphEventArray&& Prereqs, TFunction<void()>&& Lambda)
+void FRHICommandListExecutor::EnqueueSubmitTask(FGraphEventArray&& Prereqs, FTaskPipe::FFunc&& Lambda)
 {
-	// This is called on the dispatch thread
+	// This is called on the dispatch pipe
 
 	ENamedThreads::Type NamedThread = !AllowParallel()
 		? ENamedThreads::GetRenderThread_Local()
 		: ENamedThreads::RHIThread;
 
-	RHIThreadPipe.Enqueue(NamedThread, MoveTemp(Prereqs), MoveTemp(Lambda));
-	return &RHIThreadPipe;
+	SubmitPipe.Enqueue(NamedThread, MoveTemp(Prereqs), MoveTemp(Lambda));
 }
 
-FGraphEventRef FRHICommandListExecutor::FSubmitState::FinalizeCurrent()
+// -----------------------------------------------------------------------------------------------------------
+// -------------                                  Dispatch Pipe                                  -------------
+// -----------------------------------------------------------------------------------------------------------
+
+FRHICommandListExecutor::FTranslateState::FTranslateState(const FSubmitState& SubmitState, FRHICommandListBase* FirstCmdList)
+	: bAllowParallelTranslate(SubmitState.bAllowParallelTranslate)
+	, bParallel              (FirstCmdList->AllowParallelTranslate())
+	, bUsingSubCmdLists      (FirstCmdList->IsSubCommandList())
 {
-	FGraphEventRef Event = CurrentTranslateJob->Finalize();
-	
-	TranslateEvents.Add(Event);
-	CurrentTranslateJob = nullptr;
-	
-	return Event;
 }
 
-bool FRHICommandListExecutor::FSubmitState::ShouldSplitTranslateJob(FRHICommandListBase* CmdList)
+void FRHICommandListExecutor::FSubmitState::Dispatch_ProcessCommandList(FRHICommandListBase* CmdList)
 {
+	SCOPED_NAMED_EVENT(RHI_Dispatch, FColor::Yellow);
+
+	check(CmdList->DispatchEvent->IsComplete());
+	Dispatch_PrepareBreadcrumbs(CmdList);
+
+	CmdList->bAllowParallelTranslate = bAllowParallelTranslate;
+
 	//
-	// Determine if the current translate batch should be closed, i.e.:
+	// Determine if the current translate chain should be closed, i.e.:
 	//    - We've exceeded the threshold number of RHI commands.
-	//    - The next command list requires single-threaded execution, but we're in a parallel batch.
+	//    - The next command list requires single-threaded execution, but we're in a parallel chain.
 	//
+	if (TranslateJobs.Current)
+	{
+		bool bShouldSplitParallel = bAllowSingleParallelCombine
+				? TranslateJobs.Current->bParallel && !CmdList->AllowParallelTranslate()  // Only start a new translate job if we were parallel, but the new command list requires single thread.
+				: TranslateJobs.Current->bParallel !=  CmdList->AllowParallelTranslate(); // Always start a new translate job if bParallel is different. Avoids batching parallel work into single thread translates.
 	
-	bool bShouldSplitParallel = bAllowSingleParallelCombine
-			? CurrentTranslateJob->bParallel && !CmdList->AllowParallelTranslate() // Only start a new translate job if we were parallel, but the new command list requires single thread.
-			: CurrentTranslateJob->bParallel !=  CmdList->AllowParallelTranslate(); // Always start a new translate job if bParallel is different. Avoids batching parallel work into single thread translates.
-	
-	bool bShouldSplitForThreshold = MaxCommandsPerTranslate  < 0 ? true  :
-									MaxCommandsPerTranslate == 0 ? false :
-									(CurrentTranslateJob->NumCommands + CmdList->NumCommands) > uint32(MaxCommandsPerTranslate);
+		bool bShouldSplitForThreshold = MaxCommandsPerTranslate  < 0 ? true  :
+										MaxCommandsPerTranslate == 0 ? false :
+										(TranslateJobs.Current->NumCommands + CmdList->NumCommands) > uint32(MaxCommandsPerTranslate);
 	
-	bool bShouldSplitForParentChild = 	(CmdList->IsSubCommandList() && !CurrentTranslateJob->bUsingSubCmdLists) ||
-										(!CmdList->IsSubCommandList() && CurrentTranslateJob->bUsingSubCmdLists);
+		bool bShouldSplitForParentChild = ( CmdList->IsSubCommandList() && !TranslateJobs.Current->bUsingSubCmdLists) ||
+										  (!CmdList->IsSubCommandList() &&  TranslateJobs.Current->bUsingSubCmdLists);
 	
-	return bShouldSplitParallel || bShouldSplitForThreshold || bShouldSplitForParentChild;
-}
+		if (bShouldSplitParallel || bShouldSplitForThreshold || bShouldSplitForParentChild)
+		{
+			//
+			// Close the current translate chain
+			//
 
-void FRHICommandListExecutor::FSubmitState::ConditionalSplitTranslateJob(FRHICommandListBase* CmdList)
-{
-	if (CurrentTranslateJob && ShouldSplitTranslateJob(CmdList))
+			bool bAddChildWait = TranslateJobs.Current->bUsingSubCmdLists && GRHIParallelRHIExecuteChildWait;
+			bool bAddParentWait = !TranslateJobs.Current->bUsingSubCmdLists && CmdList->IsSubCommandList() && GRHIParallelRHIExecuteParentWait;
+
+			FGraphEventRef Event = Dispatch_CloseTranslateChain(false);
+			if (bAddChildWait)
+
+			{
+				ChildGraphEvents.Add(Event);
+			}
+			else if (bAddParentWait)
+			{
+				BeginGraphEvent = Event;
+			}
+		}
+	}
+
+	//
+	// We may have closed the current chain above, or this is the first submission.
+	// Start a new chain if we don't have one.
+	//
+	if (!TranslateJobs.Current)
 	{
-		bool bAddChildWait = CurrentTranslateJob->bUsingSubCmdLists && GRHIParallelRHIExecuteChildWait; 
-		bool bAddParentWait = !CurrentTranslateJob->bUsingSubCmdLists && CmdList->IsSubCommandList() && GRHIParallelRHIExecuteParentWait;
-		FGraphEventRef Event = FinalizeCurrent();
+		TranslateJobs.Current = TranslateJobs.All.Emplace_GetRef(MakeUnique<FTranslateState>(*this, CmdList)).Get();
+	}
+
+	check(TranslateJobs.Current->bShouldFinalize);
+	TranslateJobs.Current->bShouldFinalize = !(GRHIParallelRHIExecuteChildWait && CmdList->ParallelRenderPassBegin.IsValid());
+	TranslateJobs.Current->NumCommands += CmdList->NumCommands;
+
+	// Create the translation task for this command list.
+	FGraphEventArray Prereqs;
+	if (!CmdList->AllowParallelTranslate())
+	{
+		// Wait for all previous translate jobs to complete
+		Prereqs = MoveTemp(TranslateJobs.Events);
+	}
 		
-		if (bAddChildWait)
-		{
-			ChildGraphEvents.Add(Event);
-		}
-		else if (bAddParentWait)
+	// Handle adding prerequisites for Parent end waiting on children
+	if (CmdList->ParallelRenderPassEnd.IsValid())
+	{
+		if (GRHIParallelRHIExecuteChildWait)
 		{
-			BeginGraphEvent = Event;
+			Prereqs = TranslateJobs.Events;
+				
+			Prereqs.Add(BeginGraphEvent);
+			Prereqs.Append(ChildGraphEvents);
 		}
+
+		ChildGraphEvents.Empty();
+		BeginGraphEvent = nullptr;
 	}
-	
-	if (!CurrentTranslateJob)
+	else if (CmdList->IsSubCommandList() && GRHIParallelRHIExecuteParentWait)
+	{
+		// Children need to wait on the parent begin
+		Prereqs = TranslateJobs.Events;
+		Prereqs.Add(BeginGraphEvent);
+	}
+
+	// Also wait for the previous mutate fence (blocks future translates until the fence has signalled).
+	Prereqs.Add(GRHICommandList.LastMutate);
+
+	if (CmdList->LastLockFenceCommand)
 	{
-		// Start a new translate job
-		CurrentTranslateJob = TranslateJobs.Emplace_GetRef(MakeUnique<FTranslateState>()).Get();
-		CurrentTranslateJob->bParallel = CmdList->AllowParallelTranslate();
-		CurrentTranslateJob->bUsingSubCmdLists = CmdList->IsSubCommandList();
+		check(CmdList->LastLockFenceCommand->Fence);
+		GRHICommandList.LastMutate = CmdList->LastLockFenceCommand->Fence;
 	}
 
-	check(CurrentTranslateJob->bShouldFinalize);
-	CurrentTranslateJob->bShouldFinalize = !(GRHIParallelRHIExecuteChildWait && CmdList->ParallelRenderPassBegin.IsValid());
+	if (!GRHIGlobals.SupportsConcurrentTranslateAndSubmit)
+	{
+		//
+		// Don't start new translations until all prior submissions have been closed.
+		// This is because work inside RHIEndFrame needs to complete on the RHI thread before any new translates can start.
+		// 
+		// Also some RHI commands directly submit to GPU queues from within the RHI (during RHICmdList translation).
+		// Not waiting for prior submits means these internal submissions can happen out-of-order with respect to other translations.
+		// E.g. some platform implementations of EndDrawingViewport() submit directly to the GPU to perform the flip / present.
+		//
+		Prereqs.Add(GRHICommandList.LastClose);
+	}
 
-	CurrentTranslateJob->NumCommands += CmdList->NumCommands;
+	TranslateJobs.Current->EnqueueTranslateTask(MoveTemp(Prereqs),
+		[TranslateJob = TranslateJobs.Current, CmdList](FTaskPipe::ITask&)
+		{
+			TranslateJob->Translate_ExecuteCommandList(CmdList);
+		});
 }
 
-void FRHICommandListExecutor::FSubmitState::Dispatch(FRHICommandListBase* CmdList)
+FGraphEventRef FRHICommandListExecutor::FSubmitState::Dispatch_CloseTranslateChain(bool bLastChainInSubmission)
 {
-	check(CmdList->DispatchEvent->IsComplete());
-	CmdList->bAllowParallelTranslate = bAllowParallelTranslate;
+	// Tell the current chain there will be no more command lists. This finalizes the RHI contexts.
+	FTaskPipe* TranslatePipe = TranslateJobs.Current->EnqueueTranslateTask({},
+		[Job = TranslateJobs.Current](FTaskPipe::ITask&)
+		{
+			Job->Translate_CloseChain();
+		});
 
-#if WITH_RHI_BREADCRUMBS
-	CmdList->CPUBreadcrumbState.bEmitBreadcrumbs = bEmitBreadcrumbs;
-#endif
+	FGraphEventRef TranslateEvent {};
+	if (TranslatePipe == &GRHICommandList.SubmitPipe)
+	{
+		// Don't close the pipe if we got scheduled on the submit pipe, to avoid splitting tasks.
+		// Since the pipes guarantee FIFO order, we don't need an event for the submission to await.
+	}
+	else
+	{
+		// Close the translate pipe. This gives us a TaskGraph event that is signalled when the
+		// translation has completed, which submission will await.
+		TranslateEvent = TranslatePipe->Close();
 
-	ConditionalSplitTranslateJob(CmdList);	
+		TranslateJobs.Current->TranslateEvent = TranslateEvent;
+		TranslateJobs.Events.Add(TranslateEvent);
+	}
 
+	// Link the completed translate chain into the submission queue
+	// (the linked list of jobs following FTranslateState::Next).
+	bool bStartTask = false;
+	if (!TranslateJobs.First)
 	{
-	#if WITH_RHI_BREADCRUMBS
-		// Fixup unknown breadcrumb parents
-		for (FRHIBreadcrumbNode* Node : CmdList->CPUBreadcrumbState.UnknownParentList.IterateAndUnlink())
+		check(!TranslateJobs.Previous);
+		TranslateJobs.First = TranslateJobs.Current;
+		bStartTask = true;
+	}
+	else
+	{
+		check(TranslateJobs.Previous);
+
+		// Attempt to atomically append this translate job to the end of the submission queue.
+		FTranslateState* Expected = nullptr;	
+		if (TranslateJobs.Previous->Next.compare_exchange_strong(Expected, TranslateJobs.Current))
 		{
-			check(Node->GetParent() == FRHIBreadcrumbNode::Sentinel);
-			Node->SetParent(GRHICommandList.Breadcrumbs.CPU.Current);
+			// Succesfully appended to the queue before the previous task finished.
+			check(Expected == nullptr);
 		}
-
+		else
 		{
-			// Grab the first breadcrumb in both the CPU and GPU pipeline stacks, and add references to them.
-			FRHIBreadcrumbNode* CPUFirst = GRHICommandList.Breadcrumbs.CPU.Current;
-			check(CPUFirst != FRHIBreadcrumbNode::Sentinel);
-			if (CPUFirst)
-			{
-				CmdList->BreadcrumbAllocatorRefs.AddUnique(CPUFirst->Allocator);
-			}
+			// Failed to append. The previous submission task has already run, and won't process the current chain.
+			// We'll need to start a new task.
 
-			TRHIPipelineArray<FRHIBreadcrumbNode*> GPUFirst;
-			for (ERHIPipeline Pipeline : MakeFlagsRange(ERHIPipeline::All))
-			{
-				FRHIBreadcrumbNode* Node = GRHICommandList.Breadcrumbs.GPU[Pipeline].Current;
-				check(Node != FRHIBreadcrumbNode::Sentinel);
+			// Submit_SubmitToGPU uses the job's 'this' pointer as a sentinel to indicate the closure of the queue.
+			check(Expected == TranslateJobs.Previous);
 
-				GPUFirst[Pipeline] = Node;
+			TranslateJobs.First = TranslateJobs.Current;
+			bStartTask = true;
+		}
+	}
 
-				if (Node)
-				{
-					CmdList->BreadcrumbAllocatorRefs.AddUnique(Node->Allocator);
-				}
-			}
+	TranslateJobs.Previous = TranslateJobs.Current;
+	TranslateJobs.Current = nullptr;
 
-			// Walk the ActivatePipeline commands, resolve unknown targets, and update per-pipe pointers.
-			for (FRHICommandListBase::FActivatePipelineCommand* Command = CmdList->ActivatePipelineCommands.First; Command; Command = Command->Next)
+	//
+	// When batching submissions, the submit task only starts when we've closed the last translate chain (bLastChainInSubmission == true)
+	// 
+	// Otherwise, we submit work to the GPU immediately. In this case, we may need to start a new task (bStartTask == true) if:
+	//         1) there wasn't already one running
+	//     OR  2) we failed to append this translate chain to the queue
+	//
+	if (bBatchSubmissions ? bLastChainInSubmission : bStartTask)
+	{
+		check(TranslateJobs.First);
+
+		// Start a new task to process the submission queue from 'First'
+		GRHICommandList.EnqueueSubmitTask({ TranslateEvent },
+			[this, FirstJob = TranslateJobs.First](FTaskPipe::ITask& Task)
 			{
-				if (Command->Target == FRHIBreadcrumbNode::Sentinel)
-				{
-					Command->Target = GRHICommandList.Breadcrumbs.CPU.Current;
-				}
-				else
-				{
-					GRHICommandList.Breadcrumbs.CPU.Current = Command->Target;
-				}
+				Submit_SubmitToGPU(FirstJob, Task);
+			});
+	}
 
-				for (ERHIPipeline Pipeline : MakeFlagsRange(Command->Pipelines))
-				{
-					GRHICommandList.Breadcrumbs.GPU[Pipeline].Current = GRHICommandList.Breadcrumbs.CPU.Current;
-				}
-			}
+	return TranslateEvent;
+}
 
-			for (ERHIPipeline Pipeline : MakeFlagsRange(ERHIPipeline::All))
-			{
-				if (CmdList->GPUBreadcrumbState[Pipeline].Latest.IsSet())
-				{
-					// A Begin/End happened on this pipeline after the last ActivatePipeline command.
-					FRHIBreadcrumbNode* Node = CmdList->GPUBreadcrumbState[Pipeline].Latest.GetValue();
-					if (Node != FRHIBreadcrumbNode::Sentinel)
-					{
-						GRHICommandList.Breadcrumbs.GPU[Pipeline].Current = Node;
-					}
-					else
-					{
-						GRHICommandList.Breadcrumbs.GPU[Pipeline].Current = GPUFirst[Pipeline];
-					}
-				}
+void FRHICommandListExecutor::FSubmitState::Dispatch_CloseSubmission(TOptional<int64> LocalStatsFrame)
+{
+	SCOPED_NAMED_EVENT(RHI_Dispatch_CloseSubmission, FColor::Purple);
 
-				// Rewind GPU state
-				CmdList->GPUBreadcrumbState[Pipeline].Current = GPUFirst[Pipeline];
-			}
+	// We still have a translate chain open. Close it.
+	Dispatch_CloseTranslateChain(true);
 
-			if (CmdList->CPUBreadcrumbState.Current != FRHIBreadcrumbNode::Sentinel)
-			{
-				GRHICommandList.Breadcrumbs.CPU.Current = CmdList->CPUBreadcrumbState.Current;
-			}
-			else
-			{
-				GRHICommandList.Breadcrumbs.CPU.Current = CPUFirst;
-			}
+	// Forward the close command to the submit pipe.
+	FSubmitState::FCloseArgs CloseArgs;
+#if WITH_RHI_BREADCRUMBS
+	for (ERHIPipeline Pipeline : MakeFlagsRange(ERHIPipeline::All))
+	{
+		CloseArgs.GPUBreadcrumbs[Pipeline] = GRHICommandList.Breadcrumbs.GPU[Pipeline].Current;
+	}
+#endif
+#if STATS
+	CloseArgs.StatsFrame = LocalStatsFrame;
+#endif
 
-			// Rewind CPU state
-			CmdList->CPUBreadcrumbState.Current = CPUFirst;
-			CmdList->PersistentState.LocalBreadcrumb = CPUFirst;
-		}
-	#endif // WITH_RHI_BREADCRUMBS
+	GRHICommandList.LastClose = CompletionEvent;
 
-		FGraphEventArray Prereqs;
-		if (!CmdList->AllowParallelTranslate())
+	GRHICommandList.EnqueueSubmitTask({},
+		[this, CloseArgs](FTaskPipe::ITask&) mutable
 		{
-			// Wait for all previous translate jobs to complete
-			Prereqs = MoveTemp(TranslateEvents);
-		}
-		
-		// Handle adding prerequisites for Parent end waiting on children
-		if (CmdList->ParallelRenderPassEnd.IsValid())
+			Submit_Close(CloseArgs);
+		});
+}
+
+void FRHICommandListExecutor::FSubmitState::Dispatch_PrepareBreadcrumbs(FRHICommandListBase* CmdList)
+{
+#if WITH_RHI_BREADCRUMBS
+	CmdList->CPUBreadcrumbState.bEmitBreadcrumbs = bEmitBreadcrumbs;
+
+	// Fixup unknown breadcrumb parents
+	for (FRHIBreadcrumbNode* Node : CmdList->CPUBreadcrumbState.UnknownParentList.IterateAndUnlink())
+	{
+		check(Node->GetParent() == FRHIBreadcrumbNode::Sentinel);
+		Node->SetParent(GRHICommandList.Breadcrumbs.CPU.Current);
+	}
+
+	// Grab the first breadcrumb in both the CPU and GPU pipeline stacks, and add references to them.
+	FRHIBreadcrumbNode* CPUFirst = GRHICommandList.Breadcrumbs.CPU.Current;
+	check(CPUFirst != FRHIBreadcrumbNode::Sentinel);
+	if (CPUFirst)
+	{
+		CmdList->BreadcrumbAllocatorRefs.AddUnique(CPUFirst->Allocator);
+	}
+
+	TRHIPipelineArray<FRHIBreadcrumbNode*> GPUFirst;
+	for (ERHIPipeline Pipeline : MakeFlagsRange(ERHIPipeline::All))
+	{
+		FRHIBreadcrumbNode* Node = GRHICommandList.Breadcrumbs.GPU[Pipeline].Current;
+		check(Node != FRHIBreadcrumbNode::Sentinel);
+
+		GPUFirst[Pipeline] = Node;
+
+		if (Node)
 		{
-			if (GRHIParallelRHIExecuteChildWait)
-			{
-				Prereqs = TranslateEvents;
-				
-				Prereqs.Add(BeginGraphEvent);
-				Prereqs.Append(ChildGraphEvents);
-			}
-			
-			ChildGraphEvents.Empty();
-			BeginGraphEvent = nullptr;
+			CmdList->BreadcrumbAllocatorRefs.AddUnique(Node->Allocator);
 		}
-		else if (CmdList->IsSubCommandList() && GRHIParallelRHIExecuteParentWait)
+	}
+
+	// Walk the ActivatePipeline commands, resolve unknown targets, and update per-pipe pointers.
+	for (FRHICommandListBase::FActivatePipelineCommand* Command = CmdList->ActivatePipelineCommands.First; Command; Command = Command->Next)
+	{
+		if (Command->Target == FRHIBreadcrumbNode::Sentinel)
 		{
-			// Children need to wait on the parent begin
-			Prereqs = TranslateEvents;
-			Prereqs.Add(BeginGraphEvent);
+			Command->Target = GRHICommandList.Breadcrumbs.CPU.Current;
 		}
-
-		// Also wait for the previous mutate fence (blocks future translates until the fence has signalled).
-		Prereqs.Add(GRHICommandList.LastMutate);
-
-		if (CmdList->LastLockFenceCommand)
+		else
 		{
-			check(CmdList->LastLockFenceCommand->Fence);
-			GRHICommandList.LastMutate = CmdList->LastLockFenceCommand->Fence;
+			GRHICommandList.Breadcrumbs.CPU.Current = Command->Target;
 		}
 
-		if (!GRHIGlobals.SupportsConcurrentTranslateAndSubmit)
+		for (ERHIPipeline Pipeline : MakeFlagsRange(Command->Pipelines))
 		{
-			//
-			// Don't start new translations until all prior submissions have been made.
-			// This is because work inside RHIEndFrame needs to complete on the RHI thread before any new translates can start.
-			// 
-			// Also some RHI commands directly submit to GPU queues from within the RHI (during RHICmdList translation).
-			// Not waiting for prior submits means these internal submissions can happen out-of-order with respect to other translations.
-			// E.g. some platform implementations of EndDrawingViewport() submit directly to the GPU to perform the flip / present.
-			//
-			Prereqs.Add(GRHICommandList.LastSubmit);
+			GRHICommandList.Breadcrumbs.GPU[Pipeline].Current = GRHICommandList.Breadcrumbs.CPU.Current;
 		}
+	}
 
-		CurrentTranslateJob->EnqueueTranslateTask(MoveTemp(Prereqs),
-			[TranslateJob = CurrentTranslateJob, CmdList]()
+	for (ERHIPipeline Pipeline : MakeFlagsRange(ERHIPipeline::All))
+	{
+		if (CmdList->GPUBreadcrumbState[Pipeline].Latest.IsSet())
+		{
+			// A Begin/End happened on this pipeline after the last ActivatePipeline command.
+			FRHIBreadcrumbNode* Node = CmdList->GPUBreadcrumbState[Pipeline].Latest.GetValue();
+			if (Node != FRHIBreadcrumbNode::Sentinel)
 			{
-				SCOPED_NAMED_EVENT(RHI_Translate, FColor::Orange);
-				TranslateJob->Translate(CmdList);
+				GRHICommandList.Breadcrumbs.GPU[Pipeline].Current = Node;
 			}
-		);
+			else
+			{
+				GRHICommandList.Breadcrumbs.GPU[Pipeline].Current = GPUFirst[Pipeline];
+			}
+		}
+
+		// Rewind GPU state
+		CmdList->GPUBreadcrumbState[Pipeline].Current = GPUFirst[Pipeline];
+	}
+
+	if (CmdList->CPUBreadcrumbState.Current != FRHIBreadcrumbNode::Sentinel)
+	{
+		GRHICommandList.Breadcrumbs.CPU.Current = CmdList->CPUBreadcrumbState.Current;
 	}
+	else
+	{
+		GRHICommandList.Breadcrumbs.CPU.Current = CPUFirst;
+	}
+
+	// Rewind CPU state
+	CmdList->CPUBreadcrumbState.Current = CPUFirst;
+	CmdList->PersistentState.LocalBreadcrumb = CPUFirst;
+#endif
 }
 
-void FRHICommandListExecutor::FTranslateState::Translate(FRHICommandListBase* CmdList)
+// ------------------------------------------------------------------------------------------------------------
+// -------------                                  Translate Pipe                                  -------------
+// ------------------------------------------------------------------------------------------------------------
+
+void FRHICommandListExecutor::FTranslateState::Translate_ExecuteCommandList(FRHICommandListBase* CmdList)
 {
+	SCOPED_NAMED_EVENT(RHI_Translate, FColor::Orange);
+
 	if (CmdList->ParallelRenderPassEnd)
 	{
 		// Retrieve the contexts from the previous parent cmdlist.
@@ -1098,76 +1232,80 @@ void FRHICommandListExecutor::FTranslateState::Translate(FRHICommandListBase* Cm
 	delete CmdList;
 }
 
-FGraphEventRef FRHICommandListExecutor::FTranslateState::Finalize()
+void FRHICommandListExecutor::FTranslateState::Translate_CloseChain()
 {
-	FTaskPipe* Pipe = EnqueueTranslateTask({},
-		[this]()
+	SCOPED_NAMED_EVENT(RHI_CloseChain, FColor::Magenta);
+
+	TRHIPipelineArray<IRHIPlatformCommandList*> PlatformCommandLists{ InPlace, nullptr };
+	FDynamicRHI::FRHIFinalizeContextArgs FinalizeArgs;
+	for (auto& State : PipelineStates)
+	{
+		if (State.Context)
 		{
-			SCOPED_NAMED_EVENT(RHI_Finalize, FColor::Magenta);
-		
-			TRHIPipelineArray<IRHIPlatformCommandList*> PlatformCommandLists {InPlace, nullptr};
-			FDynamicRHI::FRHIFinalizeContextArgs FinalizeArgs;
-			for (auto& State : PipelineStates)
+			FinalizeArgs.Contexts.Add(State.Context);
+		}
+	}
+
+	if (GDynamicRHI)
+	{
+		FinalizeArgs.UploadContext = UploadContextState;
+
+		if (bUsingSubCmdLists)
+		{
+			auto& State = PipelineStates[ERHIPipeline::Graphics];
+			State.FinalizedCmdList = GDynamicRHI->RHIFinalizeParallelContext(State.Context);
+		}
+		else
+		{
+			GDynamicRHI->RHICloseTranslateChain(MoveTemp(FinalizeArgs), PlatformCommandLists, bShouldFinalize);
+
+			for (ERHIPipeline Pipeline : MakeFlagsRange(ERHIPipeline::All))
 			{
+				auto& State = PipelineStates[Pipeline];
 				if (State.Context)
 				{
-					FinalizeArgs.Contexts.Add(State.Context);
-				}
-			}
-		
-			if (GDynamicRHI)
-			{
-				FinalizeArgs.UploadContext = UploadContextState;
-				
-				if (bUsingSubCmdLists)
-				{
-					auto& State = PipelineStates[ERHIPipeline::Graphics];
-					State.FinalizedCmdList = GDynamicRHI->RHIFinalizeParallelContext(State.Context);
-				}
-				else
-				{
-					GDynamicRHI->RHICloseTranslateChain(MoveTemp(FinalizeArgs), PlatformCommandLists, bShouldFinalize);
-
-					for (ERHIPipeline Pipeline : MakeFlagsRange(ERHIPipeline::All))
-					{
-						auto& State = PipelineStates[Pipeline];
-						if (State.Context)
-						{
-							State.FinalizedCmdList = PlatformCommandLists[Pipeline];
-						}
-					}
+					State.FinalizedCmdList = PlatformCommandLists[Pipeline];
 				}
 			}
 		}
-	);
-
-	if (Pipe == &TranslatePipe)
-	{
-		return Pipe->Close();
-	}
-	else
-	{
-		check(Pipe == &GRHICommandList.RHIThreadPipe);
-
-		// Don't close the pipe if we got scheduled on the RHI thread pipe, to avoid splitting tasks.
-		// Since the pipes guarantee FIFO order, we don't need an event for the submission to await.
-		return {};
 	}
 }
 
-void FRHICommandListExecutor::FSubmitState::Submit(const FSubmitArgs& Args)
+// ---------------------------------------------------------------------------------------------------------
+// -------------                                  Submit Pipe                                  -------------
+// ---------------------------------------------------------------------------------------------------------
+
+void FRHICommandListExecutor::FSubmitState::Submit_SubmitToGPU(FTranslateState* First, FTaskPipe::ITask& Task)
 {
-	// Coalesce finalized platform command lists into a single array
-	TArray<IRHIPlatformCommandList*> FinalizedCmdLists;
-	for (TUniquePtr<FTranslateState> const& Job : TranslateJobs)
+	SCOPED_NAMED_EVENT(RHI_SubmitToGPU, FColor::Purple);
+
+	// Coalesce finalized platform command lists into a single array,
+	// taking as many from the submission queue as we can.
+	TArray<IRHIPlatformCommandList*> CommandLists;
+
+	FTranslateState* Current = First;
+	while (true)
 	{
+		if (Current->TranslateEvent && !Current->TranslateEvent->IsComplete())
+		{
+			// Translation has not completed on this job yet, so we stop processing the queue at this point.
+			// Extend this task to pick this back up when translation has completed later.
+			Task.ExtendTask({ Current->TranslateEvent },[this, First = Current](FTaskPipe::ITask& Task)
+			{
+				SCOPED_NAMED_EVENT(RHI_ExtendedSubmit, FColor::Orange);
+				Submit_SubmitToGPU(First, Task);
+			});
+			break;
+		}
+
+		// Gather the completed platform command lists for every RHI pipeline.
 		for (ERHIPipeline Pipeline : MakeFlagsRange(ERHIPipeline::All))
 		{
-			auto& TranslateState = Job->PipelineStates[Pipeline];
+			auto& TranslateState = Current->PipelineStates[Pipeline];
 			if (TranslateState.FinalizedCmdList)
 			{
-#if WITH_RHI_BREADCRUMBS
-				TranslateState.FinalizedCmdList->BreadcrumbAllocators = Job->BreadcrumbAllocatorRefs;
+			#if WITH_RHI_BREADCRUMBS
+				TranslateState.FinalizedCmdList->BreadcrumbAllocators = Current->BreadcrumbAllocatorRefs;
 
 				auto& Allocators = TranslateState.FinalizedCmdList->BreadcrumbAllocators;
 
@@ -1175,9 +1313,9 @@ void FRHICommandListExecutor::FSubmitState::Submit(const FSubmitArgs& Args)
 				// Link ranges of breadcrumbs together into depth-first list
 				if (GlobalLast)
 				{
-					FRHIBreadcrumbNode*& Next = GlobalLast->GetNextPtr(Pipeline);
-					check(!Next);
-					Next = TranslateState.Range.First;
+					FRHIBreadcrumbNode*& GlobalNext = GlobalLast->GetNextPtr(Pipeline);
+					check(!GlobalNext);
+					GlobalNext = TranslateState.Range.First;
 				}
 
 				FRHIBreadcrumbRange Range{};
@@ -1194,21 +1332,34 @@ void FRHICommandListExecutor::FSubmitState::Submit(const FSubmitArgs& Args)
 				if (Range.Last) { Allocators.AddUnique(Range.Last->Allocator); }
 
 				TranslateState.FinalizedCmdList->BreadcrumbRange = Range;
-#endif
+			#endif
 
-				FinalizedCmdLists.Add(TranslateState.FinalizedCmdList);
+				CommandLists.Add(TranslateState.FinalizedCmdList);
 			}
 		}
-	}
 
-	if (GDynamicRHI)
-	{
-		GDynamicRHI->RHISubmitCommandLists({ MoveTemp(FinalizedCmdLists) });
+		// Try to close the queue. Use the current job's 'this' pointer as a sentinel
+		// to indicate to the dispatch pipe that the queue was atomically closed.
+		FTranslateState* Expected = nullptr;
+		if (Current->Next.compare_exchange_strong(Expected, Current))
+		{
+			// Queue closed, no more chains to process.
+			break;
+		}
+
+		Current = Expected;
 	}
-	else
+
+	// Submit everything we gathered to the GPU.
+	if (CommandLists.Num())
 	{
-		check(FinalizedCmdLists.IsEmpty());
+		GDynamicRHI->RHISubmitCommandLists({ MoveTemp(CommandLists) });
 	}
+}
+
+void FRHICommandListExecutor::FSubmitState::Submit_Close(const FCloseArgs& Args)
+{
+	SCOPED_NAMED_EVENT(RHI_SubmitClose, FColor::Red);
 
 	{
 		SCOPED_NAMED_EVENT(DeleteRHIResources, FColor::Red);
@@ -1258,6 +1409,10 @@ void FRHICommandListExecutor::FSubmitState::Submit(const FSubmitArgs& Args)
 	delete this;
 }
 
+// -----------------------------------------------------------------------------------------------------------
+// -------------                                  Render Thread                                  -------------
+// -----------------------------------------------------------------------------------------------------------
+
 RHI_API FGraphEventRef FRHICommandListExecutor::Submit(TConstArrayView<FRHICommandListBase*> AdditionalCommandLists, ERHISubmitFlags SubmitFlags)
 {
 	check(IsInRenderingThread());
@@ -1266,7 +1421,7 @@ RHI_API FGraphEventRef FRHICommandListExecutor::Submit(TConstArrayView<FRHIComma
 	if (Bypass())
 	{
 		// Always submit to the GPU in Bypass mode. This allows us to wait for all translate tasks to 
-		// complete before  returning from this function. ensuring commands are always executed in-order.
+		// complete before returning from this function. ensuring commands are always executed in-order.
 		EnumAddFlags(SubmitFlags, ERHISubmitFlags::SubmitToGPU);
 	}
 
@@ -1321,17 +1476,21 @@ RHI_API FGraphEventRef FRHICommandListExecutor::Submit(TConstArrayView<FRHIComma
 	// The dispatch tasks are chained together so that they run "single threaded", in the same order the corresponding RHI command lists were submitted in.
 	// 
 	// RHI contexts may have multiple RHI command lists replayed into them. The translate tasks leave the contexts open. They are eventually finalized
-	// by a task in FTranslateState::Finalize() that runs after the batch of translates have completed for that context. Multiple translates for 
+	// by a task in FTranslateState::Translate_CloseChain() that runs after the batch of translates have completed for that context. Multiple translates for 
 	// different contexts are allowed to run in parallel. Specific RHI command lists require single-threaded execution on the RHI thread.
 	// 
-	// Once all finalized tasks have completed, FSubmitState::Submit() called which hands the platform GPU command lists down to the RHI.
+	// Once all finalized tasks have completed, FSubmitState::Submit_SubmitToGPU() called which hands the platform GPU command lists down to the RHI,
+	// or if submission batching is disabled, this hand-off happens as soon as possible after each chain is closed, in submission order.
 	// 
 	// If threaded rendering is disabled, the tasks are enqueued to the render thread local queue, but the tasks and dependencies are the same.
 	//
 
 	if (!SubmitState)
 	{
-		SubmitState = new FSubmitState();
+		// RHIs that do not support parallel translate must always batch submissions to enforce the ordering of translate -> submit -> translate.
+		bool bBatchSubmissions = (!GRHISupportsParallelRHIExecute) || CVarRHICmdParallelTranslateBatchSubmissions.GetValueOnRenderThread();
+
+		SubmitState = new FSubmitState(bBatchSubmissions);
 		FGraphEventRef NewCompletionEvent = FGraphEvent::CreateGraphEvent();
 		NewCompletionEvent->SetDebugName(TEXT("FRHICommandListExecutor::Submit::CompletionEvent"));
 
@@ -1381,10 +1540,8 @@ RHI_API FGraphEventRef FRHICommandListExecutor::Submit(TConstArrayView<FRHIComma
 	}
 
 	// Dispatch each command list
-	for (int32 Index = 0; Index < AllCmdLists.Num(); ++Index)
+	for (FRHICommandListBase* CmdList : AllCmdLists)
 	{
-		FRHICommandListBase* CmdList = AllCmdLists[Index];
-
 		// Accumulate dispatch ready events into the WaitOutstandingTasks list.
 		// This is used by FRHICommandListImmediate::WaitForTasks() when the render thread 
 		// wants to block until all parallel RHICmdList recording tasks are completed.
@@ -1394,18 +1551,17 @@ RHI_API FGraphEventRef FRHICommandListExecutor::Submit(TConstArrayView<FRHIComma
 		Prereqs.Add(CmdList->DispatchEvent);
 
 		EnqueueDispatchTask(MoveTemp(Prereqs),
-			[State = SubmitState, CmdList]()
+			[State = SubmitState, CmdList](FTaskPipe::ITask&)
 			{
-				SCOPED_NAMED_EVENT(RHI_Dispatch, FColor::Yellow);
-				State->Dispatch(CmdList);
+				State->Dispatch_ProcessCommandList(CmdList);
 			}
 		);
 	}
 
 	if (EnumHasAnyFlags(SubmitFlags, ERHISubmitFlags::SubmitToGPU))
 	{
-	#if STATS
 		TOptional<int64> LocalStatsFrame;
+	#if STATS
 		if (EnumHasAnyFlags(SubmitFlags, ERHISubmitFlags::EndFrame))
 		{
 			LocalStatsFrame = UE::Stats::FStats::StatsFrameRT;
@@ -1414,41 +1570,11 @@ RHI_API FGraphEventRef FRHICommandListExecutor::Submit(TConstArrayView<FRHIComma
 	#endif
 
 		EnqueueDispatchTask({},
-			[
-				  this
-				, State = SubmitState
-			#if STATS
-				, LocalStatsFrame
-			#endif
-			]() mutable
+			[this, State = SubmitState, LocalStatsFrame](FTaskPipe::ITask&)
 			{
-				SCOPED_NAMED_EVENT(RHI_FinalizeAndSubmit, FColor::Cyan);
-
-				// Finalize the last translate job
-				State->FinalizeCurrent();
-
-				FSubmitState::FSubmitArgs Args;
-			#if WITH_RHI_BREADCRUMBS
-				for (ERHIPipeline Pipeline : MakeFlagsRange(ERHIPipeline::All))
-				{
-					Args.GPUBreadcrumbs[Pipeline] = GRHICommandList.Breadcrumbs.GPU[Pipeline].Current;
-				}
-			#endif
-			#if STATS
-				Args.StatsFrame = LocalStatsFrame;
-			#endif
-
-				// Submission thread
-				FTaskPipe* Pipe = EnqueueSubmitTask(MoveTemp(State->TranslateEvents),
-					[State, Args]() mutable
-					{
-						SCOPED_NAMED_EVENT(RHI_SubmitToGPU, FColor::Purple);
-						State->Submit(Args);
-					}
-				);
-				LastSubmit = Pipe->Close();
-			}
-		);
+				SCOPED_NAMED_EVENT(RHI_FinalizeSubmission, FColor::Yellow);
+				State->Dispatch_CloseSubmission(LocalStatsFrame);
+			});
 
 		SubmitState = nullptr;
 	}
@@ -2308,11 +2434,12 @@ void FRHICommandListExecutor::CleanupGraphEvents()
 	WaitOutstandingTasks.Reset();
 
 	LastMutate.SafeRelease();
-	LastSubmit.SafeRelease();
+	LastClose.SafeRelease();
 	CompletionEvent.SafeRelease();
 
 	DispatchPipe.CleanupGraphEvents();
-	RHIThreadPipe.CleanupGraphEvents();
+	SingleTranslatePipe.CleanupGraphEvents();
+	SubmitPipe.CleanupGraphEvents();
 }
 
 void FRHICommandListBase::InvalidBufferFatalError(const FRHIBufferCreateDesc& CreateDesc)
diff --git a/Engine/Source/Runtime/RHI/Private/RHIValidation.cpp b/Engine/Source/Runtime/RHI/Private/RHIValidation.cpp
--- a/Engine/Source/Runtime/RHI/Private/RHIValidation.cpp
+++ b/Engine/Source/Runtime/RHI/Private/RHIValidation.cpp
@@ -773,6 +773,8 @@ IRHIPlatformCommandList* FValidationRHI::RHIFinalizeParallelContext(IRHIComputeC
 
 void FValidationRHI::RHISubmitCommandLists(FRHISubmitCommandListsArgs&& Args)
 {
+	SCOPED_NAMED_EVENT(RHIValidation_SubmitCommandLists, FColor::Purple);
+
 	FDynamicRHI::FRHISubmitCommandListsArgs InnerArgs;
 	InnerArgs.CommandLists.Reserve(Args.CommandLists.Num());
 
diff --git a/Engine/Source/Runtime/RHI/Public/DynamicRHI.h b/Engine/Source/Runtime/RHI/Public/DynamicRHI.h
--- a/Engine/Source/Runtime/RHI/Public/DynamicRHI.h
+++ b/Engine/Source/Runtime/RHI/Public/DynamicRHI.h
@@ -883,7 +883,6 @@ public:
 	struct FRHISubmitCommandListsArgs
 	{
 		TArray<IRHIPlatformCommandList*> CommandLists;
-		IRHIUploadContext* UploadContext;
 	};
 	virtual void RHISubmitCommandLists(FRHISubmitCommandListsArgs&& Args) = 0;
 
diff --git a/Engine/Source/Runtime/RHI/Public/RHICommandList.h b/Engine/Source/Runtime/RHI/Public/RHICommandList.h
--- a/Engine/Source/Runtime/RHI/Public/RHICommandList.h
+++ b/Engine/Source/Runtime/RHI/Public/RHICommandList.h
@@ -5142,9 +5142,17 @@ private:
 		FGraphEventRef LaunchTask(FTask* Task) const;
 		void Execute(FTask* Task, FGraphEventRef const& CurrentEvent) const;
 
-	public:	
+	public:
+		struct ITask;
+		using FFunc = TFunction<void(ITask&)>;
+
+		struct ITask
+		{
+			virtual void ExtendTask(FGraphEventArray&& NewPrereqs, FFunc&& NewLambda) = 0;
+		};
+
 		// Enqueues the given lambda to run on the named thread.
-		void Enqueue(ENamedThreads::Type NamedThread, FGraphEventArray&& Prereqs, TFunction<void()>&& Lambda);
+		void Enqueue(ENamedThreads::Type NamedThread, FGraphEventArray&& Prereqs, FFunc&& Lambda);
 
 		// Returns a graph event that will be signalled once all work submitted prior to calling Close() has completed.
 		FGraphEventRef Close();
@@ -5155,8 +5163,7 @@ private:
 		}
 	};
 
-	FTaskPipe DispatchPipe;
-	FTaskPipe RHIThreadPipe;
+	struct FSubmitState;
 
 	// One per RHI context array, multiple RHICmdLists replayed into it
 	struct FTranslateState
@@ -5179,15 +5186,26 @@ private:
 
 		FTaskPipe TranslatePipe;
 		uint32 NumCommands = 0;
-		bool bParallel = false;
-		bool bUsingSubCmdLists = false;
+
+		FGraphEventRef TranslateEvent;
+		std::atomic<FTranslateState*> Next { nullptr };
+
+		// True when parallel translate is enabled globally
+		const bool bAllowParallelTranslate;
+
+		// True when this particular chain is a parallel translate chain
+		const bool bParallel;
+
+		const bool bUsingSubCmdLists;
+
 		bool bShouldFinalize = true;
 
-		FTaskPipe* GetTranslateTaskPipe(ENamedThreads::Type& NamedThread);
-		FTaskPipe* EnqueueTranslateTask(FGraphEventArray&& Prereqs, TFunction<void()>&& Lambda);
+		FTaskPipe* EnqueueTranslateTask(FGraphEventArray&& Prereqs, FTaskPipe::FFunc&& Lambda);
 
-		void Translate(FRHICommandListBase* CmdList);
-		FGraphEventRef Finalize();
+		void Translate_ExecuteCommandList(FRHICommandListBase* CmdList);
+		void Translate_CloseChain();
+
+		FTranslateState(const FSubmitState& SubmitState, FRHICommandListBase* FirstCmdList);
 	};
 
 	// One per call to RHISubmitCommandLists
@@ -5195,14 +5213,22 @@ private:
 	{
 		FGraphEventRef CompletionEvent;
 
-		TArray<TUniquePtr<FTranslateState>> TranslateJobs;
-		FGraphEventArray TranslateEvents;
-		FTranslateState* CurrentTranslateJob = nullptr;
+		struct
+		{
+			TArray<TUniquePtr<FTranslateState>> All;
+			FGraphEventArray Events;
+
+			FTranslateState* First = nullptr;
+			FTranslateState* Current = nullptr;
+			FTranslateState* Previous = nullptr;
+		} TranslateJobs;
 
 		int32 MaxCommandsPerTranslate = 0;
 		bool bAllowSingleParallelCombine = false;
 		bool bAllowParallelTranslate = true;
 
+		const bool bBatchSubmissions;
+
 	#if WITH_RHI_BREADCRUMBS
 		bool bEmitBreadcrumbs = false;
 	#endif
@@ -5210,34 +5236,44 @@ private:
 		ERHISubmitFlags SubmitFlags = ERHISubmitFlags::None;
 		TArray<FRHIResource*> ResourcesToDelete {};
 		bool bIncludeExtendedLifetimeResources = false;
-
-		bool ShouldSplitTranslateJob(FRHICommandListBase* CmdList);
-		void ConditionalSplitTranslateJob(FRHICommandListBase* CmdList);
 		
 		FGraphEventRef BeginGraphEvent;
 		FGraphEventArray ChildGraphEvents;
 
-		void Dispatch(FRHICommandListBase* CmdList);
-
-		struct FSubmitArgs
+		struct FCloseArgs
 		{
-		#if WITH_RHI_BREADCRUMBS
+#if WITH_RHI_BREADCRUMBS
 			TRHIPipelineArray<FRHIBreadcrumbNode*> GPUBreadcrumbs;
-		#endif
-		#if STATS
+#endif
+#if STATS
 			TOptional<int64> StatsFrame;
-		#endif
+#endif
 		};
-		void Submit(const FSubmitArgs& Args);
-		FGraphEventRef FinalizeCurrent();
+
+		void Dispatch_ProcessCommandList(FRHICommandListBase* CmdList);
+		void Dispatch_PrepareBreadcrumbs(FRHICommandListBase* CmdList);
+		FGraphEventRef Dispatch_CloseTranslateChain(bool bLastChainInSubmission);
+		void Dispatch_CloseSubmission(TOptional<int64> LocalStatsFrame);
+
+		void Submit_SubmitToGPU(FTranslateState* First, FTaskPipe::ITask& Task);
+		void Submit_Close(const FCloseArgs& Args);
+
+		FSubmitState(bool bBatchSubmissions)
+			: bBatchSubmissions(bBatchSubmissions)
+		{}
+
 	} *SubmitState = nullptr;
 
+	FTaskPipe DispatchPipe;
+	FTaskPipe SingleTranslatePipe;
+	FTaskPipe SubmitPipe;
+
 	FGraphEventRef LastMutate;
-	FGraphEventRef LastSubmit;
+	FGraphEventRef LastClose;
 	FGraphEventRef CompletionEvent;
 
-	FTaskPipe* EnqueueDispatchTask(FGraphEventArray&& Prereqs, TFunction<void()>&& Lambda);
-	FTaskPipe* EnqueueSubmitTask  (FGraphEventArray&& Prereqs, TFunction<void()>&& Lambda);
+	void EnqueueDispatchTask(FGraphEventArray&& Prereqs, FTaskPipe::FFunc&& Lambda);
+	void EnqueueSubmitTask  (FGraphEventArray&& Prereqs, FTaskPipe::FFunc&& Lambda);
 	FGraphEventArray NextDispatchTaskPrerequisites;
 
 #if WITH_RHI_BREADCRUMBS
