RHI: Fix AsyncCompute contexts when parallel render passes are used with RHIValidation.

#rnx

#Synced-CL 52063422

diff --git a/Engine/Source/Runtime/RHI/Private/RHIValidation.cpp b/Engine/Source/Runtime/RHI/Private/RHIValidation.cpp
--- a/Engine/Source/Runtime/RHI/Private/RHIValidation.cpp
+++ b/Engine/Source/Runtime/RHI/Private/RHIValidation.cpp
@@ -728,21 +728,19 @@ IRHIComputeContext* FValidationRHI::RHIGetParallelCommandContext(FRHIParallelRen
 void FValidationRHI::RHICloseTranslateChain(FRHIFinalizeContextArgs&& Args, TRHIPipelineArray<IRHIPlatformCommandList*>& Output, bool bShouldFinalize)
 {
 	// If we aren't finalizing the context we need to finalize the tracking
-	if(!bShouldFinalize)
+	if (!bShouldFinalize)
 	{
-		for(IRHIComputeContext* Context : Args.Contexts)
+		for (IRHIComputeContext* Context : Args.Contexts)
 		{
-			if(Context)
+			if (Context)
 			{
-				check(Context->GetPipeline() == ERHIPipeline::Graphics);
-				
 				IRHIComputeContext& InnerContext = Context->GetLowestLevelContext();
 				
 				FValidationCommandList* OuterCommandList = new FValidationCommandList();
 				OuterCommandList->CompletedOpList = InnerContext.Tracker->Finalize();
-				OuterCommandList->Pipeline = ERHIPipeline::Graphics;
+				OuterCommandList->Pipeline = Context->GetPipeline();
 				
-				Output[ERHIPipeline::Graphics] = OuterCommandList;
+				Output[Context->GetPipeline()] = OuterCommandList;
 			}
 		}
 		
