RHIDescriptorRange
- Adding per-resource type storage for validaton loading and D3D12 residency caching.
- Adding descriptor range per-resource validation.

#Synced-CL 52065068

diff --git a/Engine/Source/Runtime/D3D12RHI/Private/D3D12DescriptorRange.cpp b/Engine/Source/Runtime/D3D12RHI/Private/D3D12DescriptorRange.cpp
--- a/Engine/Source/Runtime/D3D12RHI/Private/D3D12DescriptorRange.cpp
+++ b/Engine/Source/Runtime/D3D12RHI/Private/D3D12DescriptorRange.cpp
@@ -64,7 +64,7 @@ FRHIDescriptorRangeShaderData FD3D12DescriptorRange::GetShaderData() const
 	return ShaderData;
 }
 
-void FD3D12DescriptorRange::InitializeDescriptors(TConstArrayView<FRHIDescriptorResource> InResources)
+void FD3D12DescriptorRange::InitializeResources(TConstArrayView<FRHIDescriptorResource> InResources)
 {
 	const ERHIDescriptorTypeMask DescriptorTypeMask = GetDesc().TypeMask;
 
@@ -74,23 +74,21 @@ void FD3D12DescriptorRange::InitializeDescriptors(TConstArrayView<FRHIDescriptor
 	FD3D12BindlessDescriptorManager& Manager = Device->GetBindlessDescriptorManager();
 
 	const int32 NumResources = FMath::Min(InResources.Num(), (int32)GetDesc().Count);
-	ResourceTypes.SetNum(NumResources);
 
 	for (int32 Index = 0; Index < NumResources; Index++)
 	{
 		const FRHIDescriptorResource& Resource = InResources[Index];
 
-		const FD3D12View* View = GetViewFromDescriptorResource(Resource, GpuIndex);
-
-		checkf(EnumHasAllFlags(DescriptorTypeMask, View->GetBindlessHandle().GetTypeMask()),
-			TEXT("View descriptor type %s is not compatible with descriptor range type mask %s."),
-			GetDescriptorTypeString(View->GetBindlessHandle().GetType()),
-			*GetDescriptorMaskString(DescriptorTypeMask));
-
-		const FRHIDescriptorHandle UpdateHandle(View->GetBindlessHandle().GetType(), Allocation.StartIndex + Index);
-		Manager.InitializeDescriptor(UpdateHandle, View);
+		if (const FD3D12View* View = GetViewFromDescriptorResource(Resource, GpuIndex))
+		{
+			checkf(EnumHasAllFlags(DescriptorTypeMask, View->GetBindlessHandle().GetTypeMask()),
+				TEXT("View descriptor type %s is not compatible with descriptor range type mask %s."),
+				GetDescriptorTypeString(View->GetBindlessHandle().GetType()),
+				*GetDescriptorMaskString(DescriptorTypeMask));
 
-		ResourceTypes[Index] = Resource.Type;
+			const FRHIDescriptorHandle UpdateHandle(View->GetBindlessHandle().GetType(), Allocation.StartIndex + Index);
+			Manager.InitializeDescriptor(UpdateHandle, View);
+		}
 	}
 
 	// Rebuild cached views so residency is correct at next bind.
@@ -114,17 +112,16 @@ void FD3D12DescriptorRange::UpdateResources(const FRHIContextArray& Contexts, ui
 			const FRHIDescriptorResource& Resource = NewResources[Index];
 			UpdateResource(FirstResourceIndex + Index, Resource);
 
-			const FD3D12View* View = GetViewFromDescriptorResource(Resource, GpuIndex);
-
-			checkf(EnumHasAllFlags(DescriptorTypeMask, View->GetBindlessHandle().GetTypeMask()),
-				TEXT("View descriptor type %s is not compatible with descriptor range type mask %s."),
-				GetDescriptorTypeString(View->GetBindlessHandle().GetType()),
-				*GetDescriptorMaskString(DescriptorTypeMask));
-
-			const FRHIDescriptorHandle UpdateHandle(View->GetBindlessHandle().GetType(), Allocation.StartIndex + InternalIndex);
-			Manager.UpdateDescriptor(Contexts, UpdateHandle, View);
+			if (const FD3D12View* View = GetViewFromDescriptorResource(Resource, GpuIndex))
+			{
+				checkf(EnumHasAllFlags(DescriptorTypeMask, View->GetBindlessHandle().GetTypeMask()),
+					TEXT("View descriptor type %s is not compatible with descriptor range type mask %s."),
+					GetDescriptorTypeString(View->GetBindlessHandle().GetType()),
+					*GetDescriptorMaskString(DescriptorTypeMask));
 
-			ResourceTypes[InternalIndex] = Resource.Type;
+				const FRHIDescriptorHandle UpdateHandle(View->GetBindlessHandle().GetType(), Allocation.StartIndex + InternalIndex);
+				Manager.UpdateDescriptor(Contexts, UpdateHandle, View);
+			}
 		}
 	}
 
@@ -147,18 +144,17 @@ void FD3D12DescriptorRange::PopulateResidencyViews()
 			continue;
 		}
 
-		// Reconstruct the FRHIDescriptorResource from the stored type so that we can get the correct view from the resource.
 		const FRHIDescriptorResource DescriptorResource(ResourceTypes[Index], Resource);
 
 		if (FD3D12View* View = GetViewFromDescriptorResource(DescriptorResource, GpuIndex))
 		{
 			if (ResourceTypes[Index] == FRHIDescriptorResource::EType::UnorderedAccessView)
 			{
-				AllUavs.Emplace(static_cast<FD3D12UnorderedAccessView*>(View));
+				AllUavs.AddUnique(static_cast<FD3D12UnorderedAccessView*>(View));
 			}
 			else
 			{
-				AllSrvs.Emplace(static_cast<FD3D12ShaderResourceView*>(View));
+				AllSrvs.AddUnique(static_cast<FD3D12ShaderResourceView*>(View));
 			}
 		}
 	}
@@ -185,7 +181,7 @@ FDescriptorRangeRHIRef FD3D12DynamicRHI::RHICreateDescriptorRange(FRHICommandLis
 		{
 			for (FD3D12DescriptorRange::FLinkedObjectIterator RangeIt(DescriptorRange); RangeIt; ++RangeIt)
 			{
-				RangeIt->InitializeDescriptors(InitialResources);
+				RangeIt->InitializeResources(InitialResources);
 			}
 		});
 
diff --git a/Engine/Source/Runtime/D3D12RHI/Private/D3D12DescriptorRange.h b/Engine/Source/Runtime/D3D12RHI/Private/D3D12DescriptorRange.h
--- a/Engine/Source/Runtime/D3D12RHI/Private/D3D12DescriptorRange.h
+++ b/Engine/Source/Runtime/D3D12RHI/Private/D3D12DescriptorRange.h
@@ -19,7 +19,7 @@ public:
 
 	FRHIDescriptorRangeShaderData GetShaderData() const override;
 
-	void InitializeDescriptors(TConstArrayView<FRHIDescriptorResource> InResources);
+	void InitializeResources(TConstArrayView<FRHIDescriptorResource> InResources);
 	void UpdateResources(const FRHIContextArray& Contexts, uint32 FirstResourceIndex, TConstArrayView<FRHIDescriptorResource> NewResources);
 
 	// Rebuilds the cached AllSrvs/AllUavs arrays from the current resource list.
@@ -27,10 +27,6 @@ public:
 
 	FRHIDescriptorAllocation Allocation;
 
-	// Stores the descriptor type (Texture, SRV, UAV) for each slot. Used by PopulateResidencyViews
-	// to extract the correct D3D12 view type from the base class Resources array.
-	TArray<FRHIDescriptorResource::EType> ResourceTypes;
-
 	// Cached D3D12 views for all resources in the range. Rebuilt on create/update via PopulateResidencyViews().
 	// Queued for residency (via QueueBindlessSRVs/UAVs) when the descriptor range is bound to a shader.
 	TArray<FD3D12ShaderResourceView*> AllSrvs;
diff --git a/Engine/Source/Runtime/RHI/Private/RHIDescriptorRange.cpp b/Engine/Source/Runtime/RHI/Private/RHIDescriptorRange.cpp
--- a/Engine/Source/Runtime/RHI/Private/RHIDescriptorRange.cpp
+++ b/Engine/Source/Runtime/RHI/Private/RHIDescriptorRange.cpp
@@ -7,17 +7,23 @@ FRHIDescriptorRange::FRHIDescriptorRange(const FRHIDescriptorRangeDesc& InDesc,
 	, Desc(InDesc)
 {
 	Resources.SetNum(InDesc.Count);
+	ResourceTypes.SetNum(InDesc.Count);
 	for (int32 Index = 0; Index < InResources.Num(); Index++)
 	{
 		UpdateResource(Index, InResources[Index]);
 	}
 }
 
+FRHIDescriptorRange::~FRHIDescriptorRange()
+{
+}
+
 void FRHIDescriptorRange::UpdateResource(int32 Index, FRHIDescriptorResource InResource)
 {
 	if (Index >= 0 && Index < Resources.Num())
 	{
 		Resources[Index] = InResource.Resource;
+		ResourceTypes[Index] = InResource.Type;
 	}
 }
 
diff --git a/Engine/Source/Runtime/RHI/Private/RHIValidation.cpp b/Engine/Source/Runtime/RHI/Private/RHIValidation.cpp
--- a/Engine/Source/Runtime/RHI/Private/RHIValidation.cpp
+++ b/Engine/Source/Runtime/RHI/Private/RHIValidation.cpp
@@ -1156,14 +1156,19 @@ void FValidationRHI::RHIUnlockTexture(FRHICommandListImmediate& RHICmdList, cons
 #if PLATFORM_SUPPORTS_BINDLESS_RENDERING
 inline ERHIDescriptorType GetDescriptorTypeFromDescriptorResource(const FRHIDescriptorResource& InResource)
 {
-	switch (InResource.Type)
+	if (InResource.Resource)
 	{
-	case FRHIDescriptorResource::EType::Texture:
-		return static_cast<const FRHITexture*>(InResource.Resource)->GetDefaultBindlessHandle().GetType();
-	case FRHIDescriptorResource::EType::ResourceView:
-		return static_cast<const FRHIShaderResourceView*>(InResource.Resource)->GetBindlessHandle().GetType();
-	case FRHIDescriptorResource::EType::UnorderedAccessView:
-		return static_cast<const FRHIUnorderedAccessView*>(InResource.Resource)->GetBindlessHandle().GetType();
+		switch (InResource.Type)
+		{
+		case FRHIDescriptorResource::EType::None:
+			return ERHIDescriptorType::Invalid;
+		case FRHIDescriptorResource::EType::Texture:
+			return static_cast<const FRHITexture*>(InResource.Resource)->GetDefaultBindlessHandle().GetType();
+		case FRHIDescriptorResource::EType::ResourceView:
+			return static_cast<const FRHIShaderResourceView*>(InResource.Resource)->GetBindlessHandle().GetType();
+		case FRHIDescriptorResource::EType::UnorderedAccessView:
+			return static_cast<const FRHIUnorderedAccessView*>(InResource.Resource)->GetBindlessHandle().GetType();
+		}
 	}
 
 	return ERHIDescriptorType::Invalid;
@@ -1197,18 +1202,27 @@ inline FString GetAllDescriptorTypeGroupsAsString()
 	return FString(Result);
 }
 
-void FValidationRHI::ValidateDescriptorTypes(const FRHIDescriptorRangeDesc& RangeDesc, TConstArrayView<FRHIDescriptorResource> Resources)
+void FValidationRHI::ValidateDescriptorRangeResources(const FRHIDescriptorRangeDesc& RangeDesc, TConstArrayView<FRHIDescriptorResource> Resources)
 {
 	RHI_VALIDATION_CHECK(ValidateDescriptorGroups(RangeDesc.TypeMask),
 		*FString::Printf(TEXT("Invalid descriptor type mask \"%s\" overlaps type groups: %s"), *GetDescriptorMaskString(RangeDesc.TypeMask), *GetAllDescriptorTypeGroupsAsString()));
 
+	RHI_VALIDATION_CHECK(Resources.Num() <= (int32)RangeDesc.Count,
+		*FString::Printf(TEXT("Resources count (%d) exceeds descriptor range capacity (%d)."), Resources.Num(), (int32)RangeDesc.Count));
+
 	for (const FRHIDescriptorResource& Resource : Resources)
 	{
-		const ERHIDescriptorType DescriptorType = GetDescriptorTypeFromDescriptorResource(Resource);
+		RHI_VALIDATION_CHECK(Resource.Type == FRHIDescriptorResource::EType::None || Resource.Resource != nullptr,
+			TEXT("Descriptor resource has a type set but a null resource pointer."));
+
+		if (Resource.Resource && Resource.Type != FRHIDescriptorResource::EType::None)
+		{
+			const ERHIDescriptorType DescriptorType = GetDescriptorTypeFromDescriptorResource(Resource);
 
-		RHI_VALIDATION_CHECK(DescriptorType != ERHIDescriptorType::Invalid, TEXT("Invalid descriptor type found."));
-		RHI_VALIDATION_CHECK(EnumHasAllFlags(RangeDesc.TypeMask, RHIDescriptorTypeMaskFromType(DescriptorType)),
-			*FString::Printf(TEXT("Unexpected descriptor type found: %s With expected type mask: %s"), GetDescriptorTypeString(DescriptorType), *GetDescriptorMaskString(RangeDesc.TypeMask)));
+			RHI_VALIDATION_CHECK(DescriptorType != ERHIDescriptorType::Invalid, TEXT("Invalid descriptor type found."));
+			RHI_VALIDATION_CHECK(EnumHasAllFlags(RangeDesc.TypeMask, RHIDescriptorTypeMaskFromType(DescriptorType)),
+				*FString::Printf(TEXT("Unexpected descriptor type found: %s With expected type mask: %s"), GetDescriptorTypeString(DescriptorType), *GetDescriptorMaskString(RangeDesc.TypeMask)));
+		}
 	}
 }
 #endif // PLATFORM_SUPPORTS_BINDLESS_RENDERING
@@ -3778,6 +3792,47 @@ void ValidateShaderParameters(FRHIShader* RHIShader, RHIValidation::FTracker* Tr
 				}
 			}
 			break;
+		case FRHIShaderParameterResource::EType::DescriptorRange:
+			if (const FRHIDescriptorRange* DescriptorRange = static_cast<const FRHIDescriptorRange*>(Parameter.Resource))
+			{
+				if (GRHIValidationEnabled)
+				{
+					RHI_VALIDATION_CHECK(DescriptorRange->IsValid(),
+						TEXT("Descriptor range bound to shader is not valid. Allocation may have failed."));
+				}
+
+				for (int32 ResourceIndex = 0; ResourceIndex < DescriptorRange->GetNumResources(); ResourceIndex++)
+				{
+					const FRHIDescriptorResource Resource = DescriptorRange->GetResource(ResourceIndex);
+					switch (Resource.Type)
+					{
+					case FRHIDescriptorResource::EType::None:
+						break;
+					case FRHIDescriptorResource::EType::Texture:
+						if (FRHITexture* Texture = static_cast<FRHITexture*>(Resource.Resource))
+						{
+							Tracker->Assert(Texture->GetWholeResourceIdentitySRV(), InRequiredAccess);
+						}
+						break;
+					case FRHIDescriptorResource::EType::ResourceView:
+						if (FRHIShaderResourceView* SRV = static_cast<FRHIShaderResourceView*>(Resource.Resource))
+						{
+							Tracker->Assert(SRV->GetViewIdentity(), InRequiredAccess);
+						}
+						break;
+					case FRHIDescriptorResource::EType::UnorderedAccessView:
+						if (FRHIUnorderedAccessView* UAV = static_cast<FRHIUnorderedAccessView*>(Resource.Resource))
+						{
+							Tracker->AssertUAV(UAV, InRequiredUAVMode, Parameter.Index);
+						}
+						break;
+					default:
+						checkf(false, TEXT("Unhandled DescriptorRange resource type: %d"), static_cast<int32>(Resource.Type));
+						break;
+					}
+				}
+			}
+			break;
 		default:
 			checkf(false, TEXT("Unhandled resource type?"));
 			break;
diff --git a/Engine/Source/Runtime/RHI/Public/RHIDescriptorRange.h b/Engine/Source/Runtime/RHI/Public/RHIDescriptorRange.h
--- a/Engine/Source/Runtime/RHI/Public/RHIDescriptorRange.h
+++ b/Engine/Source/Runtime/RHI/Public/RHIDescriptorRange.h
@@ -115,13 +115,40 @@ public:
 
 	virtual FRHIDescriptorRangeShaderData GetShaderData() const = 0;
 
+	// Returns the resource and type at the given index as an FRHIDescriptorResource.
+	FRHIDescriptorResource GetResource(int32 Index) const
+	{
+		checkSlow(Index >= 0 && Index < Resources.Num());
+		return FRHIDescriptorResource(ResourceTypes[Index], Resources[Index].GetReference());
+	}
+
+	int32 GetNumResources() const
+	{
+		return Resources.Num();
+	}
+
+	TConstArrayView<FRHIDescriptorResource::EType> GetResourceTypes() const
+	{
+		return MakeConstArrayView(ResourceTypes);
+	}
+
 protected:
 	RHI_API FRHIDescriptorRange(const FRHIDescriptorRangeDesc& InDesc, TConstArrayView<FRHIDescriptorResource> InResources);
+	RHI_API ~FRHIDescriptorRange();
 
+	// Non-copyable due to manual AddRef/Release on resources.
+	FRHIDescriptorRange(const FRHIDescriptorRange&) = delete;
+	FRHIDescriptorRange& operator=(const FRHIDescriptorRange&) = delete;
+
+	// Updates a single resource slot. Handles AddRef/Release for the old and new resource.
 	RHI_API void UpdateResource(int32 Index, FRHIDescriptorResource InResource);
 	RHI_API void UpdateResources(int32 StartIndex, TConstArrayView<FRHIDescriptorResource> InResources);
 
 	const FRHIDescriptorRangeDesc Desc;
 
+	// Resources assigned to this range. Sized to Desc.Count.
 	TArray<TRefCountPtr<FRHIResource>> Resources;
+
+	// Parallel array storing the descriptor type for each slot. Sized to Desc.Count.
+	TArray<FRHIDescriptorResource::EType> ResourceTypes;
 };
diff --git a/Engine/Source/Runtime/RHI/Public/RHIValidation.h b/Engine/Source/Runtime/RHI/Public/RHIValidation.h
--- a/Engine/Source/Runtime/RHI/Public/RHIValidation.h
+++ b/Engine/Source/Runtime/RHI/Public/RHIValidation.h
@@ -1029,16 +1029,19 @@ public:
 	}
 
 #if PLATFORM_SUPPORTS_BINDLESS_RENDERING
-	RHI_API static void ValidateDescriptorTypes(const FRHIDescriptorRangeDesc& RangeDesc, TConstArrayView<FRHIDescriptorResource> Resources);
+	RHI_API static void ValidateDescriptorRangeResources(const FRHIDescriptorRangeDesc& RangeDesc, TConstArrayView<FRHIDescriptorResource> Resources);
 
 	virtual FDescriptorRangeRHIRef RHICreateDescriptorRange(FRHICommandListBase& RHICmdList, const FRHIDescriptorRangeDesc& Desc, TConstArrayView<FRHIDescriptorResource> InitialResources) override final
 	{
-		ValidateDescriptorTypes(Desc, InitialResources);
+		ValidateDescriptorRangeResources(Desc, InitialResources);
 		return RHI->RHICreateDescriptorRange(RHICmdList, Desc, InitialResources);
 	}
 	virtual void RHIUpdateDescriptorRange(FRHICommandListBase& RHICmdList, FRHIDescriptorRange* Range, uint32 StartIndex, TConstArrayView<FRHIDescriptorResource> Resources) override final
 	{
-		ValidateDescriptorTypes(Range->GetDesc(), Resources);
+		RHI_VALIDATION_CHECK((uint64)StartIndex + Resources.Num() <= Range->GetDesc().Count,
+			*FString::Printf(TEXT("UpdateDescriptorRange would write past end of range. StartIndex=%u, ResourceCount=%d, RangeCapacity=%u."),
+				StartIndex, Resources.Num(), Range->GetDesc().Count));
+		ValidateDescriptorRangeResources(Range->GetDesc(), Resources);
 		RHI->RHIUpdateDescriptorRange(RHICmdList, Range, StartIndex, Resources);
 	}
 #endif
