Configurable render tracing channels.

#rnx

#Synced-CL 52126719

diff --git a/Engine/Source/Runtime/RHI/Private/DynamicRHI.cpp b/Engine/Source/Runtime/RHI/Private/DynamicRHI.cpp
--- a/Engine/Source/Runtime/RHI/Private/DynamicRHI.cpp
+++ b/Engine/Source/Runtime/RHI/Private/DynamicRHI.cpp
@@ -22,6 +22,7 @@
 #include "RHIStrings.h"
 #include "RHITextureReference.h"
 #include "RHIAllocators.h"
+#include "RenderTracing.h"
 
 #ifndef PLATFORM_ALLOW_NULL_RHI
 	#define PLATFORM_ALLOW_NULL_RHI		0
@@ -368,6 +369,7 @@ void RHIInit(bool bHasEditorToken)
 		}
 
 		check(GDynamicRHI);
+		IF_RENDER_TRACING_ENABLED(RenderTracing::Initialize());
 	}
 
 #if WITH_EDITOR
diff --git a/Engine/Source/Runtime/RHI/Private/RenderTracing.cpp b/Engine/Source/Runtime/RHI/Private/RenderTracing.cpp
--- a/Engine/Source/Runtime/RHI/Private/RenderTracing.cpp
+++ b/Engine/Source/Runtime/RHI/Private/RenderTracing.cpp
@@ -3,6 +3,11 @@
 #include "RenderTracing.h"
 #include "Trace/Trace.inl"
 #include "HAL/PlatformTime.h"
+#include "HAL/IConsoleManager.h"
+#include "Misc/CommandLine.h"
+#include "Misc/Parse.h"
+#include "String/ParseTokens.h"
+#include "RHIDefinitions.h"
 
 UE_TRACE_CHANNEL_DEFINE(RenderTraceChannel);
 
@@ -11,8 +16,60 @@ UE_TRACE_CHANNEL_DEFINE(RenderTraceChannel);
 namespace RenderTracing
 {
 
-// TODO: make this configurable.
-static ERenderTracingChannels GRenderTracingChannels = ERenderTracingChannels::Submission;
+// Start with everything enabled to make sure we don't lose any events before initialization runs.
+static ERenderTracingChannels GRenderTracingChannels = ERenderTracingChannels::All;
+
+static void SetChannels(IConsoleVariable* CVar)
+{
+	FString ChannelString;
+	if (!FParse::Value(FCommandLine::Get(), TEXT("-rendertracechannels="), ChannelString, false))
+	{
+		ChannelString = CVar->GetString();
+	}
+
+	GRenderTracingChannels = ERenderTracingChannels::None;
+
+	FStringView Trimmed = FStringView(ChannelString).TrimStartAndEnd();
+	if (Trimmed.IsEmpty() || Trimmed.Equals(TEXT("none"), ESearchCase::IgnoreCase))
+	{
+		UE_LOG(LogRHI, Log, TEXT("Render tracing: all channels disabled."));
+		return;
+	}
+
+	UE::String::ParseTokens(Trimmed, TEXT(','), [](FStringView Token)
+	{
+		FStringView Name = Token.TrimStartAndEnd();
+		if      (Name.Equals(TEXT("RDGPasses"),			ESearchCase::IgnoreCase)) { GRenderTracingChannels |= ERenderTracingChannels::RDGPasses; }
+		else if (Name.Equals(TEXT("RHICmdLists"),		ESearchCase::IgnoreCase)) { GRenderTracingChannels |= ERenderTracingChannels::RHICmdLists; }
+		else if (Name.Equals(TEXT("RHITranslation"),	ESearchCase::IgnoreCase)) { GRenderTracingChannels |= ERenderTracingChannels::RHITranslation; }
+		else if (Name.Equals(TEXT("RHISubmission"),		ESearchCase::IgnoreCase)) { GRenderTracingChannels |= ERenderTracingChannels::RHISubmission; }
+		else if (Name.Equals(TEXT("Submission"),		ESearchCase::IgnoreCase)) { GRenderTracingChannels |= ERenderTracingChannels::Submission; }
+		else if (Name.Equals(TEXT("All"),				ESearchCase::IgnoreCase)) { GRenderTracingChannels |= ERenderTracingChannels::All; }
+		else
+		{
+			UE_LOG(LogRHI, Warning, TEXT("Unknown render tracing channel '%.*s', ignoring."), Name.Len(), Name.GetData());
+		}
+	});
+
+	UE_LOG(LogRHI, Log, TEXT("Render tracing channels set: '%s' -> 0x%llx"), *ChannelString, static_cast<uint64>(GRenderTracingChannels));
+}
+
+// Read-only cvar which controls which channels are enabled. We add a change callback because we want to allow the value to be overridden by
+// the hotfix system, which runs after the Initialize() call below. The command line flag (-rendertracechannels) always takes precedence,
+// regardless of having a hotfix value or not.
+static TAutoConsoleVariable<FString> CVarRenderTraceChannels(
+	TEXT("r.Trace.Channels"),
+	TEXT("Submission"),
+	TEXT("Comma-separated list of render tracing channels to enable. Valid values are the members of the ERenderTracingChannels enum.\n")
+	TEXT("Use 'none' or an empty string to disable all channels. Override via -rendertracechannels=<channels> on the command line."),
+	FConsoleVariableDelegate::CreateStatic(&SetChannels),
+	ECVF_ReadOnly
+);
+
+void Initialize()
+{
+	SetChannels(CVarRenderTraceChannels.AsVariable());
+}
 
 bool IsEnabled(ERenderTracingChannels ChannelMask)
 {
diff --git a/Engine/Source/Runtime/RHI/Public/RenderTracing.h b/Engine/Source/Runtime/RHI/Public/RenderTracing.h
--- a/Engine/Source/Runtime/RHI/Public/RenderTracing.h
+++ b/Engine/Source/Runtime/RHI/Public/RenderTracing.h
@@ -35,7 +35,9 @@ enum class ERenderTracingChannels : uint64
 	RHITranslation = 1 << 2,
 	RHISubmission = 1 << 3,
 
-	Submission = RDGPasses | RHICmdLists | RHITranslation | RHISubmission
+	Submission = RDGPasses | RHICmdLists | RHITranslation | RHISubmission,
+
+	All = 0xffffffffffffffffull
 };
 ENUM_CLASS_FLAGS(ERenderTracingChannels);
 
@@ -43,6 +45,8 @@ ENUM_CLASS_FLAGS(ERenderTracingChannels);
 namespace RenderTracing
 {
 
+void Initialize();
+
 RHI_API bool IsEnabled(ERenderTracingChannels ChannelMask);
 
 }
diff --git a/Engine/Source/Runtime/RenderCore/Private/RenderGraphTrace.cpp b/Engine/Source/Runtime/RenderCore/Private/RenderGraphTrace.cpp
--- a/Engine/Source/Runtime/RenderCore/Private/RenderGraphTrace.cpp
+++ b/Engine/Source/Runtime/RenderCore/Private/RenderGraphTrace.cpp
@@ -542,7 +542,7 @@ UE_TRACE_EVENT_END()
 
 static int32 GRDGTraceUseFullPassNames = 0;
 static FAutoConsoleVariableRef CVarRDGTraceUseFullPassName(
-	TEXT("r.RDG.Trace.UseFullPassNames"),
+	TEXT("r.Trace.RDG.UseFullPassNames"),
 	GRDGTraceUseFullPassNames,
 	TEXT("Controls the pass name used in RDG trace events.\n")
 	TEXT(" 0: just the event format string (default)\n")
