Add option to track all precached PSOs in the bundled file cache as well

#Synced-CL 52126632

diff --git a/Engine/Source/Runtime/RHI/Private/PipelineFileCache.cpp b/Engine/Source/Runtime/RHI/Private/PipelineFileCache.cpp
--- a/Engine/Source/Runtime/RHI/Private/PipelineFileCache.cpp
+++ b/Engine/Source/Runtime/RHI/Private/PipelineFileCache.cpp
@@ -138,6 +138,16 @@ static FAutoConsoleVariableRef CVarPSOFileCacheExcludePrecachePSO(
 														   ECVF_ReadOnly
 														   );
 
+static int32 GPSOIncludeAllPrecachedPSOsInFileCache = 0;
+static FAutoConsoleVariableRef CVarPSOFileCacheIncludeAllPrecachedPSOs(
+														   TEXT("r.ShaderPipelineCache.IncludeAllPrecachedPSOs"),
+														   GPSOIncludeAllPrecachedPSOsInFileCache,
+														   TEXT("1 includes all successfully precached PSOs in the pipeline file cache even if they were never needed for rendering.\n")
+														   TEXT("Logged entries carry EPSOPrecacheResult::Active; the result is upgraded to the actual runtime outcome if the PSO is later used at draw time.\n")
+														   TEXT("Requires PSO precaching to be enabled. Default: 0 (disabled)."),
+														   ECVF_ReadOnly
+														   );
+
 static int32 GPSOFileCachePrintNewPSODescriptors = 0;
 static FAutoConsoleVariableRef CVarPSOFileCachePrintNewPSODescriptors(
 														   TEXT("r.ShaderPipelineCache.PrintNewPSODescriptors"),
@@ -194,9 +204,8 @@ TMap<FGuid, FString> FPipelineFileCacheManager::GameGuidToCacheKey;
 TMap<uint32, FPSOUsageData> FPipelineFileCacheManager::RunTimeToPSOUsage;
 TMap<uint32, FPSOUsageData> FPipelineFileCacheManager::NewPSOUsage;
 TMap<uint32, FPipelineStateStats*> FPipelineFileCacheManager::Stats;
-TSet<FPipelineCacheFileFormatPSO> FPipelineFileCacheManager::NewPSOs;
+TMap<uint32, FPipelineCacheFileFormatPSO> FPipelineFileCacheManager::NewPSOs;
 TMpscQueue<FPipelineCacheFileFormatPSO> FPipelineFileCacheManager::NewPSOsToReport;
-TSet<uint32> FPipelineFileCacheManager::NewPSOHashes;
 uint32 FPipelineFileCacheManager::NumNewPSOs;
 FString FPipelineFileCacheManager::UserCacheKey;
 FPipelineFileCacheManager::PSOOrder FPipelineFileCacheManager::RequestedOrder = FPipelineFileCacheManager::PSOOrder::MostToLeastUsed;
@@ -2380,7 +2389,7 @@ public:
 		}
 	}
 	
-	bool SavePipelineFileCache(FPipelineFileCacheManager::SaveMode Mode, TMap<uint32, FPipelineStateStats*> const& Stats, TSet<FPipelineCacheFileFormatPSO>& NewEntries, FPipelineFileCacheManager::PSOOrder Order, TMap<uint32, FPSOUsageData>& NewPSOUsage)
+	bool SavePipelineFileCache(FPipelineFileCacheManager::SaveMode Mode, TMap<uint32, FPipelineStateStats*> const& Stats, TMap<uint32, FPipelineCacheFileFormatPSO>& NewEntries, FPipelineFileCacheManager::PSOOrder Order, TMap<uint32, FPSOUsageData>& NewPSOUsage)
 	{
 		check(CacheStatus == EStatus::NewUserCache || CacheStatus == EStatus::UserCacheOpened);
 		// remove from stats because this operation will modify the content and re-set the stat at the end.
@@ -2556,9 +2565,9 @@ public:
                             // Add new entries
 							for(auto It = NewEntries.CreateIterator(); It; ++It)
                             {
-								FPipelineCacheFileFormatPSO& NewEntry = *It;
+								FPipelineCacheFileFormatPSO& NewEntry = It->Value;
 
-                                uint32 PSOHash = GetTypeHash(NewEntry);
+								uint32 PSOHash = It->Key;
 								
 								bool bFound = FPipelineFileCacheManager::IsPSOEntryCached(NewEntry, nullptr);
 								if (bFound)
@@ -2641,15 +2650,15 @@ public:
                             for (auto& Entry : NewEntries)
                             {
                                 FPipelineCacheFileFormatPSOMetaData Meta;
-                                Meta.Stats.PSOHash = GetTypeHash(Entry);
+                                Meta.Stats.PSOHash = Entry.Key;
                                 Meta.FileOffset = 0;
                                 Meta.FileSize = 0;
                                 Meta.FileGuid = Header.Guid;
-								Meta.AddShaders(Entry);
+								Meta.AddShaders(Entry.Value);
 								check(Meta.FileGuid != FGuid());
 
                                 TempTOC.MetaData.Add(Meta.Stats.PSOHash, Meta);
-                                PSOs.Add(Meta.Stats.PSOHash, Entry);
+                                PSOs.Add(Entry.Key, Entry.Value);
                             }
 							
 							// Update TOC Metadata usage masks - don't clear NewPSOUsage as we are using a TempTOC
@@ -3313,7 +3322,6 @@ void FPipelineFileCacheManager::Shutdown()
 		}
 		Stats.Empty();
 		NewPSOs.Empty();
-		NewPSOHashes.Empty();
         NumNewPSOs = 0;
 		
 		FileCacheMap.Empty();
@@ -3392,7 +3400,6 @@ bool FPipelineFileCacheManager::OpenUserPipelineFileCache(const FString& Key, co
 
 		// User Cache now exists - these caches should be empty for this file otherwise will have false positives from any previous file caching - if not something has been caching when it should not be
 		check(NewPSOs.Num() == 0);
-		check(NewPSOHashes.Num() == 0);
 		check(RunTimeToPSOUsage.Num() == 0);
 	}
 	return bUserFileOpened;
@@ -3434,7 +3441,6 @@ void FPipelineFileCacheManager::CloseUserPipelineFileCache()
 		RunTimeToPSOUsage.Empty();
 		NewPSOUsage.Empty();
 		NewPSOs.Empty();
-		NewPSOHashes.Empty();
 		NumNewPSOs = 0;
 
 		SET_MEMORY_STAT(STAT_NewCachedPSOMemory, 0);
@@ -3587,7 +3593,7 @@ void FPipelineFileCacheManager::CacheGraphicsPSO(uint32 RunTimeHash, FGraphicsPi
 
 				if (!FPipelineFileCacheManager::IsPSOEntryCached(NewEntry, &CurrentUsageData))
 				{
-					bool bActuallyNewPSO = !NewPSOHashes.Contains(PSOHash);
+					bool bActuallyNewPSO = !NewPSOs.Contains(PSOHash);
 					if (Initializer.bFromPSOFileCache)
 					{
 						// FIXME: this is a workaround. Needs proper investigation
@@ -3602,23 +3608,27 @@ void FPipelineFileCacheManager::CacheGraphicsPSO(uint32 RunTimeHash, FGraphicsPi
 
 					if (bActuallyNewPSO)
 					{
-						LogNewGraphicsPSOToConsoleAndCSV(NewEntry, PSOHash, PSOPrecacheResult);
-
 						// Only add PSOs which are not fully precached yet if GPSOExcludePrecachePSOsInFileCache is set (includes miss, too late, unknown, NotSupported)
 						if (PSOPrecacheResult == EPSOPrecacheResult::Complete)
 						{
 							bActuallyNewPSO = !GPSOExcludePrecachePSOsInFileCache;
 						}
+						// Only include precache-only PSOs if r.ShaderPipelineCache.IncludeAllPrecachedPSOs is enabled
+						if (PSOPrecacheResult == EPSOPrecacheResult::Active)
+						{
+							bActuallyNewPSO = GPSOIncludeAllPrecachedPSOsInFileCache != 0;
+						}
 					}
 
 					if (bActuallyNewPSO) 
-					{
+					{						
+						LogNewGraphicsPSOToConsoleAndCSV(NewEntry, PSOHash, PSOPrecacheResult);
+
 						if (LogPSOtoFileCache())
 						{
-							NewPSOs.Add(NewEntry);
+							NewPSOs.Add(PSOHash, NewEntry);
 							INC_MEMORY_STAT_BY(STAT_NewCachedPSOMemory, sizeof(FPipelineCacheFileFormatPSO) + sizeof(uint32) + sizeof(uint32));
 						}
-						NewPSOHashes.Add(PSOHash);
 
 						NumNewPSOs++;
 						INC_DWORD_STAT(STAT_NewGraphicsPipelineStateCount);
@@ -3645,6 +3655,14 @@ void FPipelineFileCacheManager::CacheGraphicsPSO(uint32 RunTimeHash, FGraphicsPi
 			{
 				PSOUsage->UsageMask |= FPipelineFileCacheManager::GameUsageMask;
 				RegisterPSOUsageDataUpdateForNextSave(*PSOUsage);
+
+				// If this PSO was previously logged as precache-only, upgrade its stored result to the
+				// real runtime outcome. The Active state itself guarantees the CVar was set when logged.
+				FPipelineCacheFileFormatPSO* Entry = NewPSOs.Find(PSOUsage->PSOHash);
+				if (Entry && Entry->PSOPrecacheResult == EPSOPrecacheResult::Active)
+				{
+					Entry->PSOPrecacheResult = PSOPrecacheResult;
+				}
 			}
 		}
 
@@ -3680,28 +3698,31 @@ void FPipelineFileCacheManager::CacheComputePSO(uint32 RunTimeHash, FComputePipe
 					
 					if (!FPipelineFileCacheManager::IsPSOEntryCached(NewEntry, &CurrentUsageData))
 					{
-						bool bActuallyNewPSO = !NewPSOHashes.Contains(PSOHash);
+						bool bActuallyNewPSO = !NewPSOs.Contains(PSOHash);
 						if (bActuallyNewPSO)
 						{
-							LogNewComputePSOToConsoleAndCSV(NewEntry, PSOHash, PSOPrecacheResult);
-
 							// Only add PSOs which are not fully precached yet if GPSOExcludePrecachePSOsInFileCache is set (includes miss, too late, unknown, NotSupported)
 							if (PSOPrecacheResult == EPSOPrecacheResult::Complete)
 							{
 								bActuallyNewPSO = !GPSOExcludePrecachePSOsInFileCache;
-							}							
+							}
+							// Only include precache-only PSOs if r.ShaderPipelineCache.IncludeAllPrecachedPSOs is enabled
+							if (PSOPrecacheResult == EPSOPrecacheResult::Active)
+							{
+								bActuallyNewPSO = GPSOIncludeAllPrecachedPSOsInFileCache != 0;
+							}
 						}
 
 						if (bActuallyNewPSO)
 						{
+							LogNewComputePSOToConsoleAndCSV(NewEntry, PSOHash, PSOPrecacheResult);
+
 							if (LogPSOtoFileCache())
 							{
-								NewPSOs.Add(NewEntry);
+								NewPSOs.Add(PSOHash, NewEntry);
 								INC_MEMORY_STAT_BY(STAT_NewCachedPSOMemory, sizeof(FPipelineCacheFileFormatPSO) + sizeof(uint32) + sizeof(uint32));
 							}
 
-							NewPSOHashes.Add(PSOHash);
-
 							NumNewPSOs++;
 							INC_DWORD_STAT(STAT_NewComputePipelineStateCount);
 							INC_DWORD_STAT(STAT_TotalComputePipelineStateCount);
@@ -3727,6 +3748,14 @@ void FPipelineFileCacheManager::CacheComputePSO(uint32 RunTimeHash, FComputePipe
 				{
 					PSOUsage->UsageMask |= FPipelineFileCacheManager::GameUsageMask;
 					RegisterPSOUsageDataUpdateForNextSave(*PSOUsage);
+
+					// If this PSO was previously logged as precache-only, upgrade its stored result to the
+					// real runtime outcome. The Active state itself guarantees the CVar was set when logged.
+					FPipelineCacheFileFormatPSO* Entry = NewPSOs.Find(PSOUsage->PSOHash);
+					if (Entry && Entry->PSOPrecacheResult == EPSOPrecacheResult::Active)
+					{
+						Entry->PSOPrecacheResult = PSOPrecacheResult;
+					}
 				}
 			}
 		}
@@ -3785,7 +3814,7 @@ void FPipelineFileCacheManager::CacheRayTracingPSO(const FRayTracingPipelineStat
 
 						if (LogPSOtoFileCache())
 						{
-							NewPSOs.Add(NewEntry);
+							NewPSOs.Add(PSOHash, NewEntry);
 							INC_MEMORY_STAT_BY(STAT_NewCachedPSOMemory, sizeof(FPipelineCacheFileFormatPSO) + sizeof(uint32) + sizeof(uint32));
 						}
 
@@ -4093,9 +4122,9 @@ uint32 FPipelineFileCacheManager::NumPSOsLogged()
 		{
 			// This can happen if the Mask was zero at some point
 			
-			for (auto& PSO : NewPSOs)
+			for (auto& PSOPair : NewPSOs)
 			{
-				FPipelineStateStats const* Stat = Stats.FindRef(GetTypeHash(PSO));
+				FPipelineStateStats const* Stat = Stats.FindRef(PSOPair.Key);
 				if (Stat && Stat->TotalBindCount > 0)
 				{
 					Result++;
diff --git a/Engine/Source/Runtime/RHI/Private/PipelineStateCache.cpp b/Engine/Source/Runtime/RHI/Private/PipelineStateCache.cpp
--- a/Engine/Source/Runtime/RHI/Private/PipelineStateCache.cpp
+++ b/Engine/Source/Runtime/RHI/Private/PipelineStateCache.cpp
@@ -2960,6 +2960,14 @@ public:
 			UpdatePerPriorityCompileCount(StatePriority, false /*Increment*/);
 		}
 
+		// Log this PSO to the pipeline file cache as 'Active' (precached, not yet seen at draw time).
+		// Whether it is actually stored is controlled by r.ShaderPipelineCache.IncludeAllPrecachedPSOs.
+		// The result will be upgraded to the real runtime outcome if the PSO is later needed for rendering.
+		if (bValid)
+		{
+			TPrecachePipelineCacheDerived::CacheToFileCache(Initializer);
+		}
+
 		// Mark as in memory now - if it's not cached then it was added to list for cleanup and will be subtracted there
 		if (PSOType == FPSOPrecacheRequestID::EType::Graphics)
 		{
@@ -3324,6 +3332,11 @@ public:
 	{
 		return HashCombineFast(GetTypeHash(Key.ComputeShader->GetHash()), GetTypeHash(Key.SpecializationConstants));
 	}
+
+	static void CacheToFileCache(const FComputePipelineStateInitializer& Initializer)
+	{
+		FPipelineFileCacheManager::CacheComputePSO(GetTypeHash(Initializer), Initializer, EPSOPrecacheResult::Active, nullptr);
+	}
 };
 
 class FPrecacheGraphicsPipelineCache : public TPrecachePipelineCacheBase<FPrecacheGraphicsPipelineCache, FGraphicsPipelineStateInitializer, FGraphicsPipelineState>
@@ -3342,6 +3355,11 @@ public:
 		return RHIComputePrecachePSOHash(Key);
 	}
 
+	static void CacheToFileCache(const FGraphicsPipelineStateInitializer& Initializer)
+	{
+		FPipelineFileCacheManager::CacheGraphicsPSO(GetTypeHash(Initializer), Initializer, EPSOPrecacheResult::Active, nullptr);
+	}
+
 	static const FName GetActiveCompileStatName()
 	{
 		return GET_STATFNAME(STAT_ActiveGraphicsPSOPrecacheRequests);
diff --git a/Engine/Source/Runtime/RHI/Public/PipelineFileCache.h b/Engine/Source/Runtime/RHI/Public/PipelineFileCache.h
--- a/Engine/Source/Runtime/RHI/Public/PipelineFileCache.h
+++ b/Engine/Source/Runtime/RHI/Public/PipelineFileCache.h
@@ -444,9 +444,8 @@ private:
 	static TMap<uint32, FPSOUsageData> RunTimeToPSOUsage;		// Fast check structure - Not saved (External state cache runtime hash to seen usage data)
 	static TMap<uint32, FPSOUsageData> NewPSOUsage;				// For mask or engine updates - Merged + Saved (Our internal PSO hash to latest usage data) - temp working scratch, only holds updates since last "save" so is not the authority on state
 	static TMap<uint32, FPipelineStateStats*> Stats;
-	static TSet<FPipelineCacheFileFormatPSO> NewPSOs;
+	static TMap<uint32, FPipelineCacheFileFormatPSO> NewPSOs;
 	static TMpscQueue<FPipelineCacheFileFormatPSO> NewPSOsToReport; // New PSOs that will be broadcast via a delegate (if bound). Cleared when the delegate is broadcast.
- 	static TSet<uint32> NewPSOHashes;
 	static uint32 NumNewPSOs;
 	static PSOOrder RequestedOrder;
 	static bool FileCacheEnabled;
