Fix inconsistencies in RHI context handling, and possible use-after-free of context pointers
- Merged the parallel and non-parallel RHI context finalize functions together
- Refactored the handling of the upload context pointer so that it is passed through the pipeline alongside the standard contexts in calls to RHIFinalizeContexts / RHISubmitCommandLists.

Renamed the GRHISupportsParallelRenderPasses, GRHIParallelRHIExecuteChildWait and GRHIParallelRHIExecuteParentWait bools
- Packed them together in a GRHIGlobals.ParallelRenderPasses struct, and renamed to clarify which task needs to wait on what.

#Synced-CL 52204891

diff --git a/Engine/Plugins/RenderTraceInsights/Source/RenderTraceInsights/Private/RenderTraceAnalyzer.cpp b/Engine/Plugins/RenderTraceInsights/Source/RenderTraceInsights/Private/RenderTraceAnalyzer.cpp
--- a/Engine/Plugins/RenderTraceInsights/Source/RenderTraceInsights/Private/RenderTraceAnalyzer.cpp
+++ b/Engine/Plugins/RenderTraceInsights/Source/RenderTraceInsights/Private/RenderTraceAnalyzer.cpp
@@ -829,18 +829,18 @@ void FRenderTraceAnalyzer::ProcessSetTranslateContext(const FEventData& EventDat
 		return;
 	}
 
-	// RHIs usually have a pool of contexts, and objects are returned to the pool in their RHIFinalizeContext() implementation,
+	// RHIs usually have a pool of contexts, and objects are returned to the pool in their RHIFinalizeContexts() implementation,
 	// which runs before the FinalizeTranslateJob event is sent. Due to threading, another translate job can pick up the context
 	// before the finalize event can remove it from LiveRHIContexts. This is why we can't check if the context is already in use;
 	// instead, we just steal it.
-	// Note that we can't send FinalizeTranslateJob before calling RHIFinalizeContext(), because that function will do more things
+	// Note that we can't send FinalizeTranslateJob before calling RHIFinalizeContexts(), because that function will do more things
 	// with the context (such as signal batched sync points) and we still need to find the context -> job association when that
 	// happens. The association will still be valid when it does that, because these extra actions happen before returning the context
 	// to the pool, so another job can't steal the context while we still need the old association.
 	// Example flow with D3D12:
 	//		FRHICommandListExecutor::FTranslateState::Translate_CloseChain()
 	//			|-> GDynamicRHI->RHICloseTranslateChain()
-	//			|	|-> FD3D12DynamicRHI::RHIFinalizeContext()
+	//			|	|-> FD3D12DynamicRHI::RHIFinalizeContexts()
 	//			|		|-> FD3D12ContextCommon::Finalize()
 	//			|		|	|-> Payload->SyncPointsToWait.Append(BatchedSyncPoints.ToWait) - this needs to find the current context -> translate job mapping
 	//			|		|-> CmdContext->GetParentDevice()->ReleaseContext() - this releases the context, so another thread can pick it up for its translate job
diff --git a/Engine/Source/Runtime/Apple/MetalRHI/Private/MetalDynamicRHI.h b/Engine/Source/Runtime/Apple/MetalRHI/Private/MetalDynamicRHI.h
--- a/Engine/Source/Runtime/Apple/MetalRHI/Private/MetalDynamicRHI.h
+++ b/Engine/Source/Runtime/Apple/MetalRHI/Private/MetalDynamicRHI.h
@@ -193,13 +193,11 @@ public:
 	virtual IRHIUploadContext* RHIGetUploadContext() final override;
 	
 	virtual IRHIComputeContext* RHIGetCommandContext(ERHIPipeline Pipeline) final override;
-	virtual IRHIComputeContext* RHIGetParallelCommandContext(FRHIParallelRenderPassInfo const& ParallelRenderPass) final override;
-	
-	virtual void RHIProcessDeleteQueue() final override;
-	virtual void RHIFinalizeContext(FRHIFinalizeContextArgs&& Args, TRHIPipelineArray<IRHIPlatformCommandList*>& Output) final override;
-	virtual IRHIPlatformCommandList* RHIFinalizeParallelContext(IRHIComputeContext* Context) final override;
-	
+	virtual IRHICommandContext* RHIGetParallelCommandContext(FRHIParallelRenderPassInfo const& ParallelRenderPass) final override;
+
+	virtual FRHIFinalizeContextsResult RHIFinalizeContexts(FRHIFinalizeContextsArgs&& Args) final override;
 	virtual void RHISubmitCommandLists(FRHISubmitCommandListsArgs&& Args) final override;
+	virtual void RHIProcessDeleteQueue() final override;
 
 	virtual FUpdateTexture3DData RHIBeginUpdateTexture3D(FRHICommandListBase& RHICmdList, FRHITexture* Texture, uint32 MipIndex, const struct FUpdateTextureRegion3D& UpdateRegion) final override;
 	virtual void RHIEndUpdateTexture3D(FRHICommandListBase& RHICmdList, FUpdateTexture3DData& UpdateData) final override;
diff --git a/Engine/Source/Runtime/Apple/MetalRHI/Private/MetalRHI.cpp b/Engine/Source/Runtime/Apple/MetalRHI/Private/MetalRHI.cpp
--- a/Engine/Source/Runtime/Apple/MetalRHI/Private/MetalRHI.cpp
+++ b/Engine/Source/Runtime/Apple/MetalRHI/Private/MetalRHI.cpp
@@ -571,9 +571,10 @@ FMetalDynamicRHI::FMetalDynamicRHI(ERHIFeatureLevel::Type RequestedFeatureLevel)
 	bool bSupportsParallel = CVarMetalParallel.GetValueOnAnyThread() || FParse::Param(FCommandLine::Get(),TEXT("rhiparallel"));
 	
 	GRHISupportsParallelRHIExecute = bSupportsParallel;
-	GRHIParallelRHIExecuteChildWait = true;
-	GRHIParallelRHIExecuteParentWait = true;
-	GRHISupportsParallelRenderPasses = bSupportsParallel;
+	GRHIGlobals.ParallelRenderPasses.Supported = bSupportsParallel;
+	GRHIGlobals.ParallelRenderPasses.ParentWaitForChildTranslate = true;
+	GRHIGlobals.ParallelRenderPasses.ChildWaitForParentBegin = true;
+	GRHIGlobals.ParallelRenderPasses.UsesCustomContexts = true;
 	
 	if (FPlatformMisc::IsDebuggerPresent() && UE_BUILD_DEBUG)
 	{
@@ -1528,7 +1529,7 @@ IRHIComputeContext* FMetalDynamicRHI::RHIGetCommandContext(ERHIPipeline Pipeline
 	return static_cast<IRHIComputeContext*>(Context);
 }
 
-IRHIComputeContext* FMetalDynamicRHI::RHIGetParallelCommandContext(FRHIParallelRenderPassInfo const& ParallelRenderPass)
+IRHICommandContext* FMetalDynamicRHI::RHIGetParallelCommandContext(FRHIParallelRenderPassInfo const& ParallelRenderPass)
 {
 	check(GRHISupportsParallelRHIExecute);
 	
@@ -1541,7 +1542,7 @@ IRHIComputeContext* FMetalDynamicRHI::RHIGetParallelCommandContext(FRHIParallelR
 	Context->ResetContext();
 	Context->SetupParallelContext(&ParallelRenderPass);
 	
-	return static_cast<IRHIComputeContext*>(Context);
+	return Context;
 }
 
 class FMetalPlatformCommandList final : public IRHIPlatformCommandList
diff --git a/Engine/Source/Runtime/Apple/MetalRHI/Private/MetalSubmission.cpp b/Engine/Source/Runtime/Apple/MetalRHI/Private/MetalSubmission.cpp
--- a/Engine/Source/Runtime/Apple/MetalRHI/Private/MetalSubmission.cpp
+++ b/Engine/Source/Runtime/Apple/MetalRHI/Private/MetalSubmission.cpp
@@ -164,48 +164,39 @@ IRHIUploadContext* FMetalDynamicRHI::RHIGetUploadContext()
 	return static_cast<IRHIUploadContext*>(Context);
 }
 
-void FMetalDynamicRHI::RHIFinalizeContext(FRHIFinalizeContextArgs&& Args, TRHIPipelineArray<IRHIPlatformCommandList*>& Output)
+FDynamicRHI::FRHIFinalizeContextsResult FMetalDynamicRHI::RHIFinalizeContexts(FRHIFinalizeContextsArgs&& Args)
 {
 	MTL_SCOPED_AUTORELEASE_POOL;
 	
-	FMetalRHIUploadContext* UploadContext = static_cast<FMetalRHIUploadContext*>(Args.UploadContext);
-
-	FMetalFinalizedCommands Commands;
+	FRHIFinalizeContextsResult Result;
 	
-	if (UploadContext)
+	if (FMetalRHIUploadContext* UploadContext = static_cast<FMetalRHIUploadContext*>(Args.UploadContext))
 	{
+		FMetalFinalizedCommands Commands;
 		UploadContext->Finalize(Commands);
 		MetalUploadContextPool.Push(UploadContext);
+
+		if (Commands.Num())
+		{
+			Result.UploadCommandList = new FMetalFinalizedCommands(MoveTemp(Commands));
+		}
 	}
 	
-	for(IRHIComputeContext* Context : Args.Contexts)
+	for (ERHIPipeline Pipeline : MakeFlagsRange(ERHIPipeline::All))
 	{
-		FMetalRHICommandContext* CmdContext = static_cast<FMetalRHICommandContext*>(Context);
-		
-		if (!CmdContext->IsInsideRenderPass())
+		if (FMetalRHICommandContext* Context = static_cast<FMetalRHICommandContext*>(Args.Contexts[Pipeline]))
 		{
-			CmdContext->Finalize(Commands);
-			CmdContext->ResetContext();
-			
-			MetalCommandContextPool.Push(CmdContext);
-		}
+			FMetalFinalizedCommands Commands;
+			Context->Finalize(Commands);
+			Context->ResetContext();
+
+			MetalCommandContextPool.Push(Context);
 		
-		Output[Context->GetPipeline()] = new FMetalFinalizedCommands(MoveTemp(Commands));
+			Result.CommandLists[Pipeline] = new FMetalFinalizedCommands(MoveTemp(Commands));
+		}
 	}
-}
 
-IRHIPlatformCommandList* FMetalDynamicRHI::RHIFinalizeParallelContext(IRHIComputeContext* Context)
-{
-	FMetalFinalizedCommands* Commands = new FMetalFinalizedCommands;
-	
-	FMetalRHICommandContext* CmdContext = static_cast<FMetalRHICommandContext*>(Context);
-	CmdContext->Finalize(*Commands);
-	
-	CmdContext->ResetContext();
-	
-	MetalCommandContextPool.Push(CmdContext);
-	
-	return Commands;
+	return Result;
 }
 
 void FMetalDynamicRHI::RHISubmitCommandLists(FRHISubmitCommandListsArgs&& Args)
diff --git a/Engine/Source/Runtime/D3D12RHI/Private/D3D12RHI.cpp b/Engine/Source/Runtime/D3D12RHI/Private/D3D12RHI.cpp
--- a/Engine/Source/Runtime/D3D12RHI/Private/D3D12RHI.cpp
+++ b/Engine/Source/Runtime/D3D12RHI/Private/D3D12RHI.cpp
@@ -228,7 +228,7 @@ FD3D12DynamicRHI::FD3D12DynamicRHI(TSharedPtr<FD3D12Adapter>& InChosenAdapter, b
 	GRHISupportsRHIThread = true;
 
 	GRHISupportsParallelRHIExecute = true;
-	GRHISupportsParallelRenderPasses = true;
+	GRHIGlobals.ParallelRenderPasses.Supported = true;
 
 	GRHISupportsRawViewsForAnyBuffer = true;
 
diff --git a/Engine/Source/Runtime/D3D12RHI/Private/D3D12RHIPrivate.h b/Engine/Source/Runtime/D3D12RHI/Private/D3D12RHIPrivate.h
--- a/Engine/Source/Runtime/D3D12RHI/Private/D3D12RHIPrivate.h
+++ b/Engine/Source/Runtime/D3D12RHI/Private/D3D12RHIPrivate.h
@@ -425,7 +425,7 @@ public:
 	virtual void* RHIGetNativeComputeQueue() final override;
 	virtual class IRHICommandContext* RHIGetDefaultContext() final override;
 	virtual IRHIComputeContext* RHIGetCommandContext(ERHIPipeline Pipeline) final override;
-	virtual void RHIFinalizeContext(FRHIFinalizeContextArgs&& Args, TRHIPipelineArray<IRHIPlatformCommandList*>& Output) final override;
+	virtual FRHIFinalizeContextsResult RHIFinalizeContexts(FRHIFinalizeContextsArgs&& Args) final override;
 	virtual void RHISubmitCommandLists(FRHISubmitCommandListsArgs&& Args) final override;
 	virtual void RHIProcessDeleteQueue() final override;
 
diff --git a/Engine/Source/Runtime/D3D12RHI/Private/D3D12Submission.cpp b/Engine/Source/Runtime/D3D12RHI/Private/D3D12Submission.cpp
--- a/Engine/Source/Runtime/D3D12RHI/Private/D3D12Submission.cpp
+++ b/Engine/Source/Runtime/D3D12RHI/Private/D3D12Submission.cpp
@@ -209,9 +209,11 @@ void FD3D12DynamicRHI::ShutdownSubmissionPipe()
 struct FD3D12FinalizedCommands : public IRHIPlatformCommandList, public TArray<FD3D12Payload*>
 {};
 
-void FD3D12DynamicRHI::RHIFinalizeContext(FRHIFinalizeContextArgs&& Args, TRHIPipelineArray<IRHIPlatformCommandList*>& Output)
+FDynamicRHI::FRHIFinalizeContextsResult FD3D12DynamicRHI::RHIFinalizeContexts(FRHIFinalizeContextsArgs&& Args)
 {
-	auto FinalizeContext = [&](FD3D12CommandContext* CmdContext, FD3D12FinalizedCommands& Result)
+	checkf(Args.UploadContext == nullptr, TEXT("D3D12 RHI does not implement upload contexts."));
+
+	auto FinalizeContext = [](FD3D12CommandContext* CmdContext, FD3D12FinalizedCommands& Result)
 	{
 		CmdContext->Finalize(Result);
 
@@ -222,30 +224,37 @@ void FD3D12DynamicRHI::RHIFinalizeContext(FRHIFinalizeContextArgs&& Args, TRHIPi
 		}
 	};
 
-	for(IRHIComputeContext* Context : Args.Contexts)
+	FDynamicRHI::FRHIFinalizeContextsResult Result;
+	for (ERHIPipeline Pipeline : MakeFlagsRange(ERHIPipeline::All))
 	{
-		FD3D12FinalizedCommands Result;
-		ERHIPipeline Pipeline = Context->GetPipeline();
-		
-		FD3D12CommandContextBase* CmdContextBase = static_cast<FD3D12CommandContextBase*>(Context);
-		if (FD3D12CommandContextRedirector* Redirector = CmdContextBase->AsRedirector())
+		if (FD3D12CommandContextBase* ContextBase = static_cast<FD3D12CommandContextBase*>(Args.Contexts[Pipeline]))
 		{
-			for (uint32 GPUIndex : Redirector->GetPhysicalGPUMask())
-				FinalizeContext(Redirector->GetSingleDeviceContext(GPUIndex), Result);
-			
-			if (!Redirector->bIsDefaultContext)
+			FD3D12FinalizedCommands FinalizedCommands;
+			if (FD3D12CommandContextRedirector* Redirector = ContextBase->AsRedirector())
 			{
-				delete Redirector;
+				for (uint32 GPUIndex : Redirector->GetPhysicalGPUMask())
+				{
+					FinalizeContext(Redirector->GetSingleDeviceContext(GPUIndex), FinalizedCommands);
+				}
+
+				if (!Redirector->bIsDefaultContext)
+				{
+					delete Redirector;
+				}
+			}
+			else
+			{
+				FinalizeContext(static_cast<FD3D12CommandContext*>(ContextBase), FinalizedCommands);
+			}
+
+			if (FinalizedCommands.Num())
+			{
+				Result.CommandLists[Pipeline] = new FD3D12FinalizedCommands(MoveTemp(FinalizedCommands));
 			}
 		}
-		else
-		{
-			FD3D12CommandContext* CmdContext = static_cast<FD3D12CommandContext*>(CmdContextBase);
-			FinalizeContext(CmdContext, Result);
-		}
-		
-		Output[Pipeline] = Result.Num() ? new FD3D12FinalizedCommands(MoveTemp(Result)) : nullptr;
 	}
+
+	return Result;
 }
 
 void FD3D12DynamicRHI::RHISubmitCommandLists(FRHISubmitCommandListsArgs&& Args)
diff --git a/Engine/Source/Runtime/NullDrv/Public/NullRHI.h b/Engine/Source/Runtime/NullDrv/Public/NullRHI.h
--- a/Engine/Source/Runtime/NullDrv/Public/NullRHI.h
+++ b/Engine/Source/Runtime/NullDrv/Public/NullRHI.h
@@ -568,20 +568,10 @@ public:
 	{
 		return nullptr;
 	}
-	
-	virtual IRHIComputeContext* RHIGetParallelCommandContext(FRHIParallelRenderPassInfo const& ParallelRenderPass) final override
-	{
-		return nullptr;
-	}
 
-	virtual void RHIFinalizeContext(FRHIFinalizeContextArgs&& Args, TRHIPipelineArray<IRHIPlatformCommandList*>& Output) final override
+	virtual FRHIFinalizeContextsResult RHIFinalizeContexts(FRHIFinalizeContextsArgs&& Args) final override
 	{
-	}
-	
-	IRHIPlatformCommandList* RHIFinalizeParallelContext(IRHIComputeContext* Context) override final
-	{
-		// TODO: Implement and test this
-		return nullptr;
+		return {};
 	}
 
 	virtual void RHISubmitCommandLists(FRHISubmitCommandListsArgs&& Args) final override
diff --git a/Engine/Source/Runtime/OpenGLDrv/Private/OpenGLCommands.cpp b/Engine/Source/Runtime/OpenGLDrv/Private/OpenGLCommands.cpp
--- a/Engine/Source/Runtime/OpenGLDrv/Private/OpenGLCommands.cpp
+++ b/Engine/Source/Runtime/OpenGLDrv/Private/OpenGLCommands.cpp
@@ -2730,19 +2730,28 @@ IRHIComputeContext* FOpenGLDynamicRHI::RHIGetCommandContext(ERHIPipeline Pipelin
 	return nullptr;
 }
 
-void FOpenGLDynamicRHI::RHIFinalizeContext(FRHIFinalizeContextArgs&& Args, TRHIPipelineArray<IRHIPlatformCommandList*>& Output)
+FDynamicRHI::FRHIFinalizeContextsResult FOpenGLDynamicRHI::RHIFinalizeContexts(FRHIFinalizeContextsArgs&& Args)
 {
+	checkf(Args.UploadContext == nullptr, TEXT("OpenGL RHI does not implement upload contexts."));
+
 	FlushProfilerStats();
 
-	for(IRHIComputeContext* Context : Args.Contexts)
+#if DO_CHECK
+	for (ERHIPipeline Pipeline : MakeFlagsRange(ERHIPipeline::All))
 	{
-		// "Context" will always be the default context, since we don't implement parallel execution.
-		check(Context == this);
-	}
+	    if (IRHIComputeContext* Context = Args.Contexts[Pipeline])
+	    {
+		    // "Context" will always be the default graphics context, since we don't implement parallel execution.
+		    check(Pipeline == ERHIPipeline::Graphics && Context == this);
+	    }
+    }
+#endif
 
 	// Clear some context state
 	FMemory::Memset(PendingState.BoundUniformBuffers, 0, sizeof(PendingState.BoundUniformBuffers));
 	FMemory::Memset(PendingState.BoundUniformBuffersDynamicOffset, 0u, sizeof(PendingState.BoundUniformBuffersDynamicOffset));
+
+	return {};
 }
 
 void FOpenGLDynamicRHI::RHISubmitCommandLists(FRHISubmitCommandListsArgs&& Args)
diff --git a/Engine/Source/Runtime/OpenGLDrv/Private/OpenGLDrv.h b/Engine/Source/Runtime/OpenGLDrv/Private/OpenGLDrv.h
--- a/Engine/Source/Runtime/OpenGLDrv/Private/OpenGLDrv.h
+++ b/Engine/Source/Runtime/OpenGLDrv/Private/OpenGLDrv.h
@@ -171,7 +171,7 @@ public:
 	UE_API virtual void* RHIGetNativeDevice() final override;
 	UE_API virtual class IRHICommandContext* RHIGetDefaultContext() final override;
 	UE_API virtual IRHIComputeContext* RHIGetCommandContext(ERHIPipeline Pipeline) final override;
-	UE_API virtual void RHIFinalizeContext(FRHIFinalizeContextArgs&& Args, TRHIPipelineArray<IRHIPlatformCommandList*>& Output) final override;
+	UE_API virtual FRHIFinalizeContextsResult RHIFinalizeContexts(FRHIFinalizeContextsArgs&& Args) final override;
 	UE_API virtual void RHISubmitCommandLists(FRHISubmitCommandListsArgs&& Args) final override;
 
 	UE_API virtual void RHIBeginRenderPass(const FRHIRenderPassInfo& InInfo, const TCHAR* InName) final override;
diff --git a/Engine/Source/Runtime/RHI/Private/RHICommandList.cpp b/Engine/Source/Runtime/RHI/Private/RHICommandList.cpp
--- a/Engine/Source/Runtime/RHI/Private/RHICommandList.cpp
+++ b/Engine/Source/Runtime/RHI/Private/RHICommandList.cpp
@@ -182,10 +182,10 @@ void FRHICommandListBase::InsertParallelRenderPass_Base(TSharedPtr<FRHIParallelR
 {
 	ERHIPipeline CurrentPipelines = ActivePipelines;
 
-	bool bRequiresWait = GRHIParallelRHIExecuteChildWait || GRHIParallelRHIExecuteParentWait;
+	bool bRequiresWait = GRHIGlobals.ParallelRenderPasses.ParentWaitForChildTranslate || GRHIGlobals.ParallelRenderPasses.ChildWaitForParentBegin;
 	
 	// Finish current command list with the start of the parallel RP
-	if(bRequiresWait)
+	if (bRequiresWait)
 	{
 		ParallelRenderPassBegin = InInfo;
 	}
@@ -221,7 +221,7 @@ void FRHICommandListBase::InsertParallelRenderPass_Base(TSharedPtr<FRHIParallelR
 	ActivatePipelines(CurrentPipelines);
 
 	// Start the new RHICmdList with the end of the parallel render pass.
-	if(bRequiresWait)
+	if (bRequiresWait)
 	{	
 		ParallelRenderPassEnd = InInfo;
 	}
@@ -408,7 +408,18 @@ void FRHICommandListBase::ActivatePipelines(ERHIPipeline Pipelines)
 				{
 					if (ExecutingCmdList.IsSubCommandList())
 					{
-						Context = GDynamicRHI->RHIGetParallelCommandContext(*ExecutingCmdList.SubRenderPassInfo);
+						if (GRHIGlobals.ParallelRenderPasses.UsesCustomContexts)
+						{
+							Context = GDynamicRHI->RHIGetParallelCommandContext(*ExecutingCmdList.SubRenderPassInfo);
+						}
+						else
+						{
+							// Platform does not implement RHIGetParallelCommandContext. Use a regular context and begin/end the render pass directly.
+							IRHICommandContext* GraphicsContext = static_cast<IRHICommandContext*>(GDynamicRHI->RHIGetCommandContext(ERHIPipeline::Graphics));
+							GraphicsContext->RHIBeginRenderPass(*ExecutingCmdList.SubRenderPassInfo, ExecutingCmdList.SubRenderPassInfo->PassName);
+
+							Context = GraphicsContext;
+						}
 					}
 					else
 					{
@@ -859,8 +870,7 @@ void FRHICommandListExecutor::FSubmitState::Dispatch_ProcessCommandList(FRHIComm
 										MaxCommandsPerTranslate == 0 ? false :
 										(TranslateJobs.Current->NumCommands + CmdList->NumCommands) > uint32(MaxCommandsPerTranslate);
 	
-		bool bShouldSplitForParentChild = ( CmdList->IsSubCommandList() && !TranslateJobs.Current->bUsingSubCmdLists) ||
-										  (!CmdList->IsSubCommandList() &&  TranslateJobs.Current->bUsingSubCmdLists);
+		bool bShouldSplitForParentChild = CmdList->IsSubCommandList() != TranslateJobs.Current->bUsingSubCmdLists;
 	
 		if (bShouldSplitParallel || bShouldSplitForThreshold || bShouldSplitForParentChild)
 		{
@@ -870,12 +880,11 @@ void FRHICommandListExecutor::FSubmitState::Dispatch_ProcessCommandList(FRHIComm
 
 			IF_RENDER_TRACING_ENABLED(RenderTracing::SplitTranslateJob(TranslateJobs.Current, *CmdList, bShouldSplitParallel, bShouldSplitForThreshold, bShouldSplitForParentChild));
 
-			bool bAddChildWait = TranslateJobs.Current->bUsingSubCmdLists && GRHIParallelRHIExecuteChildWait;
-			bool bAddParentWait = !TranslateJobs.Current->bUsingSubCmdLists && CmdList->IsSubCommandList() && GRHIParallelRHIExecuteParentWait;
+			bool bAddChildWait  =  TranslateJobs.Current->bUsingSubCmdLists && GRHIGlobals.ParallelRenderPasses.ParentWaitForChildTranslate;
+			bool bAddParentWait = !TranslateJobs.Current->bUsingSubCmdLists && CmdList->IsSubCommandList() && GRHIGlobals.ParallelRenderPasses.ChildWaitForParentBegin;
 
 			FGraphEventRef Event = Dispatch_CloseTranslateChain(false, 0);
 			if (bAddChildWait)
-
 			{
 				ChildGraphEvents.Add(Event);
 			}
@@ -896,7 +905,7 @@ void FRHICommandListExecutor::FSubmitState::Dispatch_ProcessCommandList(FRHIComm
 	}
 
 	check(TranslateJobs.Current->bShouldFinalize);
-	TranslateJobs.Current->bShouldFinalize = !(GRHIParallelRHIExecuteChildWait && CmdList->ParallelRenderPassBegin.IsValid());
+	TranslateJobs.Current->bShouldFinalize = !(GRHIGlobals.ParallelRenderPasses.ParentWaitForChildTranslate && CmdList->ParallelRenderPassBegin.IsValid());
 	TranslateJobs.Current->NumCommands += CmdList->NumCommands;
 
 	// Create the translation task for this command list.
@@ -910,7 +919,7 @@ void FRHICommandListExecutor::FSubmitState::Dispatch_ProcessCommandList(FRHIComm
 	// Handle adding prerequisites for Parent end waiting on children
 	if (CmdList->ParallelRenderPassEnd.IsValid())
 	{
-		if (GRHIParallelRHIExecuteChildWait)
+		if (GRHIGlobals.ParallelRenderPasses.ParentWaitForChildTranslate)
 		{
 			Prereqs = TranslateJobs.Events;
 				
@@ -921,7 +930,7 @@ void FRHICommandListExecutor::FSubmitState::Dispatch_ProcessCommandList(FRHIComm
 		ChildGraphEvents.Empty();
 		BeginGraphEvent = nullptr;
 	}
-	else if (CmdList->IsSubCommandList() && GRHIParallelRHIExecuteParentWait)
+	else if (CmdList->IsSubCommandList() && GRHIGlobals.ParallelRenderPasses.ChildWaitForParentBegin)
 	{
 		// Children need to wait on the parent begin
 		Prereqs = TranslateJobs.Events;
@@ -1204,9 +1213,9 @@ void FRHICommandListExecutor::FTranslateState::Translate_ExecuteCommandList(FRHI
 			}
 		}
 
-		if(!CmdList->UploadContext)
+		if (!CmdList->UploadContext)
 		{
-			CmdList->UploadContext = UploadContextState;
+			CmdList->UploadContext = UploadContext;
 		}
 	}
 	
@@ -1245,7 +1254,7 @@ void FRHICommandListExecutor::FTranslateState::Translate_ExecuteCommandList(FRHI
 #endif // WITH_RHI_BREADCRUMBS
 	}
 	
-	UploadContextState = CmdList->UploadContext;
+	UploadContext = CmdList->UploadContext;
 
 #if WITH_RHI_BREADCRUMBS
 	BreadcrumbAllocatorRefs.Append(MoveTemp(CmdList->BreadcrumbAllocatorRefs));
@@ -1267,38 +1276,25 @@ void FRHICommandListExecutor::FTranslateState::Translate_CloseChain()
 {
 	SCOPED_NAMED_EVENT(RHI_CloseChain, FColor::Magenta);
 
-	TRHIPipelineArray<IRHIPlatformCommandList*> PlatformCommandLists{ InPlace, nullptr };
-	FDynamicRHI::FRHIFinalizeContextArgs FinalizeArgs;
-	for (auto& State : PipelineStates)
+	FDynamicRHI::FRHICloseTranslateChainArgs CloseChainArgs;
+	CloseChainArgs.UploadContext = UploadContext;
+	CloseChainArgs.bShouldFinalize = bShouldFinalize;
+
+	for (ERHIPipeline Pipeline : MakeFlagsRange(ERHIPipeline::All))
 	{
-		if (State.Context)
+		FPipelineState& State = PipelineStates[Pipeline];
+		CloseChainArgs.Contexts[Pipeline] = State.Context;
+
+		if (bUsingSubCmdLists && Pipeline == ERHIPipeline::Graphics && !GRHIGlobals.ParallelRenderPasses.UsesCustomContexts)
 		{
-			FinalizeArgs.Contexts.Add(State.Context);
+			check(bShouldFinalize);
+			static_cast<IRHICommandContext*>(State.Context)->RHIEndRenderPass();
 		}
 	}
 
 	if (GDynamicRHI)
 	{
-		FinalizeArgs.UploadContext = UploadContextState;
-
-		if (bUsingSubCmdLists)
-		{
-			auto& State = PipelineStates[ERHIPipeline::Graphics];
-			State.FinalizedCmdList = GDynamicRHI->RHIFinalizeParallelContext(State.Context);
-		}
-		else
-		{
-			GDynamicRHI->RHICloseTranslateChain(MoveTemp(FinalizeArgs), PlatformCommandLists, bShouldFinalize);
-
-			for (ERHIPipeline Pipeline : MakeFlagsRange(ERHIPipeline::All))
-			{
-				auto& State = PipelineStates[Pipeline];
-				if (State.Context)
-				{
-					State.FinalizedCmdList = PlatformCommandLists[Pipeline];
-				}
-			}
-		}
+		FinalizedCommandLists = GDynamicRHI->RHICloseTranslateChain(MoveTemp(CloseChainArgs));
 
 		IF_RENDER_TRACING_ENABLED(RenderTracing::FinalizeTranslateJob(this, bParallel, bUsingSubCmdLists));
 	}
@@ -1337,16 +1333,22 @@ void FRHICommandListExecutor::FSubmitState::Submit_SubmitToGPU(FTranslateState*
 			break;
 		}
 
+		// Insert the upload command list before any of the pipeline command lists
+		if (Current->FinalizedCommandLists.UploadCommandList)
+		{
+			CommandLists.Add(Current->FinalizedCommandLists.UploadCommandList);
+		}
+
 		// Gather the completed platform command lists for every RHI pipeline.
 		for (ERHIPipeline Pipeline : MakeFlagsRange(ERHIPipeline::All))
 		{
-			auto& TranslateState = Current->PipelineStates[Pipeline];
-			if (TranslateState.FinalizedCmdList)
+			if (IRHIPlatformCommandList* CommandList = Current->FinalizedCommandLists.CommandLists[Pipeline])
 			{
 			#if WITH_RHI_BREADCRUMBS
-				TranslateState.FinalizedCmdList->BreadcrumbAllocators = Current->BreadcrumbAllocatorRefs;
+				FRHIBreadcrumbRange& BreadcrumbRange = Current->PipelineStates[Pipeline].Range;
+				CommandList->BreadcrumbAllocators = Current->BreadcrumbAllocatorRefs;
 
-				auto& Allocators = TranslateState.FinalizedCmdList->BreadcrumbAllocators;
+				auto& Allocators = CommandList->BreadcrumbAllocators;
 
 				auto& GlobalLast = GRHICommandList.Breadcrumbs.GPU[Pipeline].Last;
 				// Link ranges of breadcrumbs together into depth-first list
@@ -1354,26 +1356,26 @@ void FRHICommandListExecutor::FSubmitState::Submit_SubmitToGPU(FTranslateState*
 				{
 					FRHIBreadcrumbNode*& GlobalNext = GlobalLast->GetNextPtr(Pipeline);
 					check(!GlobalNext);
-					GlobalNext = TranslateState.Range.First;
+					GlobalNext = BreadcrumbRange.First;
 				}
 
 				FRHIBreadcrumbRange Range{};
-				Range.First = GlobalLast ? GlobalLast.Get() : TranslateState.Range.First;
+				Range.First = GlobalLast ? GlobalLast.Get() : BreadcrumbRange.First;
 				if (Range.First) { Allocators.AddUnique(Range.First->Allocator); }
 
 				// Advance the global last breadcrumb forward
-				if (TranslateState.Range.Last)
+				if (BreadcrumbRange.Last)
 				{
-					GlobalLast = TranslateState.Range.Last;
+					GlobalLast = BreadcrumbRange.Last;
 				}
 
 				Range.Last = GlobalLast ? GlobalLast.Get() : Range.First;
 				if (Range.Last) { Allocators.AddUnique(Range.Last->Allocator); }
 
-				TranslateState.FinalizedCmdList->BreadcrumbRange = Range;
+				CommandList->BreadcrumbRange = Range;
 			#endif
 
-				CommandLists.Add(TranslateState.FinalizedCmdList);
+				CommandLists.Add(CommandList);
 			}
 		}
 
diff --git a/Engine/Source/Runtime/RHI/Private/RHIValidation.cpp b/Engine/Source/Runtime/RHI/Private/RHIValidation.cpp
--- a/Engine/Source/Runtime/RHI/Private/RHIValidation.cpp
+++ b/Engine/Source/Runtime/RHI/Private/RHIValidation.cpp
@@ -611,9 +611,23 @@ IRHICommandContext* FValidationRHI::RHIGetDefaultContext()
 
 struct FValidationCommandList : public IRHIPlatformCommandList
 {
-	ERHIPipeline Pipeline;
-	TRHIPipelineArray<IRHIPlatformCommandList*> InnerCommandLists;
+	const ERHIPipeline Pipeline = ERHIPipeline::None;
 	TArray<RHIValidation::FOperation> CompletedOpList;
+	IRHIPlatformCommandList* InnerCommandList = nullptr;
+
+	// Used for upload command lists
+	FValidationCommandList() = default;
+
+	// Used for regular pipeline command lists
+	FValidationCommandList(IRHIComputeContext* OuterContext)
+		: Pipeline       (OuterContext->GetPipeline())
+		, CompletedOpList(OuterContext->Tracker->Finalize())
+	{}
+
+	bool IsUpload() const
+	{
+		return Pipeline == ERHIPipeline::None;
+	}
 };
 
 IRHIComputeContext* FValidationRHI::RHIGetCommandContext(ERHIPipeline Pipeline)
@@ -643,130 +657,94 @@ IRHIComputeContext* FValidationRHI::RHIGetCommandContext(ERHIPipeline Pipeline)
 	}
 }
 
-void FValidationRHI::RHIFinalizeContext(FRHIFinalizeContextArgs&& Args, TRHIPipelineArray<IRHIPlatformCommandList*>& Output)
+IRHICommandContext* FValidationRHI::RHIGetParallelCommandContext(FRHIParallelRenderPassInfo const& ParallelRenderPass)
 {
-	FRHIFinalizeContextArgs FinalArgs;
-	
-	TRHIPipelineArray<IRHIPlatformCommandList*> FinalizedCommandLists { InPlace, nullptr };
-	TRHIPipelineArray<FValidationCommandList*> OuterCommandLists { InPlace, nullptr };
-	
-	// Re-combine the args so that the validation matches a normal call to RHIFinalizeContext
-	for(IRHIComputeContext* Context : Args.Contexts)
-	{
-		IRHIComputeContext& InnerContext = Context->GetLowestLevelContext();
-		
-		FValidationCommandList* OuterCommandList = new FValidationCommandList();
-		
-		// RHIFinalizeContext makes the context available to other threads, so finalize the tracker beforehand.
-		OuterCommandList->CompletedOpList = InnerContext.Tracker->Finalize();
-		OuterCommandList->Pipeline = Context->GetPipeline();
-		OuterCommandLists[OuterCommandList->Pipeline] = OuterCommandList;
-		
-		FinalArgs.Contexts.Add(&InnerContext);
-	}
-	FinalArgs.UploadContext = Args.UploadContext;
-	
-	RHI->RHIFinalizeContext(MoveTemp(FinalArgs), FinalizedCommandLists);
-	
-	for(IRHIComputeContext* Context : Args.Contexts)
-	{
-		FValidationCommandList* ValidationCmdList = OuterCommandLists[Context->GetPipeline()];
-		switch (ValidationCmdList->Pipeline)
-		{
-		case ERHIPipeline::Graphics:
-			if (static_cast<FValidationContext*>(Context)->Type == FValidationContext::EType::Parallel)
-				delete Context;
-			break;
-			
-		case ERHIPipeline::AsyncCompute:
-			if (static_cast<FValidationComputeContext*>(Context)->Type == FValidationComputeContext::EType::Parallel)
-				delete Context;
-			break;
-			
-		default:
-			checkNoEntry();
-			break;
-		}
-		
-		ValidationCmdList->InnerCommandLists = FinalizedCommandLists[ValidationCmdList->Pipeline];
-		Output[ValidationCmdList->Pipeline] = ValidationCmdList;
-	}
-}
+	FValidationContext* OuterContext = new FValidationContext(FValidationContext::EType::Parallel);
+	OuterContext->ValidateBeginRenderPass(ParallelRenderPass, ParallelRenderPass.PassName);
 
-IRHIComputeContext* FValidationRHI::RHIGetParallelCommandContext(FRHIParallelRenderPassInfo const& ParallelRenderPass)
-{
-	// If a platform has a ChildWait or ParentWait it is expected that they will override RHIGetParallelCommandContext
-	// otherwise, we need to manually call RHIGetCommandContext and RHIBeginRenderPass separately, because the default implementation
-	// calls one after another potentially crasing in GetHightestLevelContext (FRHICommandList_RecursiveHazardous)
-	// TODO: Remove this after implementing RHISetupParallelPass
-	if (GRHIParallelRHIExecuteChildWait || GRHIParallelRHIExecuteParentWait)
-	{
-		IRHIComputeContext* InnerContext = RHI->RHIGetParallelCommandContext(ParallelRenderPass);
-		check(InnerContext);
-
-		FValidationContext* OuterContext = new FValidationContext(FValidationContext::EType::Parallel);
-		OuterContext->LinkToContext(static_cast<IRHICommandContext*>(InnerContext));
+	IRHICommandContext* InnerContext = RHI->RHIGetParallelCommandContext(ParallelRenderPass);
+	check(InnerContext);
 
-		// Parallel contexts are always inside a renderpass
-		OuterContext->State.bInsideBeginRenderPass = true;
-		OuterContext->State.RenderPassInfo = ParallelRenderPass;
-		if (ParallelRenderPass.PassName)
-		{
-			OuterContext->State.RenderPassName = ParallelRenderPass.PassName;
-		}
+	OuterContext->LinkToContext(InnerContext);
+	return OuterContext;
+}
 
-		return OuterContext;
-	}
-	else
-	{
-		IRHICommandContext* Context = static_cast<IRHICommandContext*>(RHIGetCommandContext(ERHIPipeline::Graphics));
-		Context->RHIBeginRenderPass(ParallelRenderPass, ParallelRenderPass.PassName);
-		return Context;
-	}
+FDynamicRHI::FRHIFinalizeContextsResult FValidationRHI::RHIFinalizeContexts(FRHIFinalizeContextsArgs&& Args)
+{
+	return RHI->RHIFinalizeContexts(MoveTemp(Args));
 }
 
-void FValidationRHI::RHICloseTranslateChain(FRHIFinalizeContextArgs&& Args, TRHIPipelineArray<IRHIPlatformCommandList*>& Output, bool bShouldFinalize)
+FDynamicRHI::FRHICloseTranslateChainResult FValidationRHI::RHICloseTranslateChain(FRHICloseTranslateChainArgs&& Args)
 {
-	// If we aren't finalizing the context we need to finalize the tracking
-	if (!bShouldFinalize)
+	TRHIPipelineArray<FValidationCommandList*> OuterCommandLists { InPlace, nullptr };
+
+	// Re-combine the args so that the validation matches a normal call to RHIFinalizeContexts
+	for (ERHIPipeline Pipeline : MakeFlagsRange(ERHIPipeline::All))
 	{
-		for (IRHIComputeContext* Context : Args.Contexts)
+		if (IRHIComputeContext* OuterContext = Args.Contexts[Pipeline])
 		{
-			if (Context)
+			// RHIFinalizeContexts makes the context available to other threads, so finalize the tracker beforehand.
+			OuterCommandLists[OuterContext->GetPipeline()] = new FValidationCommandList(OuterContext);
+
+			if (Args.bShouldFinalize)
 			{
-				IRHIComputeContext& InnerContext = Context->GetLowestLevelContext();
-				
-				FValidationCommandList* OuterCommandList = new FValidationCommandList();
-				OuterCommandList->CompletedOpList = InnerContext.Tracker->Finalize();
-				OuterCommandList->Pipeline = Context->GetPipeline();
-				
-				Output[Context->GetPipeline()] = OuterCommandList;
+				// Extract the inner platform RHI context to forward to RHIFinalizeContexts below.
+				Args.Contexts[Pipeline] = &OuterContext->GetLowestLevelContext();
+
+				switch (OuterContext->GetPipeline())
+				{
+				case ERHIPipeline::Graphics:
+					if (static_cast<FValidationContext*>(OuterContext)->Type == FValidationContext::EType::Parallel)
+						delete OuterContext;
+					break;
+
+				case ERHIPipeline::AsyncCompute:
+					if (static_cast<FValidationComputeContext*>(OuterContext)->Type == FValidationComputeContext::EType::Parallel)
+						delete OuterContext;
+					break;
+
+				default:
+					checkNoEntry();
+					break;
+				}
 			}
 		}
-		
-		return;
 	}
 
-	FDynamicRHI::RHICloseTranslateChain(MoveTemp(Args), Output, bShouldFinalize);
-}
-
-IRHIPlatformCommandList* FValidationRHI::RHIFinalizeParallelContext(IRHIComputeContext* Context)
-{
-	check(Context->GetPipeline() == ERHIPipeline::Graphics);
+	FDynamicRHI::FRHICloseTranslateChainResult Result;
+	if (Args.bShouldFinalize)
+	{
+		Result = RHIFinalizeContexts(MoveTemp(Args));
+	}
 
-	IRHIComputeContext& InnerContext = Context->GetLowestLevelContext();
+	// Wrap the returned platform command lists in the validation type.
+	for (ERHIPipeline Pipeline : MakeFlagsRange(ERHIPipeline::All))
+	{
+		IRHIPlatformCommandList*& CommandList = Result.CommandLists[Pipeline];
 
-	FValidationCommandList* OuterCommandList = new FValidationCommandList();
-	OuterCommandList->CompletedOpList = InnerContext.Tracker->Finalize();
-	OuterCommandList->Pipeline = ERHIPipeline::Graphics;
+		if (FValidationCommandList* OuterCommandList = OuterCommandLists[Pipeline])
+		{
+			check(OuterCommandList);
+			OuterCommandList->InnerCommandList = CommandList;
 
-	IRHIPlatformCommandList* InnerCommandList = RHI->RHIFinalizeParallelContext(&InnerContext);
-	OuterCommandList->InnerCommandLists[ERHIPipeline::Graphics] = InnerCommandList;
+			CommandList = OuterCommandList;
+		}
+		else
+		{
+			check(!CommandList);
+		}
+	}
 
-	check(static_cast<FValidationContext*>(Context)->Type == FValidationContext::EType::Parallel);
-	delete Context;
+	if (Result.UploadCommandList)
+	{
+		// There is currently no validation on upload contexts / command lists, but we
+		// still need to wrap the type so that the casts in RHISubmitCommandLists are correct.
+		FValidationCommandList* UploadCommandList = new FValidationCommandList();
+		UploadCommandList->InnerCommandList = Result.UploadCommandList;
+		Result.UploadCommandList = UploadCommandList;
+	}
 
-	return OuterCommandList;
+	return Result;
 }
 
 void FValidationRHI::RHISubmitCommandLists(FRHISubmitCommandListsArgs&& Args)
@@ -779,19 +757,20 @@ void FValidationRHI::RHISubmitCommandLists(FRHISubmitCommandListsArgs&& Args)
 	for (IRHIPlatformCommandList* CmdList : Args.CommandLists)
 	{
 		FValidationCommandList* OuterCommandList = static_cast<FValidationCommandList*>(CmdList);
-#if WITH_RHI_BREADCRUMBS
-		OuterCommandList->CompletedOpList.Insert(RHIValidation::FOperation::SetBreadcrumbRange(CmdList->BreadcrumbRange), 0);
-#endif
 
-		// Replay or queue any barrier operations to validate resource barrier usage.
-		RHIValidation::FTracker::SubmitValidationOps(OuterCommandList->Pipeline, MoveTemp(OuterCommandList->CompletedOpList));
+		// We don't currently do validation on upload command lists.
+		if (!OuterCommandList->IsUpload())
+		{
+		#if WITH_RHI_BREADCRUMBS
+			OuterCommandList->CompletedOpList.Insert(RHIValidation::FOperation::SetBreadcrumbRange(CmdList->BreadcrumbRange), 0);
+		#endif
 
-		for(IRHIPlatformCommandList* InnerCmdList : OuterCommandList->InnerCommandLists)
+			// Replay or queue any barrier operations to validate resource barrier usage.
+			RHIValidation::FTracker::SubmitValidationOps(OuterCommandList->Pipeline, MoveTemp(OuterCommandList->CompletedOpList));
+		}
+
+		if (IRHIPlatformCommandList* InnerCmdList = OuterCommandList->InnerCommandList)
 		{
-			if(!InnerCmdList)
-			{
-				continue;
-			}
 #if WITH_RHI_BREADCRUMBS
 			// Forward the breadcrumb range and allocators
 			InnerCmdList->BreadcrumbAllocators = MoveTemp(CmdList->BreadcrumbAllocators);
diff --git a/Engine/Source/Runtime/RHI/Public/DynamicRHI.h b/Engine/Source/Runtime/RHI/Public/DynamicRHI.h
--- a/Engine/Source/Runtime/RHI/Public/DynamicRHI.h
+++ b/Engine/Source/Runtime/RHI/Public/DynamicRHI.h
@@ -805,7 +805,7 @@ public:
 	//
 	// Retrieves a new command context to begin the recording of a new platform command list.
 	// The returned context is specific to the given pipeline. It can later be converted to an IRHIPlatformCommandList
-	// by calling RHIFinalizeContext(), and then submitted to the GPU by calling RHISubmitCommandLists().
+	// by calling RHIFinalizeContexts(), and then submitted to the GPU by calling RHISubmitCommandLists().
 	//
 	// Called by parallel worker threads, and the render thread. Platform implementations must be thread safe.
 	//
@@ -818,12 +818,10 @@ public:
 	//
 	// Called by parallel worker threads, and the render thread. Platform implementations must be thread safe.
 	//
-	virtual IRHIComputeContext* RHIGetParallelCommandContext(FRHIParallelRenderPassInfo const& ParallelRenderPass)
+	virtual IRHICommandContext* RHIGetParallelCommandContext(FRHIParallelRenderPassInfo const& ParallelRenderPass)
 	{
-		IRHICommandContext* Context = static_cast<IRHICommandContext*>(RHIGetCommandContext(ERHIPipeline::Graphics));
-		Context->RHIBeginRenderPass(ParallelRenderPass, ParallelRenderPass.PassName);
-		
-		return Context;
+		checkNoEntry();
+		return nullptr;
 	}
 	
 	virtual IRHIUploadContext* RHIGetUploadContext()
@@ -839,39 +837,39 @@ public:
 	//
 	// Called by parallel worker threads, and the RHI thread. Platform implementations must be thread safe.
 	//
-	struct FRHIFinalizeContextArgs
+	struct FRHIFinalizeContextsArgs
 	{
-		TArray<IRHIComputeContext*> Contexts;
-		IRHIUploadContext* UploadContext;
+		TRHIPipelineArray<IRHIComputeContext*> Contexts { InPlace, nullptr };
+		IRHIUploadContext* UploadContext = nullptr;
 	};
 	
+	struct FRHIFinalizeContextsResult
+	{
+		TRHIPipelineArray<IRHIPlatformCommandList*> CommandLists { InPlace, nullptr };
+		IRHIPlatformCommandList* UploadCommandList = nullptr;
+	};
+
+	virtual FRHIFinalizeContextsResult RHIFinalizeContexts(FRHIFinalizeContextsArgs&& Args) = 0;
+
 	//
-	// Close the current translate chain or ignore if we are not finalizing
+	// Close the current translate chain, finalizing the specified contexts.
+	// Provided as a hook for the RHI validation layer to handle addition logic for parallel render passes.
 	//
-	#if ENABLE_RHI_VALIDATION
-	virtual
-	#endif
-	void RHICloseTranslateChain(FRHIFinalizeContextArgs&& Args, TRHIPipelineArray<IRHIPlatformCommandList*>& Output, bool bShouldFinalize)
+	struct FRHICloseTranslateChainArgs : FRHIFinalizeContextsArgs
 	{
-		if(bShouldFinalize)
-		{
-			RHIFinalizeContext(MoveTemp(Args), Output);
-		}
-	}
+		bool bShouldFinalize = true;
+	};
 
-	virtual void RHIFinalizeContext(FRHIFinalizeContextArgs&& Args, TRHIPipelineArray<IRHIPlatformCommandList*>& Output) = 0;
-	
-	virtual IRHIPlatformCommandList* RHIFinalizeParallelContext(IRHIComputeContext* Context)
+	using FRHICloseTranslateChainResult = FRHIFinalizeContextsResult;
+
+#if ENABLE_RHI_VALIDATION
+	virtual
+#endif
+	FRHICloseTranslateChainResult RHICloseTranslateChain(FRHICloseTranslateChainArgs&& Args)
 	{
-		FRHIFinalizeContextArgs Args;
-		static_cast<IRHICommandContext*>(Context)->RHIEndRenderPass();
-		
-		Args.Contexts.Add(Context);
-		
-		TRHIPipelineArray<IRHIPlatformCommandList*> Output {nullptr, nullptr};
-		RHIFinalizeContext(MoveTemp(Args), Output);
-		
-		return Output[ERHIPipeline::Graphics];
+		return Args.bShouldFinalize
+			? RHIFinalizeContexts(MoveTemp(Args))
+			: FRHICloseTranslateChainResult();
 	}
 
 	//
diff --git a/Engine/Source/Runtime/RHI/Public/RHICommandList.h b/Engine/Source/Runtime/RHI/Public/RHICommandList.h
--- a/Engine/Source/Runtime/RHI/Public/RHICommandList.h
+++ b/Engine/Source/Runtime/RHI/Public/RHICommandList.h
@@ -5174,11 +5174,12 @@ private:
 			FRHIBreadcrumbRange Range {};
 #endif
 			IRHIComputeContext* Context = nullptr;
-			IRHIPlatformCommandList* FinalizedCmdList = nullptr;
 		};
 		TRHIPipelineArray<FPipelineState> PipelineStates {};
-		IRHIUploadContext* UploadContextState = nullptr;
-		
+
+		IRHIUploadContext* UploadContext = nullptr;
+
+		FDynamicRHI::FRHICloseTranslateChainResult FinalizedCommandLists;
 
 #if WITH_RHI_BREADCRUMBS
 		FRHIBreadcrumbAllocatorArray BreadcrumbAllocatorRefs {};
diff --git a/Engine/Source/Runtime/RHI/Public/RHIGlobals.h b/Engine/Source/Runtime/RHI/Public/RHIGlobals.h
--- a/Engine/Source/Runtime/RHI/Public/RHIGlobals.h
+++ b/Engine/Source/Runtime/RHI/Public/RHIGlobals.h
@@ -545,16 +545,24 @@ struct FRHIGlobals
 	  */
 	bool SupportsConcurrentTranslateAndSubmit = false;
 
-	/** Enables support for platforms that use render passes to split command list translation across multiple threads. For platforms that require
-	   dependencies on Parent/Child command lists i.e. requiring translation of children to be complete before finishing translation of the
-	   parent, should use ParallelRHIExecuteChildWait and/or ParallelRHIExecuteParentWait */
-	bool SupportsParallelRenderPasses = false;
-	
-	/** Used for parallel render pass support. When true, indicates that the platform RHI requires that the child command list translations in a parallel render pass are completed before the remainder of the parent command list can be translated. */
-	bool ParallelRHIExecuteChildWait = false;
-	
-	/** Used for parallel render pass support. When true, indicates that the platform RHI requires that child command list translations in a parallel render pass do not start until the parent command list translation has reached the BeginParallelRenderPass command. */
-	bool ParallelRHIExecuteParentWait = false;
+	struct
+	{
+		/** Enables support for platforms that use render passes to split command list translation across multiple threads. For platforms that require
+		   dependencies on Parent/Child command lists i.e. requiring translation of children to be complete before finishing translation of the
+		   parent, should use ParentWaitForChildTranslate and/or ChildWaitForParentBegin */
+		bool Supported = false;
+
+		/** When true, signals to the renderer that the platform RHI requires a call to RHIGetParallelCommandContext to retrieve RHI contexts for use with parallel render passes.
+		    Otherwise, regular graphics contexts are used. */
+		bool UsesCustomContexts = false;
+
+		/** When true, indicates that the platform RHI requires that the child command list translations in a parallel render pass are completed before the remainder of the parent command list can be translated. */
+		bool ParentWaitForChildTranslate = false;
+
+		/** When true, indicates that the platform RHI requires that child command list translations in a parallel render pass do not start until the parent command list translation has reached the BeginParallelRenderPass command. */
+		bool ChildWaitForParentBegin = false;
+
+	} ParallelRenderPasses;
 
 	/** Whether or not the RHI can perform MSAA sample load. */
 	bool SupportsMSAADepthSampleAccess = false;
@@ -914,9 +922,6 @@ extern RHI_API FRHIGlobals GRHIGlobals;
 #define GRHISupportsRHIThread                                  GRHIGlobals.SupportsRHIThread
 #define GRHISupportsRHIOnTaskThread                            GRHIGlobals.SupportsRHIOnTaskThread
 #define GRHISupportsParallelRHIExecute                         GRHIGlobals.SupportsParallelRHIExecute
-#define GRHISupportsParallelRenderPasses                       GRHIGlobals.SupportsParallelRenderPasses
-#define GRHIParallelRHIExecuteChildWait                        GRHIGlobals.ParallelRHIExecuteChildWait
-#define GRHIParallelRHIExecuteParentWait                       GRHIGlobals.ParallelRHIExecuteParentWait
 #define GRHISupportsParallelRHIExecute                         GRHIGlobals.SupportsParallelRHIExecute
 #define GRHISupportsMSAADepthSampleAccess                      GRHIGlobals.SupportsMSAADepthSampleAccess
 #define GRHISupportsBackBufferWithCustomDepthStencil           GRHIGlobals.SupportsBackBufferWithCustomDepthStencil
diff --git a/Engine/Source/Runtime/RHI/Public/RHIValidation.h b/Engine/Source/Runtime/RHI/Public/RHIValidation.h
--- a/Engine/Source/Runtime/RHI/Public/RHIValidation.h
+++ b/Engine/Source/Runtime/RHI/Public/RHIValidation.h
@@ -894,14 +894,13 @@ public:
 	{
 		return RHI->RHIGetUploadContext();
 	}
-	
-	virtual void RHICloseTranslateChain(FRHIFinalizeContextArgs&& Args, TRHIPipelineArray<IRHIPlatformCommandList*>& Output, bool bShouldFinalize) override final;
-	
-	virtual void RHIFinalizeContext(FRHIFinalizeContextArgs&& Args, TRHIPipelineArray<IRHIPlatformCommandList*>& Output) override final;
+
+	virtual FRHIFinalizeContextsResult RHIFinalizeContexts(FRHIFinalizeContextsArgs&& Args) override final;
+	virtual FRHICloseTranslateChainResult RHICloseTranslateChain(FRHICloseTranslateChainArgs&& Args) override final;
+
 	virtual void RHISubmitCommandLists(FRHISubmitCommandListsArgs&& Args) override final;
 	
-	virtual IRHIComputeContext* RHIGetParallelCommandContext(FRHIParallelRenderPassInfo const& ParallelRenderPass) override final;
-	virtual IRHIPlatformCommandList* RHIFinalizeParallelContext(IRHIComputeContext* Context) override final;
+	virtual IRHICommandContext* RHIGetParallelCommandContext(FRHIParallelRenderPassInfo const& ParallelRenderPass) override final;
 
 	virtual void RHIProcessDeleteQueue() override final
 	{
diff --git a/Engine/Source/Runtime/RHI/Public/RHIValidationContext.h b/Engine/Source/Runtime/RHI/Public/RHIValidationContext.h
--- a/Engine/Source/Runtime/RHI/Public/RHIValidationContext.h
+++ b/Engine/Source/Runtime/RHI/Public/RHIValidationContext.h
@@ -1031,10 +1031,11 @@ public:
 	}
 #endif // WITH_RHI_BREADCRUMBS
 
-	virtual void RHIBeginRenderPass(const FRHIRenderPassInfo& InInfo, const TCHAR* InName) override final
+	void ValidateBeginRenderPass(const FRHIRenderPassInfo& InInfo, const TCHAR* InName)
 	{
 		checkf(!State.bInsideBeginRenderPass, TEXT("Trying to begin RenderPass '%s', but already inside '%s'!"), *State.RenderPassName, InName);
-		checkf(InName!=nullptr, TEXT("RenderPass should have a name!"));
+		checkf(InName != nullptr, TEXT("RenderPass should have a name!"));
+
 		State.bInsideBeginRenderPass = true;
 		State.RenderPassInfo = InInfo;
 		State.RenderPassName = InName;
@@ -1144,7 +1145,11 @@ public:
 			checkf(ShadingRateTexture->GetFormat() == GRHIVariableRateShadingImageFormat, TEXT("Shading rate texture is bound, but is not the correct format for this RHI."));
 			Tracker->Assert(ShadingRateTexture->GetViewIdentity(0, 0, 0, 0, 0, 0), ERHIAccess::ShadingRateSource);
 		}
+	}
 
+	virtual void RHIBeginRenderPass(const FRHIRenderPassInfo& InInfo, const TCHAR* InName) override final
+	{
+		ValidateBeginRenderPass(InInfo, InName);
 		RHIContext->RHIBeginRenderPass(InInfo, InName);
 	}
 
diff --git a/Engine/Source/Runtime/RenderCore/Private/RenderGraphPass.cpp b/Engine/Source/Runtime/RenderCore/Private/RenderGraphPass.cpp
--- a/Engine/Source/Runtime/RenderCore/Private/RenderGraphPass.cpp
+++ b/Engine/Source/Runtime/RenderCore/Private/RenderGraphPass.cpp
@@ -120,7 +120,7 @@ FRHICommandList* FRDGDispatchPassBuilder::CreateCommandList()
 {
 	FRHICommandList* RHICmdList = nullptr;
 	
-	if (RenderPassInfo && Pass->bParallelExecute && GRHISupportsParallelRenderPasses)
+	if (RenderPassInfo && Pass->bParallelExecute && GRHIGlobals.ParallelRenderPasses.Supported)
 	{
 		FRHISubCommandList* SubCmdList = new FRHISubCommandList(Pass->GetGPUMask(), RenderPassInfo);
 		SubCmdList->SwitchPipeline(Pass->GetPipeline());
diff --git a/Engine/Source/Runtime/VulkanRHI/Private/VulkanDynamicRHI.h b/Engine/Source/Runtime/VulkanRHI/Private/VulkanDynamicRHI.h
--- a/Engine/Source/Runtime/VulkanRHI/Private/VulkanDynamicRHI.h
+++ b/Engine/Source/Runtime/VulkanRHI/Private/VulkanDynamicRHI.h
@@ -230,12 +230,11 @@ public:
 	virtual class IRHICommandContext* RHIGetDefaultContext() final override;
 	virtual IRHIComputeContext* RHIGetCommandContext(ERHIPipeline Pipeline) final override;
 	virtual IRHIUploadContext* RHIGetUploadContext() final override;
-	virtual void RHIFinalizeContext(FRHIFinalizeContextArgs&& Args, TRHIPipelineArray<IRHIPlatformCommandList*>& Output) final override;
+	virtual FRHIFinalizeContextsResult RHIFinalizeContexts(FRHIFinalizeContextsArgs&& Args) final override;
 	virtual void RHISubmitCommandLists(FRHISubmitCommandListsArgs&& Args) final override;
 	virtual uint64 RHIGetMinimumAlignmentForBufferBackedSRV(EPixelFormat Format) final override;
 
-	virtual IRHIComputeContext* RHIGetParallelCommandContext(FRHIParallelRenderPassInfo const& ParallelRenderPass) final override;
-	virtual IRHIPlatformCommandList* RHIFinalizeParallelContext(IRHIComputeContext* Context) final override;
+	virtual IRHICommandContext* RHIGetParallelCommandContext(FRHIParallelRenderPassInfo const& ParallelRenderPass) final override;
 
 	// Transient Resource Functions and Helpers
 	virtual IRHITransientResourceAllocator* RHICreateTransientResourceAllocator() final override;
diff --git a/Engine/Source/Runtime/VulkanRHI/Private/VulkanRHI.cpp b/Engine/Source/Runtime/VulkanRHI/Private/VulkanRHI.cpp
--- a/Engine/Source/Runtime/VulkanRHI/Private/VulkanRHI.cpp
+++ b/Engine/Source/Runtime/VulkanRHI/Private/VulkanRHI.cpp
@@ -947,9 +947,10 @@ void FVulkanDynamicRHI::InitInstance()
 		GRHISupportsParallelRHIExecute = Device->SupportsParallelRendering() && (GRHIThreadCvar->GetInt() > 1);
 #endif
 
-		GRHISupportsParallelRenderPasses = GRHISupportsParallelRHIExecute;
-		GRHIParallelRHIExecuteChildWait = GRHISupportsParallelRHIExecute;
-		GRHIParallelRHIExecuteParentWait = GRHISupportsParallelRHIExecute;
+		GRHIGlobals.ParallelRenderPasses.Supported                   = GRHISupportsParallelRHIExecute;
+		GRHIGlobals.ParallelRenderPasses.ParentWaitForChildTranslate = GRHISupportsParallelRHIExecute;
+		GRHIGlobals.ParallelRenderPasses.ChildWaitForParentBegin     = GRHISupportsParallelRHIExecute;
+		GRHIGlobals.ParallelRenderPasses.UsesCustomContexts = true;
 
 		GRHISupportsUAVFormatAliasing = true;
 
diff --git a/Engine/Source/Runtime/VulkanRHI/Private/VulkanSubmission.cpp b/Engine/Source/Runtime/VulkanRHI/Private/VulkanSubmission.cpp
--- a/Engine/Source/Runtime/VulkanRHI/Private/VulkanSubmission.cpp
+++ b/Engine/Source/Runtime/VulkanRHI/Private/VulkanSubmission.cpp
@@ -242,50 +242,66 @@ void FVulkanDynamicRHI::KickSubmissionThread(bool bForce)
 	}
 }
 
-void FVulkanDynamicRHI::RHIFinalizeContext(FRHIFinalizeContextArgs&& Args, TRHIPipelineArray<IRHIPlatformCommandList*>& Output)
+static FCriticalSection GSecondaryPayloadsCS;
+FDynamicRHI::FRHIFinalizeContextsResult FVulkanDynamicRHI::RHIFinalizeContexts(FRHIFinalizeContextsArgs&& Args)
 {
-	auto GetPlatformCommandList = [&Output](ERHIPipeline Pipeline)
-	{
-		if (!Output[Pipeline])
-		{
-			Output[Pipeline] = new FVulkanPlatformCommandList();
-		}
-		return ResourceCast(Output[Pipeline]);
-	};
+	FRHIFinalizeContextsResult Result;
 
-	// :todo-jn: Place uploads first on the graphics pipe for now
 	if (Args.UploadContext)
 	{
-		FVulkanPlatformCommandList& PlatformCmdList = *GetPlatformCommandList(ERHIPipeline::Graphics);
+		FVulkanPlatformCommandList* PlatformCmdList = new FVulkanPlatformCommandList;
+		Result.UploadCommandList = PlatformCmdList;
+
 		FVulkanUploadContext* UploadContext = ResourceCast(Args.UploadContext);
-		UploadContext->Finalize(PlatformCmdList);
+		UploadContext->Finalize(*PlatformCmdList);
+
 		FVulkanUploadContext::Pool.Push(UploadContext);
-		Args.UploadContext = nullptr;
 	}
 
-	for (IRHIComputeContext* Context : Args.Contexts)
+	for (ERHIPipeline Pipeline : MakeFlagsRange(ERHIPipeline::All))
 	{
-		FVulkanPlatformCommandList& PlatformCmdList = *GetPlatformCommandList(Context->GetPipeline());
+		if (FVulkanCommandListContext* Context = static_cast<FVulkanCommandListContext*>(Args.Contexts[Pipeline]))
+		{
+			if (FVulkanParallelRenderPassInfo* ParallelRPInfo = Context->GetParallelRenderPassInfo())
+			{
+				// Handle finalizing of parallel render pass contexts.
+				if (Device->GetOptionalExtensions().HasKHRDynamicRendering)
+				{
+					Context->GetCommandBuffer().EndDynamicRendering();
+				}
 
-		FVulkanCommandListContext* VulkanContext = static_cast<FVulkanCommandListContext*>(Context);
-		VulkanContext->Finalize(PlatformCmdList);
+				FScopeLock Lock(&GSecondaryPayloadsCS);
+				Context->Finalize(ParallelRPInfo->SecondaryPayloads);
+			}
+			else
+			{
+				FVulkanPlatformCommandList* PlatformCmdList = new FVulkanPlatformCommandList;
+				Result.CommandLists[Pipeline] = PlatformCmdList;
 
-		if (VulkanContext->IsImmediate())
-		{
-			// :todo-jn: clean up immediate context for next use?
-		}
-		else
-		{
-			delete VulkanContext;
+				Context->Finalize(*PlatformCmdList);
+			}
+
+			if (Context->IsImmediate())
+			{
+				// :todo-jn: clean up immediate context for next use?
+			}
+			else
+			{
+				delete Context;
+			}
 		}
 	}
+
+	return Result;
 }
 
-IRHIComputeContext* FVulkanDynamicRHI::RHIGetParallelCommandContext(FRHIParallelRenderPassInfo const& ParallelRenderPass)
+IRHICommandContext* FVulkanDynamicRHI::RHIGetParallelCommandContext(FRHIParallelRenderPassInfo const& ParallelRenderPass)
 {
-	FVulkanParallelRenderPassInfo* VulkanParallelRenderPassInfo = (FVulkanParallelRenderPassInfo*)ParallelRenderPass.RHIPlatformData;
+	FVulkanParallelRenderPassInfo* VulkanParallelRenderPassInfo = static_cast<FVulkanParallelRenderPassInfo*>(ParallelRenderPass.RHIPlatformData);
 	checkf(VulkanParallelRenderPassInfo, TEXT("Must begin parallel render pass on parent context before creating parallel contexts!"));
+
 	FVulkanCommandListContext* ParallelCommandContext = new FVulkanCommandListContext(*Device, VulkanParallelRenderPassInfo);
+
 	ParallelCommandContext->RenderPassInfo = ParallelRenderPass;
 
 	if (Device->GetOptionalExtensions().HasKHRDynamicRendering)
@@ -309,28 +325,6 @@ IRHIComputeContext* FVulkanDynamicRHI::RHIGetParallelCommandContext(FRHIParallel
 	return ParallelCommandContext;
 }
 
-static FCriticalSection GSecondaryPayloadsCS;
-IRHIPlatformCommandList* FVulkanDynamicRHI::RHIFinalizeParallelContext(IRHIComputeContext* Context)
-{
-	FVulkanCommandListContext* VulkanContext = static_cast<FVulkanCommandListContext*>(Context);
-
-	if (Device->GetOptionalExtensions().HasKHRDynamicRendering)
-	{
-		VulkanContext->GetCommandBuffer().EndDynamicRendering();
-	}
-
-	// Call Finalize under this lock
-	{
-		FScopeLock Lock(&GSecondaryPayloadsCS);
-		FVulkanParallelRenderPassInfo* VulkanParallelRenderPassInfo = VulkanContext->GetParallelRenderPassInfo();
-		checkf(VulkanParallelRenderPassInfo, TEXT("Attempting to call RHIFinalizeParallelContext on a non-parallel context!"));
-		VulkanContext->Finalize(VulkanParallelRenderPassInfo->SecondaryPayloads);
-	}
-
-	delete VulkanContext;
-	return nullptr;
-}
-
 bool FVulkanDynamicRHI::ProcessSubmissionQueue()
 {
 	SCOPED_NAMED_EVENT_TEXT("ProcessSubmissionQueue", FColor::Orange);
diff --git a/Engine/Source/Runtime/Windows/D3D11RHI/Private/D3D11Commands.cpp b/Engine/Source/Runtime/Windows/D3D11RHI/Private/D3D11Commands.cpp
--- a/Engine/Source/Runtime/Windows/D3D11RHI/Private/D3D11Commands.cpp
+++ b/Engine/Source/Runtime/Windows/D3D11RHI/Private/D3D11Commands.cpp
@@ -1366,20 +1366,28 @@ struct FD3D11PlatformCommandList : public IRHIPlatformCommandList
 	virtual ~FD3D11PlatformCommandList() = default;
 };
 
-void FD3D11DynamicRHI::RHIFinalizeContext(FRHIFinalizeContextArgs&& Args, TRHIPipelineArray<IRHIPlatformCommandList*>& Output)
+FDynamicRHI::FRHIFinalizeContextsResult FD3D11DynamicRHI::RHIFinalizeContexts(FRHIFinalizeContextsArgs&& Args)
 {
+	checkf(Args.UploadContext == nullptr, TEXT("D3D11 RHI does not implement upload contexts."));
+
 	FlushProfilerStats();
 
-	// "Context" will always be the default context, since we don't implement parallel execution.
-	for (IRHIComputeContext* Context : Args.Contexts)
+	FRHIFinalizeContextsResult Result;
+
+	for (ERHIPipeline Pipeline : MakeFlagsRange(ERHIPipeline::All))
 	{
-		// "Context" will always be the default context, since we don't implement parallel execution.
-		check(Context == this);
+		checkf(Pipeline == ERHIPipeline::Graphics || !Args.Contexts[Pipeline], TEXT("Only graphics pipeline contexts are supported in D3D11 RHI."));
 
-#if WITH_RHI_BREADCRUMBS
-		// We need platform command lists to contain the breadcrumb allocators
-		Output[Context->GetPipeline()] = new FD3D11PlatformCommandList;
-#endif
+		if (IRHIComputeContext* Context = Args.Contexts[Pipeline])
+		{
+			// "Context" will always be the default context, since we don't implement parallel execution.
+			check(Context == this);
+
+		#if WITH_RHI_BREADCRUMBS
+			// We need platform command lists to contain the breadcrumb allocators
+			Result.CommandLists[Pipeline] = new FD3D11PlatformCommandList;
+		#endif
+		}
 	}
 
 	// Reset some context state
@@ -1392,6 +1400,8 @@ void FD3D11DynamicRHI::RHIFinalizeContext(FRHIFinalizeContextArgs&& Args, TRHIPi
 			BoundUniformBuffers[Frequency][BindIndex] = nullptr;
 		}
 	}
+
+	return Result;
 }
 
 void FD3D11DynamicRHI::RHISubmitCommandLists(FRHISubmitCommandListsArgs&& Args)
diff --git a/Engine/Source/Runtime/Windows/D3D11RHI/Private/D3D11RHIPrivate.h b/Engine/Source/Runtime/Windows/D3D11RHI/Private/D3D11RHIPrivate.h
--- a/Engine/Source/Runtime/Windows/D3D11RHI/Private/D3D11RHIPrivate.h
+++ b/Engine/Source/Runtime/Windows/D3D11RHI/Private/D3D11RHIPrivate.h
@@ -304,7 +304,7 @@ public:
 	D3D11RHI_API virtual void* RHIGetNativeCommandBuffer() final override;
 	D3D11RHI_API virtual class IRHICommandContext* RHIGetDefaultContext() final override;
 	D3D11RHI_API virtual IRHIComputeContext* RHIGetCommandContext(ERHIPipeline Pipeline) final override;
-	D3D11RHI_API virtual void RHIFinalizeContext(FRHIFinalizeContextArgs&& Args, TRHIPipelineArray<IRHIPlatformCommandList*>& Output) final override;
+	D3D11RHI_API virtual FRHIFinalizeContextsResult RHIFinalizeContexts(FRHIFinalizeContextsArgs&& Args) final override;
 	D3D11RHI_API virtual void RHISubmitCommandLists(FRHISubmitCommandListsArgs&& Args) final override;
 
 	// SRV / UAV creation functions
