[RHI] Fix double AddRef/Release of pipeline state blend/rasterizer/depth-stencil states

Remove the duplicate AddRef block from the constructor and the duplicate Release block from CompilePSOInternal().

#rnx
#Synced-CL 52440464

diff --git a/Engine/Source/Runtime/RHI/Private/PipelineStateCache.cpp b/Engine/Source/Runtime/RHI/Private/PipelineStateCache.cpp
--- a/Engine/Source/Runtime/RHI/Private/PipelineStateCache.cpp
+++ b/Engine/Source/Runtime/RHI/Private/PipelineStateCache.cpp
@@ -3869,19 +3869,6 @@ public:
 		{
 			Initializer.DepthStencilState->AddRef();
 		}
-
-		if (Initializer.BlendState)
-		{
-			Initializer.BlendState->AddRef();
-		}
-		if (Initializer.RasterizerState)
-		{
-			Initializer.RasterizerState->AddRef();
-		}
-		if (Initializer.DepthStencilState)
-		{
-			Initializer.DepthStencilState->AddRef();
-		}
 	}
 
 	void CompilePSOInternal()
@@ -4003,19 +3990,6 @@ public:
 		{
 			Initializer.DepthStencilState->Release();
 		}
-
-		if (Initializer.BlendState)
-		{
-			Initializer.BlendState->Release();
-		}
-		if (Initializer.RasterizerState)
-		{
-			Initializer.RasterizerState->Release();
-		}
-		if (Initializer.DepthStencilState)
-		{
-			Initializer.DepthStencilState->Release();
-		}
 	}
 };
 
