VulkanRHI:  Add PLATFORM_SUPPORTS_SPECIALIZATION_CONSTANTS to be able to exclude spec const members (and some functions).  Always load excluded resources since they are used to skip resources even in shipping code paths (pull it out of RHI_INCLUDE_SHADER_DEBUG_DATA define).  Fix spec consts id.  Turn LightGridInjection spec consts back on.

#rnx

#Synced-CL 52508799

diff --git a/Engine/Source/Runtime/Core/Public/Android/AndroidPlatform.h b/Engine/Source/Runtime/Core/Public/Android/AndroidPlatform.h
--- a/Engine/Source/Runtime/Core/Public/Android/AndroidPlatform.h
+++ b/Engine/Source/Runtime/Core/Public/Android/AndroidPlatform.h
@@ -60,6 +60,7 @@ typedef FAndroidTypes FPlatformTypes;
 #define PLATFORM_UI_NEEDS_FOCUS_OUTLINES			0
 #define PLATFORM_SUPPORTS_EARLY_MOVIE_PLAYBACK		1 // movies will start before engine is initialized
 #define PLATFORM_SUPPORTS_GEOMETRY_SHADERS			(RHI_RAYTRACING) // Currently required due to the way we offset ShaderStage::EStage
+#define PLATFORM_SUPPORTS_SPECIALIZATION_CONSTANTS	1
 #define PLATFORM_USE_ANSI_POSIX_MALLOC				1
 #define PLATFORM_USE_MINIMAL_HANG_DETECTION			1
 #define PLATFORM_CODE_SECTION(Name)					__attribute__((section(Name)))
diff --git a/Engine/Source/Runtime/Core/Public/HAL/Platform.h b/Engine/Source/Runtime/Core/Public/HAL/Platform.h
--- a/Engine/Source/Runtime/Core/Public/HAL/Platform.h
+++ b/Engine/Source/Runtime/Core/Public/HAL/Platform.h
@@ -432,6 +432,10 @@
 	#define PLATFORM_SUPPORTS_BINDLESS_RENDERING 0
 #endif
 
+#ifndef PLATFORM_SUPPORTS_SPECIALIZATION_CONSTANTS 
+	#define PLATFORM_SUPPORTS_SPECIALIZATION_CONSTANTS 0
+#endif
+
 #ifndef PLATFORM_BUILTIN_VERTEX_HALF_FLOAT
 	#define PLATFORM_BUILTIN_VERTEX_HALF_FLOAT	1
 #endif
diff --git a/Engine/Source/Runtime/Core/Public/Unix/UnixPlatform.h b/Engine/Source/Runtime/Core/Public/Unix/UnixPlatform.h
--- a/Engine/Source/Runtime/Core/Public/Unix/UnixPlatform.h
+++ b/Engine/Source/Runtime/Core/Public/Unix/UnixPlatform.h
@@ -61,6 +61,7 @@ typedef FUnixPlatformTypes FPlatformTypes;
 #define PLATFORM_ALLOW_ALLOCATIONS_IN_FASYNCWRITER_SERIALIZEBUFFERTOARCHIVE 0
 #define PLATFORM_USE_ANSI_POSIX_MALLOC					1
 #define PLATFORM_SUPPORTS_BINDLESS_RENDERING			1
+#define PLATFORM_SUPPORTS_SPECIALIZATION_CONSTANTS		1
 #define PLATFORM_SUPPORTS_MESH_SHADERS					1
 
 
diff --git a/Engine/Source/Runtime/Core/Public/Windows/WindowsPlatform.h b/Engine/Source/Runtime/Core/Public/Windows/WindowsPlatform.h
--- a/Engine/Source/Runtime/Core/Public/Windows/WindowsPlatform.h
+++ b/Engine/Source/Runtime/Core/Public/Windows/WindowsPlatform.h
@@ -99,6 +99,7 @@ typedef FWindowsPlatformTypes FPlatformTypes;
 #define PLATFORM_SUPPORTS_MESH_SHADERS						1
 #define PLATFORM_SUPPORTS_WORKGRAPH_SHADERS					1
 #define PLATFORM_SUPPORTS_BINDLESS_RENDERING				1
+#define PLATFORM_SUPPORTS_SPECIALIZATION_CONSTANTS			1
 #define PLATFORM_USES_DYNAMIC_HDR_SETTING					1
 #define PLATFORM_USES__ALIGNED_MALLOC						1
 
diff --git a/Engine/Source/Runtime/RHI/Public/RHIResources.h b/Engine/Source/Runtime/RHI/Public/RHIResources.h
--- a/Engine/Source/Runtime/RHI/Public/RHIResources.h
+++ b/Engine/Source/Runtime/RHI/Public/RHIResources.h
@@ -875,25 +875,11 @@ public:
 	TArray<FShaderCodeValidationType> DebugUAVTypeValidationData;
 	TArray<FShaderCodeValidationUBSize> DebugUBSizeValidationData;
 
-	TMap<uint32, TArray<uint32>> DebugSpecializationExclusions;
-	bool IsExcludedFromSpecialization(uint32 BindPoint, uint32 SpecializationHash) const
-	{
-		if (bHasSpecializationConstants)
-		{
-			if (const TArray<uint32>* Array = DebugSpecializationExclusions.Find(BindPoint))
-			{
-				return Array->Find(SpecializationHash) != INDEX_NONE;
-			}
-		}
-		return false;
-	}
-
 #else
 
 	bool HasShaderName() const { return false; }
 	const TCHAR* GetShaderName() const { return TEXT("<unknown>"); }
 	FString GetUniformBufferName(uint32 Index) const { return TEXT("<unknown>"); }
-	bool IsExcludedFromSpecialization(uint32 BindPoint, uint32 SpecializationHash) const { return false; }
 
 #endif
 
@@ -903,7 +889,9 @@ public:
 		, Frequency(InFrequency)
 		, bNoDerivativeOps(false)
 		, bHasShaderBundleUsage(false)
+#if PLATFORM_SUPPORTS_SPECIALIZATION_CONSTANTS
 		, bHasSpecializationConstants(false)
+#endif
 	{
 	}
 
@@ -937,6 +925,7 @@ public:
 		return bHasShaderBundleUsage;
 	}
 
+#if PLATFORM_SUPPORTS_SPECIALIZATION_CONSTANTS
 	inline void SetSpecializationConstants(bool bValue)
 	{
 		bHasSpecializationConstants = bValue;
@@ -947,12 +936,42 @@ public:
 		return bHasSpecializationConstants;
 	}
 
+	const TMap<uint32, TArray<uint32>>& GetSpecializationExclusions() const
+	{
+		return SpecializationExclusions;
+	}
+
+	void SetSpecializationExclusions(TMap<uint32, TArray<uint32>>&& Exclusions)
+	{
+		SpecializationExclusions = Exclusions;
+	}
+
+	bool IsExcludedFromSpecialization(uint32 BindPoint, uint32 SpecializationHash) const
+	{
+		if (bHasSpecializationConstants)
+		{
+			if (const TArray<uint32>* Array = SpecializationExclusions.Find(BindPoint))
+			{
+				return Array->Find(SpecializationHash) != INDEX_NONE;
+			}
+		}
+		return false;
+	}
+#else
+	constexpr bool HasSpecializationConstants() const { return false; }
+#endif // PLATFORM_SUPPORTS_SPECIALIZATION_CONSTANTS
+
 private:
+#if PLATFORM_SUPPORTS_SPECIALIZATION_CONSTANTS
+	TMap<uint32, TArray<uint32>> SpecializationExclusions;
+#endif // PLATFORM_SUPPORTS_SPECIALIZATION_CONSTANTS
 	FShaderHash Hash;
 	EShaderFrequency Frequency;
 	uint8 bNoDerivativeOps : 1;
 	uint8 bHasShaderBundleUsage : 1;
+#if PLATFORM_SUPPORTS_SPECIALIZATION_CONSTANTS
 	uint8 bHasSpecializationConstants : 1;
+#endif // PLATFORM_SUPPORTS_SPECIALIZATION_CONSTANTS
 #if PLATFORM_WINDOWS
 	volatile int16 InUseByPSOCompilation = 0;
 #endif // PLATFORM_WINDOWS
diff --git a/Engine/Source/Runtime/RHICore/Private/RHICoreShader.cpp b/Engine/Source/Runtime/RHICore/Private/RHICoreShader.cpp
--- a/Engine/Source/Runtime/RHICore/Private/RHICoreShader.cpp
+++ b/Engine/Source/Runtime/RHICore/Private/RHICoreShader.cpp
@@ -76,7 +76,12 @@ void UE::RHICore::SetupShaderCodeValidationData(FRHIShader* RHIShader, FShaderCo
 			RHIShader->DebugUBSizeValidationData.Append(ShaderCodeValidationExtension.ShaderCodeValidationUBSize);
 		}
 	}
+#endif
+}
 
+void UE::RHICore::SetupSpecializationConstantData(FRHIShader* RHIShader, FShaderCodeReader& ShaderCodeReader)
+{
+#if PLATFORM_SUPPORTS_SPECIALIZATION_CONSTANTS
 	if (RHIShader)
 	{
 		int32 SpecializationConstantValidationSize = 0;
@@ -86,7 +91,7 @@ void UE::RHICore::SetupShaderCodeValidationData(FRHIShader* RHIShader, FShaderCo
 			FBufferReader ArValidationExtensionData((void*)SpecializationConstantValidationData, SpecializationConstantValidationSize, false);
 			FSpecializationConstantData SpecializationConstantData;
 			ArValidationExtensionData << SpecializationConstantData;
-			RHIShader->DebugSpecializationExclusions = SpecializationConstantData.SpecializationConstantExclusions;
+			RHIShader->SetSpecializationExclusions(MoveTemp(SpecializationConstantData.SpecializationConstantExclusions));
 		}
 	}
 #endif
diff --git a/Engine/Source/Runtime/RHICore/Public/RHICoreShader.h b/Engine/Source/Runtime/RHICore/Public/RHICoreShader.h
--- a/Engine/Source/Runtime/RHICore/Public/RHICoreShader.h
+++ b/Engine/Source/Runtime/RHICore/Public/RHICoreShader.h
@@ -18,6 +18,7 @@ namespace RHICore
 /** Validates that the uniform buffer at the requested static slot. */
 extern RHICORE_API void ValidateStaticUniformBuffer(FRHIUniformBuffer* UniformBuffer, FUniformBufferStaticSlot Slot, uint32 ExpectedHash);
 extern RHICORE_API void SetupShaderCodeValidationData(FRHIShader* RHIShader, class FShaderCodeReader& ShaderCodeReader);
+extern RHICORE_API void SetupSpecializationConstantData(FRHIShader* RHIShader, class FShaderCodeReader& ShaderCodeReader);
 extern RHICORE_API void SetupShaderDiagnosticData(FRHIShader* RHIShader, class FShaderCodeReader& ShaderCodeReader);
 extern RHICORE_API void RegisterDiagnosticMessages(const TArray<FShaderDiagnosticData>& In);
 extern RHICORE_API const FString* GetDiagnosticMessage(uint32 MessageID);
diff --git a/Engine/Source/Runtime/RenderCore/Private/ShaderParameterStruct.cpp b/Engine/Source/Runtime/RenderCore/Private/ShaderParameterStruct.cpp
--- a/Engine/Source/Runtime/RenderCore/Private/ShaderParameterStruct.cpp
+++ b/Engine/Source/Runtime/RenderCore/Private/ShaderParameterStruct.cpp
@@ -621,9 +621,9 @@ void ValidateShaderParameterResourcesRHI(const void* Contents, const FRHIUniform
 
 #endif // DO_CHECK
 
-#if RHI_INCLUDE_SHADER_DEBUG_DATA
 FResourceExclusionArray GetSpecializationResourceExclusions(const TShaderRef<FShader>& Shader)
 {
+#if PLATFORM_SUPPORTS_SPECIALIZATION_CONSTANTS
 	FRHIShader* RHIShader = Shader.GetRHIShaderBase(Shader->GetFrequency());
 	if (RHIShader->HasSpecializationConstants())
 	{
@@ -651,10 +651,9 @@ FResourceExclusionArray GetSpecializationResourceExclusions(const TShaderRef<FSh
 
 		return ResourceExclusions;
 	}
-
+#endif // PLATFORM_SUPPORTS_SPECIALIZATION_CONSTANTS
 	return FResourceExclusionArray();
 }
-#endif // RHI_INCLUDE_SHADER_DEBUG_DATA
 
 template<typename BindingParameterType>
 FRHIShaderParameterResource ExtractShaderParameterResource(FShaderParameterReader Reader, const BindingParameterType& Parameter)
diff --git a/Engine/Source/Runtime/RenderCore/Public/GlobalShader.h b/Engine/Source/Runtime/RenderCore/Public/GlobalShader.h
--- a/Engine/Source/Runtime/RenderCore/Public/GlobalShader.h
+++ b/Engine/Source/Runtime/RenderCore/Public/GlobalShader.h
@@ -211,7 +211,7 @@ public:
 	template<typename ShaderType>
 	TShaderRef<ShaderType> GetShader(const typename ShaderType::FPermutationDomain& PermutationVector) const
 	{
-		return GetShader<ShaderType>(UE::ShaderPermutationUtils::ConvertToUnspecializedId(GetPlatform(), PermutationVector));
+		return GetShader<ShaderType>(PermutationVector.ToDimensionValueId());
 	}
 
 	bool HasShader(FShaderType* Type, int32 PermutationId) const
diff --git a/Engine/Source/Runtime/RenderCore/Public/ShaderParameterStruct.h b/Engine/Source/Runtime/RenderCore/Public/ShaderParameterStruct.h
--- a/Engine/Source/Runtime/RenderCore/Public/ShaderParameterStruct.h
+++ b/Engine/Source/Runtime/RenderCore/Public/ShaderParameterStruct.h
@@ -121,11 +121,7 @@ extern RENDERCORE_API void EmitNullShaderParameterFatalError(const TShaderRef<FS
 /** Generate a list of excluded resources from the specialization constants. */
 using FResourceExclusionArray = TArray<uint16>;
 using FResourceExclusionArrayView = TArrayView<uint16>;
-#if RHI_INCLUDE_SHADER_DEBUG_DATA
 extern RENDERCORE_API FResourceExclusionArray GetSpecializationResourceExclusions(const TShaderRef<FShader>& Shader);
-#else
-inline FResourceExclusionArray GetSpecializationResourceExclusions(const TShaderRef<FShader>& Shader) { return FResourceExclusionArray();  }
-#endif
 
 /** Validates that all resource parameters of a shader are set. */
 #if DO_CHECK
diff --git a/Engine/Source/Runtime/Renderer/Private/LightGridInjection.cpp b/Engine/Source/Runtime/Renderer/Private/LightGridInjection.cpp
--- a/Engine/Source/Runtime/Renderer/Private/LightGridInjection.cpp
+++ b/Engine/Source/Runtime/Renderer/Private/LightGridInjection.cpp
@@ -190,9 +190,9 @@ FAutoConsoleVariableRef CVarLightCullingMaxDistanceOverride(
 // Experimental version of FLightGridInjectionCS that uses specialization constants
 static TAutoConsoleVariable<int32> CVarLightGridInjectionUseSpecializationConstants(
 	TEXT("r.LightGridInjection.UseSpecializationConstants"),
-	0,
-	TEXT("Experimental: Selects whether LightGridInjection should use specialization constant. (default: 0)"),
-	ECVF_ReadOnly);
+	1,
+	TEXT("Experimental: Selects whether LightGridInjection should use specialization constant. (default: 1)"),
+	ECVF_RenderThreadSafe);
 
 static inline bool UseSpecializationConstants(EShaderPlatform ShaderPlatform)
 {
diff --git a/Engine/Source/Runtime/VulkanRHI/Private/VulkanPipelineState.cpp b/Engine/Source/Runtime/VulkanRHI/Private/VulkanPipelineState.cpp
--- a/Engine/Source/Runtime/VulkanRHI/Private/VulkanPipelineState.cpp
+++ b/Engine/Source/Runtime/VulkanRHI/Private/VulkanPipelineState.cpp
@@ -417,13 +417,12 @@ TArray<uint32> FVulkanComputePipelineDescriptorState::GetExcludedResources(uint3
 	check(SetIndex == 0);
 
 	TArray<uint32> ExcludedResources;
-#if RHI_INCLUDE_SHADER_DEBUG_DATA
 	FRHIShader* RHIShader = ComputePipeline->GetComputeShader();
 	if (RHIShader && RHIShader->HasSpecializationConstants())
 	{
 		const uint32 SpecializationHash = GetTypeHash(ComputePipeline->GetSpecializationConstants());
 
-		for (auto& Pair : RHIShader->DebugSpecializationExclusions)
+		for (auto& Pair : RHIShader->GetSpecializationExclusions())
 		{
 			if (Pair.Value.Contains(SpecializationHash))
 			{
@@ -431,7 +430,6 @@ TArray<uint32> FVulkanComputePipelineDescriptorState::GetExcludedResources(uint3
 			}
 		}
 	}
-#endif
 	return ExcludedResources;
 }
 
@@ -766,13 +764,12 @@ TArray<uint32> FVulkanGraphicsPipelineDescriptorState::GetExcludedResources(uint
 	check(SetIndex < ShaderStage::NumGraphicsStages);
 
 	TArray<uint32> ExcludedResources;
-#if RHI_INCLUDE_SHADER_DEBUG_DATA
 	FRHIShader* RHIShader = GfxPipeline->GetShader(GetFrequencyForGfxStage((ShaderStage::EStage)SetIndex));
 	if (RHIShader && RHIShader->HasSpecializationConstants())
 	{
 		const uint32 SpecializationHash = GetTypeHash(GfxPipeline->SpecializationConstants[SetIndex]);
 
-		for (auto& Pair : RHIShader->DebugSpecializationExclusions)
+		for (auto& Pair : RHIShader->GetSpecializationExclusions())
 		{
 			if (Pair.Value.Contains(SpecializationHash))
 			{
@@ -780,7 +777,6 @@ TArray<uint32> FVulkanGraphicsPipelineDescriptorState::GetExcludedResources(uint
 			}
 		}
 	}
-#endif
 	return ExcludedResources;
 }
 
diff --git a/Engine/Source/Runtime/VulkanRHI/Private/VulkanShaders.cpp b/Engine/Source/Runtime/VulkanRHI/Private/VulkanShaders.cpp
--- a/Engine/Source/Runtime/VulkanRHI/Private/VulkanShaders.cpp
+++ b/Engine/Source/Runtime/VulkanRHI/Private/VulkanShaders.cpp
@@ -55,6 +55,8 @@ static void ReadShaderOptionalData(FShaderCodeReader& ShaderCode, ShaderType* RH
 	RHIShader->Debug.ShaderName = ShaderCode.FindOptionalData(FShaderCodeName::Key);
 	UE::RHICore::SetupShaderCodeValidationData(RHIShader, ShaderCode);
 #endif
+
+	UE::RHICore::SetupSpecializationConstantData(RHIShader, ShaderCode);
 }
 
 static VkPipeline CreatePreRasterShaderLibrary(FVulkanDevice* Device, FVulkanShader* Shader, VkShaderStageFlagBits StageFlag)
