Fix async compute test for Intel

diff --git a/Engine/Source/Runtime/RHI/Private/DynamicRHI.cpp b/Engine/Source/Runtime/RHI/Private/DynamicRHI.cpp
--- a/Engine/Source/Runtime/RHI/Private/DynamicRHI.cpp
+++ b/Engine/Source/Runtime/RHI/Private/DynamicRHI.cpp
@@ -135,6 +135,8 @@ void RHIDetectAndWarnOfBadDrivers(bool bHasEditorToken)
 				{
 					// Allow config to override vendor API queries.
 					DriverInfo.AdapterGeneration = DriverGeneration;
+					UE_LOGF(LogRHI, Log, "Overriding AdapterGeneration %u with DriverGeneration %d", GRHIGlobals.GpuInfo.AdapterGeneration, DriverGeneration);
+					GRHIGlobals.GpuInfo.AdapterGeneration = DriverGeneration;
 				}
 
 				bDeviceCanUpdateDriver = true;
