ProfileGPU log output now includes Graphics pipe wait times, which greatly improves Frame GPU time accuracy with async compute, and shows where on the graphics pipe the waits were

#Synced-CL 52556135

diff --git a/Engine/Source/Runtime/RHI/Private/GPUProfiler.cpp b/Engine/Source/Runtime/RHI/Private/GPUProfiler.cpp
--- a/Engine/Source/Runtime/RHI/Private/GPUProfiler.cpp
+++ b/Engine/Source/Runtime/RHI/Private/GPUProfiler.cpp
@@ -248,6 +248,11 @@ namespace UE::RHI::GPUProfiler
 				return FPlatformTime::ToMilliseconds64(BusyCycles);
 			}
 
+			double GetWaitMilliseconds() const
+			{
+				return FPlatformTime::ToMilliseconds64(WaitCycles);
+			}
+
 			bool HasWork() const
 			{
 				return NumDraws > 0 || NumDispatches > 0;
@@ -1206,6 +1211,7 @@ namespace UE::RHI::GPUProfiler
 					FWildcardString RootWildcard(GCVarProfileGPU_Root.GetValueOnAnyThread());
 					const bool bShowEmptyNodes = GCVarProfileGPU_ShowLeafEvents.GetValueOnAnyThread();
 					const double PercentThreshold = FMath::Clamp(GCVarProfileGPU_ThresholdPercent.GetValueOnAnyThread(), 0.0f, 100.0f);
+					const bool bGraphicsPipeline = QueueState.Queue.Type == FQueue::EType::Graphics;
 
 					if (SortMode != EGPUProfileSortMode::Chronological)
 					{
@@ -1229,12 +1235,23 @@ namespace UE::RHI::GPUProfiler
 
 					FTable Table;
 					FString NodeOutput;
-					double FrameInclusiveTime = First ? First->Inclusive.GetBusyMilliseconds() : 0.0;
+					double FrameInclusiveTime = 0.0;
+					
+					if (First)
+					{
+						FrameInclusiveTime = First->Inclusive.GetBusyMilliseconds();
+
+						if (bGraphicsPipeline)
+						{
+							FrameInclusiveTime += First->Inclusive.GetWaitMilliseconds();
+						}
+					}
 
 					auto Recurse = [&](auto& Recurse, FNode* Root, FNode* CurrentNode, bool bParentMatchedFilter, int32 Level) -> bool
 					{
 						// Percent that this node was of the total frame time
-						const double Percent = (CurrentNode->Inclusive.GetBusyMilliseconds() / FrameInclusiveTime) * 100.0;
+						const double WaitMilliseconds = bGraphicsPipeline ? CurrentNode->Inclusive.GetWaitMilliseconds() : 0.0;
+						const double Percent = ((CurrentNode->Inclusive.GetBusyMilliseconds() + WaitMilliseconds) / FrameInclusiveTime) * 100.0;
 
 						// Filter nodes according to cvar settings
 						const bool bAboveThreshold = Percent >= PercentThreshold;
@@ -1261,10 +1278,11 @@ namespace UE::RHI::GPUProfiler
 							}
 							else
 							{
-								NodeOutput += FString::Printf(TEXT("%s%4.1f%% %5.2fms   %s\n"),
+								NodeOutput += FString::Printf(TEXT("%s%4.1f%% %5.2fms%s   %s\n"),
 									*FString().LeftPad(Level * 3),
 									Percent,
 									CurrentNode->Inclusive.GetBusyMilliseconds(),
+									WaitMilliseconds > 0 && (CurrentNode->Children.Num() == 0 || CurrentNode == First) ? *FString::Printf(TEXT(" + %2.2f Wait"), WaitMilliseconds) : TEXT(""),
 									*CurrentNode->Name);
 							}
 						}
@@ -1287,7 +1305,8 @@ namespace UE::RHI::GPUProfiler
 
 						if (bDisplayEvent && NumHiddenChildren > 0)
 						{
-							const double OtherChildrenInclusiveMs = FPlatformTime::ToMilliseconds64(OtherChildrenInclusive.BusyCycles);
+							const double OtherChildrenInclusiveWaitMs = bGraphicsPipeline ? FPlatformTime::ToMilliseconds64(OtherChildrenInclusive.WaitCycles) : 0.0;
+							const double OtherChildrenInclusiveMs = FPlatformTime::ToMilliseconds64(OtherChildrenInclusive.BusyCycles) + OtherChildrenInclusiveWaitMs;
 							const double OtherChildrenPercent = (OtherChildrenInclusiveMs / FrameInclusiveTime) * 100.0;
 							// Show the "other children" node if their total inclusive time is above the threshold
 							if (OtherChildrenPercent >= PercentThreshold)
@@ -1304,10 +1323,11 @@ namespace UE::RHI::GPUProfiler
 								}
 								else
 								{
-									NodeOutput += FString::Printf(TEXT("%s%4.1f%% %5.2fms   %d Other Children\n"),
+									NodeOutput += FString::Printf(TEXT("%s%4.1f%% %5.2fms%s   %d Other Children\n"),
 										*FString().LeftPad(Level * 3),
 										OtherChildrenPercent,
 										OtherChildrenInclusiveMs,
+										OtherChildrenInclusiveWaitMs > 0 ? *FString::Printf(TEXT(" + %2.2f Wait"), OtherChildrenInclusiveWaitMs) : TEXT(""),
 										NumHiddenChildren);
 								}
 							}
@@ -1332,7 +1352,7 @@ namespace UE::RHI::GPUProfiler
 							TEXT("    - %-30s: %s\n")
 							TEXT("\n")
 							, TEXT("Frame Time")
-							, First ? First->Inclusive.GetBusyMilliseconds() : 0.0
+							, First ? (First->Inclusive.GetBusyMilliseconds() + (bGraphicsPipeline ? First->Inclusive.GetWaitMilliseconds() : 0.0)) : 0.0
 							, *IConsoleManager::Get().FindConsoleObjectName(GCVarProfileGPU_Root.AsVariable())
 							, *RootWildcard
 							, *IConsoleManager::Get().FindConsoleObjectName(GCVarProfileGPU_ThresholdPercent.AsVariable())
@@ -1362,7 +1382,7 @@ namespace UE::RHI::GPUProfiler
 
 					FString Final;
 
-					if (QueueState.Queue.Type == FQueue::EType::Graphics)
+					if (bGraphicsPipeline)
 					{
 						Final = FString::Printf(
 							TEXT("\n")
@@ -1394,7 +1414,7 @@ namespace UE::RHI::GPUProfiler
 
 				#if !UE_BUILD_SHIPPING
 					// Create and display profile visualizer data for the graphics queue
-					if (GCVarProfileGPU_ShowUI.GetValueOnAnyThread() && QueueState.Queue.Type == FQueue::EType::Graphics && bShowQueue)
+					if (GCVarProfileGPU_ShowUI.GetValueOnAnyThread() && bGraphicsPipeline && bShowQueue)
 					{
 						// Count the total number of exclusive cycles in the frame. Needed to draw the horizontal bars in the visualizer.
 						uint64 TotalBusyCycles = 0;
