- Re-enabling preview platforms for editor when running SM6 Metal RHI
- SM6 preview is still disabled due to issues with D3D
- Fixed Metal texture sampler count to enable SM6 preview in the future

#Synced-CL 52187599

diff --git a/Engine/Source/Developer/Apple/MetalShaderFormat/Private/MetalShaderFormat.cpp b/Engine/Source/Developer/Apple/MetalShaderFormat/Private/MetalShaderFormat.cpp
--- a/Engine/Source/Developer/Apple/MetalShaderFormat/Private/MetalShaderFormat.cpp
+++ b/Engine/Source/Developer/Apple/MetalShaderFormat/Private/MetalShaderFormat.cpp
@@ -452,15 +452,17 @@ public:
 		else if (Input.ShaderFormat == NAME_SF_METAL_SM6)
 		{
 			Input.Environment.SetDefine(TEXT("METAL_SM6_PROFILE"), 1);
-			Input.Environment.SetDefine(TEXT("PLATFORM_SUPPORTS_STATIC_SAMPLERS"), 1);	
 			Input.Environment.SetDefine(TEXT("PLATFORM_SUPPORTS_CALLABLE_SHADERS"), true);
 		}
 
 		Input.Environment.SetDefine(TEXT("COMPILER_HLSLCC"), 2);
-		const bool bUseMetalShaderConverter = FDataDrivenShaderPlatformInfo::GetSupportsBindless(Input.Target.GetPlatform());
+		const bool bUseMetalShaderConverter = (FDataDrivenShaderPlatformInfo::GetIsPreviewPlatform(Input.Target.GetPlatform()) && 
+											   FDataDrivenShaderPlatformInfo::GetSupportsBindless(GMaxRHIShaderPlatform))
+												 || FDataDrivenShaderPlatformInfo::GetSupportsBindless(Input.Target.GetPlatform());
 		
 		if (bUseMetalShaderConverter)
 		{
+			Input.Environment.SetDefine(TEXT("PLATFORM_SUPPORTS_STATIC_SAMPLERS"), 1);
 			Input.Environment.SetDefine(TEXT("COMPILER_METAL_SHADER_CONVERTER"), 1);
 		}
 		
diff --git a/Engine/Source/Editor/LevelEditor/Private/LevelEditorActions.cpp b/Engine/Source/Editor/LevelEditor/Private/LevelEditorActions.cpp
--- a/Engine/Source/Editor/LevelEditor/Private/LevelEditorActions.cpp
+++ b/Engine/Source/Editor/LevelEditor/Private/LevelEditorActions.cpp
@@ -939,8 +939,8 @@ void FLevelEditorActionCallbacks::SetPreviewPlatform(FPreviewPlatformInfo NewPre
 
 bool FLevelEditorActionCallbacks::CanExecutePreviewPlatform(FPreviewPlatformInfo NewPreviewPlatform)
 {
-	// Temporary - We have disable SM6 preview for platforms < SM6 for now as it causes crashes due to a bindful/bindless mismatch
-	if (IsMetalSM6Platform(GMaxRHIShaderPlatform))
+	// TODO: Carl - Disable SM6 platforms for Metal SM6 until we fix issues.
+	if (IsMetalSM6Platform(GMaxRHIShaderPlatform) && NewPreviewPlatform.PreviewFeatureLevel >= GMaxRHIFeatureLevel)
 	{
 		return false;
 	}
diff --git a/Engine/Source/Runtime/Apple/MetalRHI/Private/MetalRHI.cpp b/Engine/Source/Runtime/Apple/MetalRHI/Private/MetalRHI.cpp
--- a/Engine/Source/Runtime/Apple/MetalRHI/Private/MetalRHI.cpp
+++ b/Engine/Source/Runtime/Apple/MetalRHI/Private/MetalRHI.cpp
@@ -464,8 +464,8 @@ FMetalDynamicRHI::FMetalDynamicRHI(ERHIFeatureLevel::Type RequestedFeatureLevel)
 	}
 	
 	GRHIGlobals.bSupportsBindless = UE::RHICore::GetBindlessConfigurationOnStartup(GMaxRHIShaderPlatform) != ERHIBindlessConfiguration::Disabled;
-	
 	GRHIGlobals.MaxSimultaneousUAVs = FMath::Min(METAL_MAX_TEXTURES, METAL_MAX_BUFFERS);
+	GRHIGlobals.MaxTextureSamplers = GRHIGlobals.bSupportsBindless ? MTLDevice->maxArgumentBufferSamplerCount() : 16; 
 	
 	ValidateTargetedRHIFeatureLevelExists(GMaxRHIShaderPlatform);
 	
diff --git a/Engine/Source/Runtime/RHI/Private/DataDrivenShaderPlatformInfo.cpp b/Engine/Source/Runtime/RHI/Private/DataDrivenShaderPlatformInfo.cpp
--- a/Engine/Source/Runtime/RHI/Private/DataDrivenShaderPlatformInfo.cpp
+++ b/Engine/Source/Runtime/RHI/Private/DataDrivenShaderPlatformInfo.cpp
@@ -594,6 +594,8 @@ void FGenericDataDrivenShaderPlatformInfo::UpdatePreviewPlatforms()
 				PREVIEW_USE_RUNTIME_VALUE(bSupportsDescriptorRange);
 				PREVIEW_USE_RUNTIME_VALUE(StaticShaderBindingLayoutSupport);
 
+				PREVIEW_USE_RUNTIME_VALUE(bSupportsRayTracingClusterOps);
+				
 				// Settings that will never be supported in preview
 				PREVIEW_FORCE_DISABLE(bSupportsShaderRootConstants);
 				PREVIEW_FORCE_DISABLE(bSupportsShaderBundleDispatch);
@@ -602,6 +604,13 @@ void FGenericDataDrivenShaderPlatformInfo::UpdatePreviewPlatforms()
 				PREVIEW_FORCE_DISABLE(bSupportsDOFHybridScattering);
 				PREVIEW_FORCE_DISABLE(bSupports4ComponentUAVReadWrite);
 
+				// Metals bindless and bSupportsSceneDataCompressedTransforms need to match the max shader platform
+				if (IsMetalPlatform(GMaxRHIShaderPlatform))
+				{
+					PREVIEW_FORCE_SETTING(bSupportsBindless, Infos[GMaxRHIShaderPlatform].bSupportsBindless);
+					PREVIEW_FORCE_SETTING(bSupportsSceneDataCompressedTransforms, Infos[GMaxRHIShaderPlatform].bSupportsSceneDataCompressedTransforms);
+				}
+				
 				// Make sure we're marked valid
 				PreviewInfo.bContainsValidPlatformInfo = true;
 
diff --git a/Engine/Source/Runtime/RenderCore/Private/ShaderPlatformConfig.cpp b/Engine/Source/Runtime/RenderCore/Private/ShaderPlatformConfig.cpp
--- a/Engine/Source/Runtime/RenderCore/Private/ShaderPlatformConfig.cpp
+++ b/Engine/Source/Runtime/RenderCore/Private/ShaderPlatformConfig.cpp
@@ -85,7 +85,11 @@ void FShaderPlatformConfig::FPlatformConfig::InitializePropertiesForPreview(cons
 #define PREVIEW_FORCE_DISABLE(SettingName) \
 	PREVIEW_FORCE_SETTING(SettingName, false)
 
-	// TODO:
+	// Metal platforms require bindless config to match
+	if (IsMetalPlatform(GMaxRHIShaderPlatform))
+	{
+		PREVIEW_FORCE_SETTING(BindlessConfiguration, PlatformConfigs[GMaxRHIShaderPlatform].BindlessConfiguration);
+	}
 
 #undef PREVIEW_FORCE_DISABLE
 #undef PREVIEW_FORCE_SETTING
